EvalGraph - Renamed modules

diff --git a/Engine/Plugins/Experimental/Dataflow/Dataflow.uplugin b/Engine/Plugins/Experimental/Dataflow/Dataflow.uplugin
new file mode 100644
--- /dev/null
+++ b/Engine/Plugins/Experimental/Dataflow/Dataflow.uplugin
@@ -0,0 +1,27 @@
+{
+    "FileVersion" : 3,
+    "Version" : 1,
+    "VersionName" : "0.1",
+    "FriendlyName" : "Editor DataflowGraph",
+    "Description" : "Editor Dataflow Graph",
+    "Category": "Editor",
+    "CreatedBy" : "Epic Games, Inc.",
+    "CreatedByURL" : "http://epicgames.com",
+    "DocsURL" : "",
+    "MarketplaceURL" : "",
+    "SupportURL" : "",
+    "EnabledByDefault" : true,
+    "CanContainContent" : true,
+    "IsBetaVersion" : true,
+	"IsExperimentalVersion": true,
+    "Installed" : true,
+	"Modules": [
+		{
+			"Name": "DataflowEditor",
+			"Type": "Editor",
+			"LoadingPhase": "Default"
+		},
+	],
+	"Plugins": [
+	]
+}
diff --git a/Engine/Plugins/Experimental/Dataflow/Source/DataflowEditor/DataflowEditor.Build.cs b/Engine/Plugins/Experimental/Dataflow/Source/DataflowEditor/DataflowEditor.Build.cs
new file mode 100644
--- /dev/null
+++ b/Engine/Plugins/Experimental/Dataflow/Source/DataflowEditor/DataflowEditor.Build.cs
@@ -0,0 +1,42 @@
+// Copyright Epic Games, Inc. All Rights Reserved.
+
+namespace UnrealBuildTool.Rules
+{
+	public class DataflowEditor : ModuleRules
+	{
+        public DataflowEditor(ReadOnlyTargetRules Target) : base(Target)
+		{
+			PrivateIncludePaths.Add("DataflowEditor/Private");
+            PublicIncludePaths.Add(ModuleDirectory + "/Public");
+
+			PublicDependencyModuleNames.AddRange(
+				new string[]
+				{
+					"Core",
+					"CoreUObject",
+				    "Slate",
+				    "SlateCore",
+				    "Engine",
+					"EditorFramework",
+					"UnrealEd",
+				    "PropertyEditor",
+				    "RenderCore",
+				    "RHI",
+				    "AssetTools",
+				    "AssetRegistry",
+				    "SceneOutliner",
+					"EditorStyle",
+					"AssetTools",
+					"ToolMenus",
+					"LevelEditor",
+					"InputCore",
+					"AdvancedPreviewScene",
+					"GraphEditor",
+					"DataflowCore",
+					"DataflowEngine",
+					"Slate",
+				}
+			);
+		}
+	}
+}
diff --git a/Engine/Plugins/Experimental/Dataflow/Source/DataflowEditor/Private/Dataflow/DataflowAssetActions.cpp b/Engine/Plugins/Experimental/Dataflow/Source/DataflowEditor/Private/Dataflow/DataflowAssetActions.cpp
new file mode 100644
--- /dev/null
+++ b/Engine/Plugins/Experimental/Dataflow/Source/DataflowEditor/Private/Dataflow/DataflowAssetActions.cpp
@@ -0,0 +1,55 @@
+// Copyright Epic Games, Inc. All Rights Reserved.
+
+#include "Dataflow/DataflowAssetActions.h"
+#include "Dataflow/DataflowEditorPlugin.h"
+
+#include "Dataflow/DataflowObject.h"
+
+#define LOCTEXT_NAMESPACE "AssetActions_DataflowAsset"
+
+FText FDataflowAssetActions::GetName() const
+{
+	return LOCTEXT("Name", "Dataflow Graph");
+}
+
+UClass* FDataflowAssetActions::GetSupportedClass() const
+{
+	return UDataflow::StaticClass();
+}
+
+FColor FDataflowAssetActions::GetTypeColor() const
+{
+	return FColor(255, 127, 40);
+}
+
+void FDataflowAssetActions::GetActions(const TArray<UObject*>& InObjects,
+										class FMenuBuilder&     MenuBuilder)
+{
+}
+
+void FDataflowAssetActions::OpenAssetEditor(
+	const TArray<UObject*>& InObjects, TSharedPtr<IToolkitHost> EditWithinLevelEditor)
+{
+	EToolkitMode::Type Mode = EditWithinLevelEditor.IsValid() ? EToolkitMode::WorldCentric : EToolkitMode::Standalone;
+	for (auto ObjIt = InObjects.CreateConstIterator(); ObjIt; ++ObjIt)
+	{
+		if (auto Object = Cast<UDataflow>(*ObjIt))
+		{
+			IDataflowEditorPlugin* DataflowEditorPlugin = &FModuleManager::LoadModuleChecked<IDataflowEditorPlugin>("DataflowEditor");
+			DataflowEditorPlugin->CreateDataflowAssetEditor(Mode, EditWithinLevelEditor, Object);
+		}
+	}
+}
+
+
+uint32 FDataflowAssetActions::GetCategories()
+{
+	return EAssetTypeCategories::Physics;
+}
+
+FText FDataflowAssetActions::GetAssetDescription(const struct FAssetData& AssetData) const
+{
+	return LOCTEXT("Description", "A dataflow graph for asset authoring.");
+}
+
+#undef LOCTEXT_NAMESPACE
diff --git a/Engine/Plugins/Experimental/Dataflow/Source/DataflowEditor/Private/Dataflow/DataflowAssetFactory.cpp b/Engine/Plugins/Experimental/Dataflow/Source/DataflowEditor/Private/Dataflow/DataflowAssetFactory.cpp
new file mode 100644
--- /dev/null
+++ b/Engine/Plugins/Experimental/Dataflow/Source/DataflowEditor/Private/Dataflow/DataflowAssetFactory.cpp
@@ -0,0 +1,41 @@
+// Copyright Epic Games, Inc. All Rights Reserved.
+
+#include "Dataflow/DataflowAssetFactory.h"
+
+#include "Dataflow/DataflowObject.h"
+#include "HAL/IConsoleManager.h"
+
+bool bDataflowEnableCreation = false;
+FAutoConsoleVariableRef CVarDataflowEnableCreation(TEXT("c.DataflowEnableCreation"), bDataflowEnableCreation, TEXT("Enable creation for the dataflow system (Curently Dev-Only)"));
+
+
+UDataflowAssetFactory::UDataflowAssetFactory()
+{
+	SupportedClass = UDataflow::StaticClass();
+}
+
+bool UDataflowAssetFactory::CanCreateNew() const
+{
+	return bDataflowEnableCreation;
+}
+
+bool UDataflowAssetFactory::FactoryCanImport(const FString& Filename)
+{
+	return false;
+}
+
+UObject* UDataflowAssetFactory::FactoryCreateNew(UClass* InClass, UObject* InParent, FName InName, EObjectFlags Flags,
+												   UObject* Context, FFeedbackContext* Warn)
+{
+	return NewObject<UDataflow>(InParent, InClass, InName, Flags);
+}
+
+bool UDataflowAssetFactory::ShouldShowInNewMenu() const
+{
+	return bDataflowEnableCreation;
+}
+
+bool UDataflowAssetFactory::ConfigureProperties()
+{
+	return true;
+}
diff --git a/Engine/Plugins/Experimental/Dataflow/Source/DataflowEditor/Private/Dataflow/DataflowEditorActions.cpp b/Engine/Plugins/Experimental/Dataflow/Source/DataflowEditor/Private/Dataflow/DataflowEditorActions.cpp
new file mode 100644
--- /dev/null
+++ b/Engine/Plugins/Experimental/Dataflow/Source/DataflowEditor/Private/Dataflow/DataflowEditorActions.cpp
@@ -0,0 +1,47 @@
+// Copyright Epic Games, Inc. All Rights Reserved.
+
+#include "Dataflow/DataflowEditorActions.h"
+#include "Dataflow/DataflowNodeFactory.h"
+
+#define LOCTEXT_NAMESPACE "DataflowEditorCommands"
+
+void FDataflowEditorCommandsImpl::RegisterCommands()
+{
+	UI_COMMAND(EvaluateNode, "Evaluate", "Trigger an evaluation of the selected node.", EUserInterfaceActionType::Button, FInputChord());
+
+	if (Dataflow::FNodeFactory* Factory = Dataflow::FNodeFactory::GetInstance())
+	{
+		for (FName NodeName : Factory->RegisteredNodes())
+		{
+			TSharedPtr< FUICommandInfo > AddNode;
+			FUICommandInfo::MakeCommandInfo(
+				this->AsShared(),
+				AddNode,
+				NodeName, //FName("UseCreationFormToggle"),
+				NSLOCTEXT("DataFlow", "DataflowButton", "New Dataflow Node"),
+				NSLOCTEXT("DataFlow", "NewDataflowNodeTooltip", "New Dataflow Node Tooltip"),
+				FSlateIcon(),
+				EUserInterfaceActionType::Button,
+				FInputChord()
+			);
+			CreateNodesMap.Add(NodeName, AddNode);
+		}
+	}
+}
+
+void FDataflowEditorCommands::Register()
+{
+	return FDataflowEditorCommandsImpl::Register();
+}
+
+const FDataflowEditorCommandsImpl& FDataflowEditorCommands::Get()
+{
+	return FDataflowEditorCommandsImpl::Get();
+}
+
+void FDataflowEditorCommands::Unregister()
+{
+	return FDataflowEditorCommandsImpl::Unregister();
+}
+
+#undef LOCTEXT_NAMESPACE
diff --git a/Engine/Plugins/Experimental/Dataflow/Source/DataflowEditor/Private/Dataflow/DataflowEditorPlugin.cpp b/Engine/Plugins/Experimental/Dataflow/Source/DataflowEditor/Private/Dataflow/DataflowEditorPlugin.cpp
new file mode 100644
--- /dev/null
+++ b/Engine/Plugins/Experimental/Dataflow/Source/DataflowEditor/Private/Dataflow/DataflowEditorPlugin.cpp
@@ -0,0 +1,55 @@
+// Copyright Epic Games, Inc. All Rights Reserved.
+
+#include "Dataflow/DataflowEditorPlugin.h"
+
+#include "AssetToolsModule.h"
+#include "CoreMinimal.h"
+#include "EdGraphUtilities.h"
+#include "DataflowEditorToolkit.h"
+#include "Dataflow/DataflowNodeFactory.h"
+#include "Dataflow/DataflowAssetActions.h"
+#include "Dataflow/DataflowSNodeFactories.h"
+
+#define LOCTEXT_NAMESPACE "DataflowEditor"
+
+//#define BOX_BRUSH(StyleSet, RelativePath, ...) FSlateBoxBrush(StyleSet->RootToContentDir(RelativePath, TEXT(".png")), __VA_ARGS__)
+//#define DEFAULT_FONT(...) FCoreStyle::GetDefaultFontStyle(__VA_ARGS__)
+
+void IDataflowEditorPlugin::StartupModule()
+{
+	DataflowAssetActions = new FDataflowAssetActions();
+
+	FAssetToolsModule& AssetToolsModule = FAssetToolsModule::GetModule();
+	IAssetTools& AssetTools = AssetToolsModule.Get();
+	AssetTools.RegisterAssetTypeActions(MakeShareable(DataflowAssetActions));
+
+
+	DataflowSNodeFactory = MakeShareable(new FDataflowSNodeFactory());
+	FEdGraphUtilities::RegisterVisualNodeFactory(DataflowSNodeFactory);
+}
+
+void IDataflowEditorPlugin::ShutdownModule()
+{
+	if (UObjectInitialized())
+	{
+		FAssetToolsModule& AssetToolsModule = FAssetToolsModule::GetModule();
+		IAssetTools& AssetTools = AssetToolsModule.Get();
+
+		AssetTools.UnregisterAssetTypeActions(DataflowAssetActions->AsShared());
+
+		FEdGraphUtilities::UnregisterVisualNodeFactory(DataflowSNodeFactory);
+	}
+}
+
+TSharedRef<FAssetEditorToolkit> IDataflowEditorPlugin::CreateDataflowAssetEditor(const EToolkitMode::Type Mode, const TSharedPtr<IToolkitHost>& InitToolkitHost, UObject* FleshAsset)
+{
+	TSharedPtr<FDataflowEditorToolkit> NewDataflowAssetEditor = MakeShared<FDataflowEditorToolkit>();
+	NewDataflowAssetEditor->InitDataflowEditor(Mode, InitToolkitHost, FleshAsset);
+	return StaticCastSharedPtr<FAssetEditorToolkit>(NewDataflowAssetEditor).ToSharedRef();
+}
+
+
+IMPLEMENT_MODULE(IDataflowEditorPlugin, DataflowEditor)
+
+
+#undef LOCTEXT_NAMESPACE
diff --git a/Engine/Plugins/Experimental/Dataflow/Source/DataflowEditor/Private/Dataflow/DataflowEditorToolkit.cpp b/Engine/Plugins/Experimental/Dataflow/Source/DataflowEditor/Private/Dataflow/DataflowEditorToolkit.cpp
new file mode 100644
--- /dev/null
+++ b/Engine/Plugins/Experimental/Dataflow/Source/DataflowEditor/Private/Dataflow/DataflowEditorToolkit.cpp
@@ -0,0 +1,242 @@
+// Copyright Epic Games, Inc. All Rights Reserved.
+
+#include "DataflowEditorToolkit.h"
+
+#include "Dataflow/DataflowEditorActions.h"
+#include "Dataflow/DataflowEdNode.h"
+#include "Dataflow/DataflowNodeFactory.h"
+#include "Dataflow/DataflowObject.h"
+#include "Dataflow/DataflowSchema.h"
+#include "Dataflow/Dataflow.h"
+#include "EditorStyleSet.h"
+#include "EditorViewportTabContent.h"
+#include "EditorViewportLayout.h"
+#include "EditorViewportCommands.h"
+#include "GraphEditorActions.h"
+#include "Modules/ModuleManager.h"
+#include "PropertyEditorModule.h"
+#include "Styling/SlateStyleRegistry.h"
+
+#define LOCTEXT_NAMESPACE "DataflowEditorToolkit"
+
+//DEFINE_LOG_CATEGORY_STATIC(FDataflowEditorToolkitLog, Log, All);
+
+
+const FName FDataflowEditorToolkit::GraphCanvasTabId(TEXT("DataflowEditor_GraphCanvas"));
+const FName FDataflowEditorToolkit::PropertiesTabId(TEXT("DataflowEditor_Properties"));
+
+void FDataflowEditorToolkit::InitDataflowEditor(const EToolkitMode::Type Mode, const TSharedPtr<IToolkitHost>& InitToolkitHost, UObject* ObjectToEdit)
+{
+	Dataflow = CastChecked<UDataflow>(ObjectToEdit);
+	if (Dataflow != nullptr)
+	{
+		Dataflow->Schema = UDataflowSchema::StaticClass();
+
+		GraphEditor = CreateGraphEditorWidget(Dataflow);
+		PropertiesEditor = CreatePropertiesEditorWidget(ObjectToEdit);
+
+		const TSharedRef<FTabManager::FLayout> StandaloneDefaultLayout = FTabManager::NewLayout("Dataflow_Layout")
+			->AddArea
+			(
+				FTabManager::NewPrimaryArea()->SetOrientation(Orient_Vertical)
+				->Split
+				(
+					FTabManager::NewSplitter()->SetOrientation(Orient_Vertical)
+					->Split
+					(
+						FTabManager::NewSplitter()->SetOrientation(Orient_Horizontal)
+						->SetSizeCoefficient(0.9f)
+						->Split
+						(
+							FTabManager::NewStack()
+							->SetSizeCoefficient(0.6f)
+							->AddTab(GraphCanvasTabId, ETabState::OpenedTab)
+						)
+						->Split
+						(
+							FTabManager::NewSplitter()->SetOrientation(Orient_Vertical)
+							->SetSizeCoefficient(0.2f)
+							->Split
+							(
+								FTabManager::NewStack()
+								->SetSizeCoefficient(0.7f)
+								->AddTab(PropertiesTabId, ETabState::OpenedTab)
+							)
+						)
+					)
+				)
+			);
+
+		const bool bCreateDefaultStandaloneMenu = true;
+		const bool bCreateDefaultToolbar = true;
+		FAssetEditorToolkit::InitAssetEditor(Mode, InitToolkitHost, FName(TEXT("DataflowEditorApp")), StandaloneDefaultLayout, bCreateDefaultToolbar, bCreateDefaultStandaloneMenu, ObjectToEdit);
+	}
+}
+
+void FDataflowEditorToolkit::EvaluateNode()
+{
+	if (UDataflow* Graph = dynamic_cast<UDataflow*>(GraphEditor->GetCurrentGraph()))
+	{
+		if (TSharedPtr<Dataflow::FGraph> DataflowGraph = Graph->GetDataflow())
+		{
+			for (UObject* Ode : GetSelectedNodes())
+			{
+				if (UDataflowEdNode* Node = dynamic_cast<UDataflowEdNode*>(Ode))
+				{
+					if (TSharedPtr<Dataflow::FNode> DataflowNode = DataflowGraph->FindBaseNode(Node->GetDataflowNodeGuid()))
+					{
+						for (Dataflow::FConnection* NodeOutput : DataflowNode->GetOutputs())
+						{
+							DataflowNode->Evaluate({0.f}, NodeOutput);
+						}
+					}
+				}
+			}
+		}
+	}
+}
+
+FGraphPanelSelectionSet FDataflowEditorToolkit::GetSelectedNodes() const
+{
+	if (GraphEditor.IsValid())
+	{
+		return GraphEditor->GetSelectedNodes();
+	}
+	return FGraphPanelSelectionSet();
+}
+
+void FDataflowEditorToolkit::OnSelectedNodesChanged(const TSet<UObject*>& NewSelection)
+{
+}
+
+
+
+TSharedRef<SGraphEditor> FDataflowEditorToolkit::CreateGraphEditorWidget(UDataflow* DataflowToEdit)
+{
+	ensure(DataflowToEdit);
+
+	FDataflowEditorCommands::Register();
+	FGraphEditorCommands::Register();
+
+	// No need to regenerate the commands.
+	if (!GraphEditorCommands.IsValid())
+	{
+		GraphEditorCommands = MakeShareable(new FUICommandList);
+		{
+			GraphEditorCommands->MapAction(FDataflowEditorCommands::Get().EvaluateNode,
+				FExecuteAction::CreateSP(this, &FDataflowEditorToolkit::EvaluateNode)
+			);
+		}
+	}
+
+
+	FGraphAppearanceInfo AppearanceInfo;
+	AppearanceInfo.CornerText = LOCTEXT("AppearanceCornerText_DataflowEditor", "Dataflow");
+
+	SGraphEditor::FGraphEditorEvents InEvents;
+	InEvents.OnSelectionChanged = SGraphEditor::FOnSelectionChanged::CreateSP(this, &FDataflowEditorToolkit::OnSelectedNodesChanged);
+
+	return SNew(SGraphEditor)
+		.AdditionalCommands(GraphEditorCommands)
+		.IsEditable(true)
+		.Appearance(AppearanceInfo)
+		.GraphToEdit(DataflowToEdit)
+		.GraphEvents(InEvents)
+		.ShowGraphStateOverlay(false);
+}
+
+TSharedPtr<IDetailsView> FDataflowEditorToolkit::CreatePropertiesEditorWidget(UObject* ObjectToEdit)
+{
+	ensure(ObjectToEdit);
+	FPropertyEditorModule& PropertyEditorModule = FModuleManager::LoadModuleChecked<FPropertyEditorModule>(TEXT("PropertyEditor"));
+
+	FDetailsViewArgs DetailsViewArgs;
+	DetailsViewArgs.bAllowSearch = true;
+	DetailsViewArgs.bLockable = false;
+	DetailsViewArgs.bUpdatesFromSelection = false;
+	DetailsViewArgs.NameAreaSettings = FDetailsViewArgs::HideNameArea;
+	DetailsViewArgs.NotifyHook = this;
+
+	TSharedPtr<IDetailsView> DetailsView = PropertyEditorModule.CreateDetailView(DetailsViewArgs);
+	DetailsView->SetObject(ObjectToEdit);
+	return DetailsView;
+
+}
+
+
+TSharedRef<SDockTab> FDataflowEditorToolkit::SpawnTab_GraphCanvas(const FSpawnTabArgs& Args)
+{
+	check(Args.GetTabId() == GraphCanvasTabId);
+
+	TSharedRef<SDockTab> SpawnedTab = SNew(SDockTab)
+		.Label(LOCTEXT("DataflowEditor_Dataflow_TabTitle", "Graph"))
+		[
+			GraphEditor.ToSharedRef()
+		];
+
+	return SpawnedTab;
+}
+
+TSharedRef<SDockTab> FDataflowEditorToolkit::SpawnTab_Properties(const FSpawnTabArgs& Args)
+{
+	check(Args.GetTabId() == PropertiesTabId);
+
+	return SNew(SDockTab)
+		.Label(LOCTEXT("DataflowEditor_Properties_TabTitle", "Details"))
+		[
+			PropertiesEditor.ToSharedRef()
+		];
+}
+
+
+void FDataflowEditorToolkit::RegisterTabSpawners(const TSharedRef<FTabManager>& InTabManager)
+{
+	TSharedRef<FWorkspaceItem> WorkspaceMenuCategoryRef = InTabManager->AddLocalWorkspaceMenuCategory(LOCTEXT("WorkspaceMenu_DataflowEditor", "Dataflow Editor"));
+
+	InTabManager->RegisterTabSpawner(GraphCanvasTabId, FOnSpawnTab::CreateSP(this, &FDataflowEditorToolkit::SpawnTab_GraphCanvas))
+		.SetDisplayName(LOCTEXT("DataflowTab", "Graph"))
+		.SetGroup(WorkspaceMenuCategoryRef)
+		.SetIcon(FSlateIcon(FAppStyle::GetAppStyleSetName(), "GraphEditor.EventGraph_16x"));
+
+	InTabManager->RegisterTabSpawner(PropertiesTabId, FOnSpawnTab::CreateSP(this, &FDataflowEditorToolkit::SpawnTab_Properties))
+		.SetDisplayName(LOCTEXT("PropertiesTab", "Details"))
+		.SetGroup(WorkspaceMenuCategoryRef)
+		.SetIcon(FSlateIcon(FAppStyle::GetAppStyleSetName(), "LevelEditor.Tabs.Details"));
+
+	FAssetEditorToolkit::RegisterTabSpawners(InTabManager);
+}
+
+
+FName FDataflowEditorToolkit::GetToolkitFName() const
+{
+	return FName("DataflowEditor");
+}
+
+FText FDataflowEditorToolkit::GetBaseToolkitName() const
+{
+	return LOCTEXT("AppLabel", "Dataflow Editor");
+}
+
+FString FDataflowEditorToolkit::GetWorldCentricTabPrefix() const
+{
+	return LOCTEXT("WorldCentricTabPrefix", "Dataflow").ToString();
+}
+
+FLinearColor FDataflowEditorToolkit::GetWorldCentricTabColorScale() const
+{
+	return FLinearColor(0.3f, 0.2f, 0.5f, 0.5f);
+}
+
+FString FDataflowEditorToolkit::GetReferencerName() const
+{
+	return TEXT("DataflowEditorToolkit");
+}
+
+void FDataflowEditorToolkit::AddReferencedObjects(FReferenceCollector& Collector)
+{
+	if (Dataflow)
+	{
+		Collector.AddReferencedObject(Dataflow);
+	}
+}
+#undef LOCTEXT_NAMESPACE
diff --git a/Engine/Plugins/Experimental/Dataflow/Source/DataflowEditor/Private/Dataflow/DataflowEditorToolkit.h b/Engine/Plugins/Experimental/Dataflow/Source/DataflowEditor/Private/Dataflow/DataflowEditorToolkit.h
new file mode 100644
--- /dev/null
+++ b/Engine/Plugins/Experimental/Dataflow/Source/DataflowEditor/Private/Dataflow/DataflowEditorToolkit.h
@@ -0,0 +1,59 @@
+// Copyright Epic Games, Inc. All Rights Reserved.
+
+#pragma once
+
+#include "CoreMinimal.h"
+#include "Toolkits/SimpleAssetEditor.h"
+#include "Misc/NotifyHook.h"
+
+class FEditorViewportTabContent;
+class IDetailsView;
+class FTabManager;
+class IToolkitHost;
+class UDataflow;
+class SDataflowEditorViewport;
+
+class FDataflowEditorToolkit : public FAssetEditorToolkit, public FNotifyHook, public FGCObject
+{
+public:
+	void InitDataflowEditor(const EToolkitMode::Type Mode, const TSharedPtr<IToolkitHost>& InitToolkitHost, UObject* ObjectToEdit);
+
+	// IToolkit Interface
+	virtual FName GetToolkitFName() const override;
+	virtual FText GetBaseToolkitName() const override;
+	virtual FString GetWorldCentricTabPrefix() const override;
+	virtual FLinearColor GetWorldCentricTabColorScale() const override;
+
+	// FGCObject Interface
+	virtual void AddReferencedObjects(FReferenceCollector& Collector) override;
+	virtual FString GetReferencerName() const override; 
+
+	// Tab spawners 
+	virtual void RegisterTabSpawners(const TSharedRef<FTabManager>& TabManager) override;
+	TSharedRef<SDockTab> SpawnTab_GraphCanvas(const FSpawnTabArgs& Args);
+	TSharedRef<SDockTab> SpawnTab_Properties(const FSpawnTabArgs& Args);
+
+	// Graph Editor Operations
+	void OnSelectedNodesChanged(const TSet<UObject*>& NewSelection);
+	void EvaluateNode();
+	FGraphPanelSelectionSet GetSelectedNodes() const;
+
+	UDataflow* GetDataflow() { return Dataflow; }
+	const UDataflow* GetDataflow() const { return Dataflow; }
+
+private:
+
+	static const FName GraphCanvasTabId;
+	TSharedPtr<SGraphEditor> GraphEditor;
+	TSharedRef<SGraphEditor> CreateGraphEditorWidget(UDataflow* ObjectToEdit);
+
+	static const FName PropertiesTabId;
+	TSharedPtr<IDetailsView> PropertiesEditor;
+	TSharedPtr<IDetailsView> CreatePropertiesEditorWidget(UObject* ObjectToEdit);
+
+
+	UDataflow* Dataflow = nullptr;
+
+	/** Command list for the graph editor */
+	TSharedPtr<FUICommandList> GraphEditorCommands;
+};
diff --git a/Engine/Plugins/Experimental/Dataflow/Source/DataflowEditor/Private/Dataflow/DataflowSNode.cpp b/Engine/Plugins/Experimental/Dataflow/Source/DataflowEditor/Private/Dataflow/DataflowSNode.cpp
new file mode 100644
--- /dev/null
+++ b/Engine/Plugins/Experimental/Dataflow/Source/DataflowEditor/Private/Dataflow/DataflowSNode.cpp
@@ -0,0 +1,94 @@
+// Copyright Epic Games, Inc. All Rights Reserved.
+
+#include "Dataflow/DataflowSNode.h"
+
+#include "Dataflow/DataflowEdNode.h"
+#include "Dataflow/DataflowNodeFactory.h"
+#include "Dataflow/DataflowObject.h"
+#include "Dataflow/Dataflow.h"
+#include "Logging/LogMacros.h"
+#include "Widgets/Layout/SBorder.h"
+
+#define LOCTEXT_NAMESPACE "SDataflowEdNode"
+//
+// SDataflowEdNode
+//
+
+void SDataflowEdNode::Construct(const FArguments& InArgs, UDataflowEdNode* InNode)
+{
+	GraphNode = InNode;
+	UpdateGraphNode();
+}
+
+FReply SDataflowEdNode::OnMouseButtonDoubleClick(const FGeometry& InMyGeometry, const FPointerEvent& InMouseEvent)
+{
+	return Super::OnMouseButtonDoubleClick(InMyGeometry, InMouseEvent);
+}
+
+
+
+//
+// Add a menu option to create a graph node.
+//
+TSharedPtr<FAssetSchemaAction_Dataflow_CreateNode_DataflowEdNode> FAssetSchemaAction_Dataflow_CreateNode_DataflowEdNode::CreateAction(UEdGraph* ParentGraph, const FName & InNodeTypeName)
+{
+ 	const FText AddToolTip = LOCTEXT("DataflowNodeTooltip_Example", "Add a Dataflow node.");
+	const FText NodeName = FText::FromString(*InNodeTypeName.ToString());
+	const FText Catagory = LOCTEXT("DataflowNodeDescription_Example", "Dataflow"); 	
+	TSharedPtr<FAssetSchemaAction_Dataflow_CreateNode_DataflowEdNode> NewNodeAction(
+		new FAssetSchemaAction_Dataflow_CreateNode_DataflowEdNode(InNodeTypeName, Catagory, NodeName, AddToolTip, 0));
+	return NewNodeAction;
+}
+
+//
+//  Created the EdGraph node and bind the guids to the Dataflow's node. 
+//
+UEdGraphNode* FAssetSchemaAction_Dataflow_CreateNode_DataflowEdNode::PerformAction(class UEdGraph* ParentGraph, UEdGraphPin* FromPin, const FVector2D Location, bool bSelectNewNode)
+{
+	if (UDataflow* Dataflow = Cast<UDataflow>(ParentGraph))
+	{
+		const FName NodeName(GetMenuDescription().ToString());
+		if (UDataflowEdNode* EdNode = NewObject<UDataflowEdNode>(Dataflow, UDataflowEdNode::StaticClass(), NodeName))
+		{
+			//const FScopedTransaction Transaction(LOCTEXT("DataflowNewNode", "Generic Graph Editor: New Node"));
+			Dataflow->Modify();
+			if (FromPin != nullptr)
+				FromPin->Modify();
+
+			Dataflow->AddNode(EdNode, true, bSelectNewNode);
+
+			EdNode->CreateNewGuid();
+			EdNode->PostPlacedNewNode();
+
+			if (Dataflow::FNodeFactory* Factory = Dataflow::FNodeFactory::GetInstance())
+			{
+				if (TSharedPtr<Dataflow::FNode> DataflowNode = Factory->NewNodeFromRegisteredType(*Dataflow->GetDataflow(), { FGuid::NewGuid(),NodeTypeName,NodeName}))
+				{
+					EdNode->SetDataflowGraph(Dataflow->GetDataflow());
+					EdNode->SetDataflowNodeGuid(DataflowNode->GetGuid());
+					EdNode->AllocateDefaultPins();
+				}
+			}
+
+			//EdNode->AllocateDefaultPins();
+			EdNode->AutowireNewNode(FromPin);
+
+			EdNode->NodePosX = Location.X;
+			EdNode->NodePosY = Location.Y;
+
+			//EdNode->GenericGraphNode->SetFlags(RF_Transactional);
+			EdNode->SetFlags(RF_Transactional);
+
+			return EdNode;
+		}
+	}
+	return nullptr;
+}
+
+//void FAssetSchemaAction_Dataflow_CreateNode_DataflowEdNode::AddReferencedObjects(FReferenceCollector& Collector)
+//{
+//	FEdGraphSchemaAction::AddReferencedObjects(Collector);
+//	Collector.AddReferencedObject(NodeTemplate);
+//}
+
+#undef LOCTEXT_NAMESPACE
diff --git a/Engine/Plugins/Experimental/Dataflow/Source/DataflowEditor/Private/Dataflow/DataflowSNodeFactories.cpp b/Engine/Plugins/Experimental/Dataflow/Source/DataflowEditor/Private/Dataflow/DataflowSNodeFactories.cpp
new file mode 100644
--- /dev/null
+++ b/Engine/Plugins/Experimental/Dataflow/Source/DataflowEditor/Private/Dataflow/DataflowSNodeFactories.cpp
@@ -0,0 +1,16 @@
+// Copyright Epic Games, Inc. All Rights Reserved.
+
+#include "Dataflow/DataflowSNodeFactories.h"
+#include "Dataflow/DataflowSNode.h"
+#include "Dataflow/DataflowEdNode.h"
+
+TSharedPtr<class SGraphNode> FDataflowSNodeFactory::CreateNode(UEdGraphNode* InNode) const
+{
+	if (UDataflowEdNode* Node = Cast<UDataflowEdNode>(InNode))
+	{
+		return SNew(SDataflowEdNode, Node);
+	}
+	return NULL;
+}
+
+
diff --git a/Engine/Plugins/Experimental/Dataflow/Source/DataflowEditor/Private/Dataflow/DataflowSchema.cpp b/Engine/Plugins/Experimental/Dataflow/Source/DataflowEditor/Private/Dataflow/DataflowSchema.cpp
new file mode 100644
--- /dev/null
+++ b/Engine/Plugins/Experimental/Dataflow/Source/DataflowEditor/Private/Dataflow/DataflowSchema.cpp
@@ -0,0 +1,72 @@
+// Copyright Epic Games, Inc. All Rights Reserved.
+
+#include "Dataflow/DataflowSchema.h"
+
+#include "Dataflow/DataflowEdNode.h"
+#include "Dataflow/DataflowEditorActions.h"
+#include "Dataflow/DataflowEdNode.h"
+#include "Dataflow/DataflowSNode.h"
+#include "Dataflow/DataflowNodeFactory.h"
+#include "EdGraph/EdGraphNode.h"
+#include "EdGraph/EdGraph.h"
+#include "Framework/Commands/GenericCommands.h"
+#include "Logging/LogMacros.h"
+#include "ToolMenuSection.h"
+#include "ToolMenu.h"
+#include "GraphEditorActions.h"
+
+#define LOCTEXT_NAMESPACE "DataflowNode"
+
+UDataflowSchema::UDataflowSchema()
+{
+}
+
+void UDataflowSchema::GetContextMenuActions(class UToolMenu* Menu, class UGraphNodeContextMenuContext* Context) const
+{
+	if (Context->Node)
+	{
+		FToolMenuSection& Section = Menu->AddSection("TestGraphSchemaNodeActions", LOCTEXT("ClassActionsMenuHeader", "Node Actions"));
+		{
+			Section.AddMenuEntry(FGenericCommands::Get().Delete);
+			Section.AddMenuEntry(FGenericCommands::Get().Cut);
+			Section.AddMenuEntry(FGenericCommands::Get().Copy);
+			Section.AddMenuEntry(FGenericCommands::Get().Duplicate);
+			Section.AddMenuEntry(FGraphEditorCommands::Get().BreakNodeLinks);
+			Section.AddMenuEntry(FDataflowEditorCommands::Get().EvaluateNode);
+		}
+	}
+
+	Super::GetContextMenuActions(Menu, Context);
+}
+
+void UDataflowSchema::GetGraphContextActions(FGraphContextMenuBuilder& ContextMenuBuilder) const
+{
+	if (Dataflow::FNodeFactory* Factory = Dataflow::FNodeFactory::GetInstance())
+	{
+		for (FName NodeTypeName : Factory->RegisteredNodes())
+		{
+			if (FDataflowEditorCommands::Get().CreateNodesMap.Contains(NodeTypeName))
+			{
+				ContextMenuBuilder.AddAction(FAssetSchemaAction_Dataflow_CreateNode_DataflowEdNode::CreateAction(ContextMenuBuilder.OwnerOfTemporaries, NodeTypeName));
+			}
+		}
+	}
+}
+
+const FPinConnectionResponse UDataflowSchema::CanCreateConnection(const UEdGraphPin* PinA, const UEdGraphPin* PinB) const
+{
+	// Make sure the pins are not on the same node
+	if (PinA->GetOwningNode() != PinB->GetOwningNode())
+	{
+		// Make sure types match. 
+		if (PinA->PinType == PinB->PinType)
+		{
+			return FPinConnectionResponse(CONNECT_RESPONSE_MAKE, LOCTEXT("PinConnect", "Connect nodes"));
+		}
+	}
+
+	return FPinConnectionResponse(CONNECT_RESPONSE_DISALLOW, LOCTEXT("PinErrorSameNode", "Both are on the same node"));
+}
+
+
+#undef LOCTEXT_NAMESPACE
diff --git a/Engine/Plugins/Experimental/Dataflow/Source/DataflowEditor/Public/Dataflow/DataflowAssetActions.h b/Engine/Plugins/Experimental/Dataflow/Source/DataflowEditor/Public/Dataflow/DataflowAssetActions.h
new file mode 100644
--- /dev/null
+++ b/Engine/Plugins/Experimental/Dataflow/Source/DataflowEditor/Public/Dataflow/DataflowAssetActions.h
@@ -0,0 +1,17 @@
+// Copyright Epic Games, Inc. All Rights Reserved.
+
+#pragma once
+
+#include "AssetTypeActions_Base.h"
+
+class FDataflowAssetActions : public FAssetTypeActions_Base
+{
+public:
+	virtual FText GetName() const override;
+	virtual UClass* GetSupportedClass() const override;
+	virtual FColor GetTypeColor() const override;
+	virtual void GetActions(const TArray<UObject*>& InObjects, class FMenuBuilder& MenuBuilder) override;
+	virtual void OpenAssetEditor(const TArray<UObject*>& InObjects, TSharedPtr<IToolkitHost> EditWithinLevelEditor = TSharedPtr<IToolkitHost>()) override;
+	virtual uint32 GetCategories() override;
+	virtual FText GetAssetDescription(const struct FAssetData& AssetData) const override;
+};
\ No newline at end of file
diff --git a/Engine/Plugins/Experimental/Dataflow/Source/DataflowEditor/Public/Dataflow/DataflowAssetFactory.h b/Engine/Plugins/Experimental/Dataflow/Source/DataflowEditor/Public/Dataflow/DataflowAssetFactory.h
new file mode 100644
--- /dev/null
+++ b/Engine/Plugins/Experimental/Dataflow/Source/DataflowEditor/Public/Dataflow/DataflowAssetFactory.h
@@ -0,0 +1,32 @@
+// Copyright Epic Games, Inc. All Rights Reserved.
+
+#pragma once
+
+#include "Factories/Factory.h"
+
+#include "DataflowAssetFactory.generated.h"
+
+UCLASS()
+class UDataflowAssetFactory : public UFactory
+{
+	GENERATED_BODY()
+public:
+
+	UDataflowAssetFactory();
+
+	/** UFactory Interface */
+	bool     CanCreateNew() const override;
+	bool     FactoryCanImport(const FString& Filename) override;
+	UObject* FactoryCreateNew(UClass*           InClass,
+							  UObject*          InParent,
+							  FName             InName,
+							  EObjectFlags      Flags,
+							  UObject*          Context,
+							  FFeedbackContext* Warn) override;
+
+	bool ShouldShowInNewMenu() const override;
+	bool ConfigureProperties() override;
+	/** End UFactory Interface */
+
+private:
+};
diff --git a/Engine/Plugins/Experimental/Dataflow/Source/DataflowEditor/Public/Dataflow/DataflowEditorActions.h b/Engine/Plugins/Experimental/Dataflow/Source/DataflowEditor/Public/Dataflow/DataflowEditorActions.h
new file mode 100644
--- /dev/null
+++ b/Engine/Plugins/Experimental/Dataflow/Source/DataflowEditor/Public/Dataflow/DataflowEditorActions.h
@@ -0,0 +1,36 @@
+// Copyright Epic Games, Inc. All Rights Reserved.
+
+#pragma once
+
+#include "CoreMinimal.h"
+#include "Framework/Commands/Commands.h"
+#include "Styling/AppStyle.h"
+
+class FDataflowEditorCommandsImpl : public TCommands<FDataflowEditorCommandsImpl>
+{
+public:
+
+	FDataflowEditorCommandsImpl()
+		: TCommands<FDataflowEditorCommandsImpl>( TEXT("DataflowEditor"), NSLOCTEXT("Contexts", "DataflowEditor", "Scene Graph Editor"), NAME_None, FAppStyle::GetAppStyleSetName() )
+	{
+	}	
+
+	virtual ~FDataflowEditorCommandsImpl()
+	{
+	}
+
+	DATAFLOWEDITOR_API virtual void RegisterCommands() override;
+
+	TSharedPtr< FUICommandInfo > EvaluateNode;
+	TMap< FName, TSharedPtr<FUICommandInfo> > CreateNodesMap;
+};
+
+class DATAFLOWEDITOR_API FDataflowEditorCommands
+{
+public:
+	static void Register();
+
+	static const FDataflowEditorCommandsImpl& Get();
+
+	static void Unregister();
+};
diff --git a/Engine/Plugins/Experimental/Dataflow/Source/DataflowEditor/Public/Dataflow/DataflowEditorPlugin.h b/Engine/Plugins/Experimental/Dataflow/Source/DataflowEditor/Public/Dataflow/DataflowEditorPlugin.h
new file mode 100644
--- /dev/null
+++ b/Engine/Plugins/Experimental/Dataflow/Source/DataflowEditor/Public/Dataflow/DataflowEditorPlugin.h
@@ -0,0 +1,63 @@
+// Copyright Epic Games, Inc. All Rights Reserved.
+
+#pragma once
+
+#include "CoreMinimal.h"
+#include "Modules/ModuleInterface.h"
+#include "Modules/ModuleManager.h"
+//#include "Toolkits/IToolkit.h"
+
+//class ISlateStyle;
+//class FSlateStyleSet;
+class FDataflowSNodeFactory;
+class FDataflowAssetActions;
+
+/**
+ * The public interface to this module
+ */
+class IDataflowEditorPlugin : public IModuleInterface
+{
+public:
+	virtual void StartupModule();
+	virtual void ShutdownModule();
+
+	TSharedRef<FAssetEditorToolkit> CreateDataflowAssetEditor(const EToolkitMode::Type Mode, const TSharedPtr< IToolkitHost >& InitToolkitHost, UObject* Dataflow);
+
+
+	/**
+	 * Singleton-like access to this module's interface.  This is just for convenience!
+	 * Beware of calling this during the shutdown phase, though.  Your module might have been unloaded already.
+	 *
+	 * @return Returns singleton instance, loading the module on demand if needed
+	 */
+	static inline IDataflowEditorPlugin& Get()
+	{
+		return FModuleManager::LoadModuleChecked< IDataflowEditorPlugin >( "DataflowEditor" );
+	}
+
+	/**
+	 * Checks to see if this module is loaded and ready.  It is only valid to call Get() if IsAvailable() returns true.
+	 *
+	 * @return True if the module is loaded and ready to use
+	 */
+	static inline bool IsAvailable()
+	{
+		return FModuleManager::Get().IsModuleLoaded( "DataflowEditor" );
+	}
+
+	//TSharedPtr<FSlateStyleSet> GetStyleSet() { return StyleSet; }
+	//static FName GetEditorStyleName();
+	//static const ISlateStyle* GetEditorStyle();
+
+private:
+	void RegisterMenus();
+
+private:
+
+	FDataflowAssetActions* DataflowAssetActions;
+
+	//TSharedPtr<FSlateStyleSet> StyleSet;
+	TSharedPtr<FDataflowSNodeFactory> DataflowSNodeFactory;
+
+};
+
diff --git a/Engine/Plugins/Experimental/Dataflow/Source/DataflowEditor/Public/Dataflow/DataflowSNode.h b/Engine/Plugins/Experimental/Dataflow/Source/DataflowEditor/Public/Dataflow/DataflowSNode.h
new file mode 100644
--- /dev/null
+++ b/Engine/Plugins/Experimental/Dataflow/Source/DataflowEditor/Public/Dataflow/DataflowSNode.h
@@ -0,0 +1,56 @@
+// Copyright Epic Games, Inc. All Rights Reserved.
+
+#pragma once
+
+#include "EdGraphUtilities.h"
+#include "SGraphNode.h"
+
+#include "DataflowSNode.generated.h"
+
+class UDataflowEdNode;
+
+//
+// SDataflowEdNode
+//
+
+class DATAFLOWEDITOR_API SDataflowEdNode : public SGraphNode
+{
+	typedef SGraphNode Super;
+
+public:
+	SLATE_BEGIN_ARGS(SDataflowEdNode)
+		: _GraphNodeObj(nullptr)
+	{}
+
+	SLATE_ARGUMENT(UDataflowEdNode*, GraphNodeObj)
+
+	SLATE_END_ARGS()
+	
+	void Construct(const FArguments& InArgs, UDataflowEdNode* InNode);
+
+	// SGraphNode interface
+	virtual FReply OnMouseButtonDoubleClick(const FGeometry& InMyGeometry, const FPointerEvent& InMouseEvent) override;
+
+};
+
+//
+// Action to add a node to the graph
+//
+USTRUCT()
+struct DATAFLOWEDITOR_API FAssetSchemaAction_Dataflow_CreateNode_DataflowEdNode : public FEdGraphSchemaAction
+{
+	GENERATED_USTRUCT_BODY();
+
+public:
+	FAssetSchemaAction_Dataflow_CreateNode_DataflowEdNode() : FEdGraphSchemaAction() {}
+
+	FAssetSchemaAction_Dataflow_CreateNode_DataflowEdNode(const FName& InType, const FText& InNodeCategory, const FText& InMenuDesc, const FText& InToolTip, const int32 InGrouping)
+		: FEdGraphSchemaAction(InNodeCategory, InMenuDesc, InToolTip, InGrouping), NodeTypeName(InType) {}
+
+	static TSharedPtr<FAssetSchemaAction_Dataflow_CreateNode_DataflowEdNode> CreateAction(UEdGraph* ParentGraph, const FName & NodeTypeName);
+
+	virtual UEdGraphNode* PerformAction(class UEdGraph* ParentGraph, UEdGraphPin* FromPin, const FVector2D Location, bool bSelectNewNode = true) override;
+	//virtual void AddReferencedObjects(FReferenceCollector& Collector) override;
+
+	FName NodeTypeName;
+};
diff --git a/Engine/Plugins/Experimental/Dataflow/Source/DataflowEditor/Public/Dataflow/DataflowSNodeFactories.h b/Engine/Plugins/Experimental/Dataflow/Source/DataflowEditor/Public/Dataflow/DataflowSNodeFactories.h
new file mode 100644
--- /dev/null
+++ b/Engine/Plugins/Experimental/Dataflow/Source/DataflowEditor/Public/Dataflow/DataflowSNodeFactories.h
@@ -0,0 +1,13 @@
+// Copyright Epic Games, Inc. All Rights Reserved.
+
+#pragma once
+
+#include "EdGraphUtilities.h"
+#include "SGraphNode.h"
+#include "EdGraph/EdGraphNode.h"
+
+class DATAFLOWEDITOR_API FDataflowSNodeFactory : public FGraphPanelNodeFactory
+{
+	virtual TSharedPtr<class SGraphNode> CreateNode(class UEdGraphNode* InNode) const override;
+
+};
diff --git a/Engine/Plugins/Experimental/Dataflow/Source/DataflowEditor/Public/Dataflow/DataflowSchema.h b/Engine/Plugins/Experimental/Dataflow/Source/DataflowEditor/Public/Dataflow/DataflowSchema.h
new file mode 100644
--- /dev/null
+++ b/Engine/Plugins/Experimental/Dataflow/Source/DataflowEditor/Public/Dataflow/DataflowSchema.h
@@ -0,0 +1,21 @@
+// Copyright Epic Games, Inc. All Rights Reserved.
+
+#pragma once
+
+#include "DataflowSchema.generated.h"
+
+class UDataflow; 
+
+UCLASS()
+class DATAFLOWEDITOR_API UDataflowSchema : public UEdGraphSchema
+{
+	GENERATED_BODY()
+public:
+	UDataflowSchema();
+
+	//~ Begin EdGraphSchema Interface
+	virtual const FPinConnectionResponse CanCreateConnection(const UEdGraphPin* PinA, const UEdGraphPin* PinB) const override;
+	virtual void GetContextMenuActions(class UToolMenu* Menu, class UGraphNodeContextMenuContext* Context) const override;
+	virtual void GetGraphContextActions(FGraphContextMenuBuilder& ContextMenuBuilder) const override;
+	//~ End EdGraphSchema Interface
+};
diff --git a/Engine/Plugins/Experimental/EvalGraph/EvalGraph.uplugin b/Engine/Plugins/Experimental/EvalGraph/EvalGraph.uplugin
deleted file mode 100644
diff --git a/Engine/Plugins/Experimental/EvalGraph/Source/EvalGraphEditor/EvalGraphEditor.Build.cs b/Engine/Plugins/Experimental/EvalGraph/Source/EvalGraphEditor/EvalGraphEditor.Build.cs
deleted file mode 100644
diff --git a/Engine/Plugins/Experimental/EvalGraph/Source/EvalGraphEditor/Private/EvalGraph/EvalGraphAssetActions.cpp b/Engine/Plugins/Experimental/EvalGraph/Source/EvalGraphEditor/Private/EvalGraph/EvalGraphAssetActions.cpp
deleted file mode 100644
diff --git a/Engine/Plugins/Experimental/EvalGraph/Source/EvalGraphEditor/Private/EvalGraph/EvalGraphAssetFactory.cpp b/Engine/Plugins/Experimental/EvalGraph/Source/EvalGraphEditor/Private/EvalGraph/EvalGraphAssetFactory.cpp
deleted file mode 100644
diff --git a/Engine/Plugins/Experimental/EvalGraph/Source/EvalGraphEditor/Private/EvalGraph/EvalGraphEditorActions.cpp b/Engine/Plugins/Experimental/EvalGraph/Source/EvalGraphEditor/Private/EvalGraph/EvalGraphEditorActions.cpp
deleted file mode 100644
diff --git a/Engine/Plugins/Experimental/EvalGraph/Source/EvalGraphEditor/Private/EvalGraph/EvalGraphEditorPlugin.cpp b/Engine/Plugins/Experimental/EvalGraph/Source/EvalGraphEditor/Private/EvalGraph/EvalGraphEditorPlugin.cpp
deleted file mode 100644
diff --git a/Engine/Plugins/Experimental/EvalGraph/Source/EvalGraphEditor/Private/EvalGraph/EvalGraphEditorToolkit.cpp b/Engine/Plugins/Experimental/EvalGraph/Source/EvalGraphEditor/Private/EvalGraph/EvalGraphEditorToolkit.cpp
deleted file mode 100644
diff --git a/Engine/Plugins/Experimental/EvalGraph/Source/EvalGraphEditor/Private/EvalGraph/EvalGraphEditorToolkit.h b/Engine/Plugins/Experimental/EvalGraph/Source/EvalGraphEditor/Private/EvalGraph/EvalGraphEditorToolkit.h
deleted file mode 100644
diff --git a/Engine/Plugins/Experimental/EvalGraph/Source/EvalGraphEditor/Private/EvalGraph/EvalGraphSNode.cpp b/Engine/Plugins/Experimental/EvalGraph/Source/EvalGraphEditor/Private/EvalGraph/EvalGraphSNode.cpp
deleted file mode 100644
diff --git a/Engine/Plugins/Experimental/EvalGraph/Source/EvalGraphEditor/Private/EvalGraph/EvalGraphSNodeFactories.cpp b/Engine/Plugins/Experimental/EvalGraph/Source/EvalGraphEditor/Private/EvalGraph/EvalGraphSNodeFactories.cpp
deleted file mode 100644
diff --git a/Engine/Plugins/Experimental/EvalGraph/Source/EvalGraphEditor/Private/EvalGraph/EvalGraphSchema.cpp b/Engine/Plugins/Experimental/EvalGraph/Source/EvalGraphEditor/Private/EvalGraph/EvalGraphSchema.cpp
deleted file mode 100644
diff --git a/Engine/Plugins/Experimental/EvalGraph/Source/EvalGraphEditor/Public/EvalGraph/EvalGraphAssetActions.h b/Engine/Plugins/Experimental/EvalGraph/Source/EvalGraphEditor/Public/EvalGraph/EvalGraphAssetActions.h
deleted file mode 100644
diff --git a/Engine/Plugins/Experimental/EvalGraph/Source/EvalGraphEditor/Public/EvalGraph/EvalGraphAssetFactory.h b/Engine/Plugins/Experimental/EvalGraph/Source/EvalGraphEditor/Public/EvalGraph/EvalGraphAssetFactory.h
deleted file mode 100644
diff --git a/Engine/Plugins/Experimental/EvalGraph/Source/EvalGraphEditor/Public/EvalGraph/EvalGraphEditorActions.h b/Engine/Plugins/Experimental/EvalGraph/Source/EvalGraphEditor/Public/EvalGraph/EvalGraphEditorActions.h
deleted file mode 100644
diff --git a/Engine/Plugins/Experimental/EvalGraph/Source/EvalGraphEditor/Public/EvalGraph/EvalGraphEditorPlugin.h b/Engine/Plugins/Experimental/EvalGraph/Source/EvalGraphEditor/Public/EvalGraph/EvalGraphEditorPlugin.h
deleted file mode 100644
diff --git a/Engine/Plugins/Experimental/EvalGraph/Source/EvalGraphEditor/Public/EvalGraph/EvalGraphSNode.h b/Engine/Plugins/Experimental/EvalGraph/Source/EvalGraphEditor/Public/EvalGraph/EvalGraphSNode.h
deleted file mode 100644
diff --git a/Engine/Plugins/Experimental/EvalGraph/Source/EvalGraphEditor/Public/EvalGraph/EvalGraphSNodeFactories.h b/Engine/Plugins/Experimental/EvalGraph/Source/EvalGraphEditor/Public/EvalGraph/EvalGraphSNodeFactories.h
deleted file mode 100644
diff --git a/Engine/Plugins/Experimental/EvalGraph/Source/EvalGraphEditor/Public/EvalGraph/EvalGraphSchema.h b/Engine/Plugins/Experimental/EvalGraph/Source/EvalGraphEditor/Public/EvalGraph/EvalGraphSchema.h
deleted file mode 100644
diff --git a/Engine/Source/Runtime/Experimental/Dataflow/Core/DataflowCore.Build.cs b/Engine/Source/Runtime/Experimental/Dataflow/Core/DataflowCore.Build.cs
new file mode 100644
--- /dev/null
+++ b/Engine/Source/Runtime/Experimental/Dataflow/Core/DataflowCore.Build.cs
@@ -0,0 +1,19 @@
+// Copyright Epic Games, Inc. All Rights Reserved.
+
+namespace UnrealBuildTool.Rules
+{
+	public class DataflowCore : ModuleRules
+	{
+        public DataflowCore(ReadOnlyTargetRules Target) : base(Target)
+		{
+			PublicDependencyModuleNames.AddRange(
+				new string[]
+				{
+					"Core",
+					"CoreUObject",
+					"Chaos",
+				}
+				);
+		}
+	}
+}
diff --git a/Engine/Source/Runtime/Experimental/Dataflow/Core/Private/Dataflow/DataflowConnection.cpp b/Engine/Source/Runtime/Experimental/Dataflow/Core/Private/Dataflow/DataflowConnection.cpp
new file mode 100644
--- /dev/null
+++ b/Engine/Source/Runtime/Experimental/Dataflow/Core/Private/Dataflow/DataflowConnection.cpp
@@ -0,0 +1,24 @@
+// Copyright Epic Games, Inc. All Rights Reserved.
+
+#include "Dataflow/DataflowConnection.h"
+
+#include "Dataflow/DataflowNodeParameters.h"
+#include "Dataflow/DataflowNode.h"
+
+namespace Dataflow
+{
+	FConnection::FConnection(FPin::EDirection InDirection, FName InType, FName InName, FNode* InOwningNode, FGuid InGuid)
+		: Direction(InDirection)
+		, Type(InType)
+		, Name(InName)
+		, Guid(InGuid)
+		, OwningNode(InOwningNode)
+	{}
+
+	void FConnection::BindInput(FNode* InNode, FConnection* That) 
+	{ InNode->AddInput(That); }
+	void FConnection::BindOutput(FNode* InNode, FConnection* That) 
+	{ InNode->AddOutput(That); }
+
+}
+
diff --git a/Engine/Source/Runtime/Experimental/Dataflow/Core/Private/Dataflow/DataflowCoreLogPrivate.h b/Engine/Source/Runtime/Experimental/Dataflow/Core/Private/Dataflow/DataflowCoreLogPrivate.h
new file mode 100644
--- /dev/null
+++ b/Engine/Source/Runtime/Experimental/Dataflow/Core/Private/Dataflow/DataflowCoreLogPrivate.h
@@ -0,0 +1,12 @@
+// Copyright Epic Games, Inc. All Rights Reserved.
+
+#pragma once
+
+#include "CoreMinimal.h"
+#include "Stats/Stats.h"
+
+// You should place include statements to your module's private header files here.  You only need to
+// add includes for headers that are used in most of your module's source files though.
+
+
+DECLARE_STATS_GROUP(TEXT("Dataflow"), STATGROUP_Dataflow, STATCAT_Advanced);
diff --git a/Engine/Source/Runtime/Experimental/Dataflow/Core/Private/Dataflow/DataflowCoreModule.cpp b/Engine/Source/Runtime/Experimental/Dataflow/Core/Private/Dataflow/DataflowCoreModule.cpp
new file mode 100644
--- /dev/null
+++ b/Engine/Source/Runtime/Experimental/Dataflow/Core/Private/Dataflow/DataflowCoreModule.cpp
@@ -0,0 +1,20 @@
+// Copyright Epic Games, Inc. All Rights Reserved.
+
+#include "Dataflow/DataflowCoreModule.h"
+
+#include "CoreMinimal.h"
+#include "Dataflow/DataflowNodeFactory.h"
+#include "Modules/ModuleInterface.h"
+#include "Modules/ModuleManager.h"
+
+#define LOCTEXT_NAMESPACE "DataflowCore"
+
+void IDataflowCoreModule::StartupModule()
+{/*Never Called ?*/}
+
+void IDataflowCoreModule::ShutdownModule()
+{/*Never Called ?*/ }
+
+IMPLEMENT_MODULE(IDataflowCoreModule, DataflowCore)
+
+#undef LOCTEXT_NAMESPACE
diff --git a/Engine/Source/Runtime/Experimental/Dataflow/Core/Private/Dataflow/DataflowGraph.cpp b/Engine/Source/Runtime/Experimental/Dataflow/Core/Private/Dataflow/DataflowGraph.cpp
new file mode 100644
--- /dev/null
+++ b/Engine/Source/Runtime/Experimental/Dataflow/Core/Private/Dataflow/DataflowGraph.cpp
@@ -0,0 +1,185 @@
+// Copyright Epic Games, Inc. All Rights Reserved.
+
+#include "Dataflow/DataflowGraph.h"
+
+#include "Dataflow/DataflowInputOutput.h"
+#include "Dataflow/DataflowNodeParameters.h"
+#include "Dataflow/DataflowNodeFactory.h"
+#include "Logging/LogMacros.h"
+
+DEFINE_LOG_CATEGORY_STATIC(DATAFLOW_LOG, Error, All);
+
+namespace Dataflow
+{
+
+	FGraph::FGraph(FGuid InGuid)
+		: Guid(InGuid)
+	{
+	}
+
+
+	void FGraph::RemoveNode(TSharedPtr<FNode> Node)
+	{
+		for (FConnection* Output : Node->GetOutputs())
+		{
+			if (Output)
+			{
+				for (FConnection* Input : Output->GetConnectedInputs())
+				{
+					if (Input)
+					{
+						Disconnect(Output, Input);
+					}
+				}
+			}
+		}
+		for (FConnection* Input : Node->GetInputs())
+		{
+			if (Input)
+			{
+				TArray<FConnection*> Outputs = Input->GetConnectedOutputs();
+				for (FConnection* Out : Outputs)
+				{
+					if (Out)
+					{
+						Disconnect(Out, Input);
+					}
+				}
+			}
+		}
+		Nodes.Remove(Node);
+	}
+
+	void FGraph::ClearConnections(FConnection* Connection)
+	{
+		// Todo(dataflow) : do this without triggering a invalidation. 
+		//            or implement a better sync for the EdGraph and DataflowGraph
+		if (Connection->GetDirection() == FPin::EDirection::INPUT)
+		{
+			TArray<FConnection*> BaseOutputs = Connection->GetConnectedOutputs();
+			for (FConnection* Output : BaseOutputs)
+			{
+				Disconnect(Connection, Output);
+			}
+		}
+		else if (Connection->GetDirection() == FPin::EDirection::OUTPUT)
+		{
+			TArray<FConnection*> BaseInputs = Connection->GetConnectedInputs();
+			for (FConnection* Input : BaseInputs)
+			{
+				Disconnect(Input, Connection);
+			}
+		}
+	}
+
+
+	void FGraph::Connect(FConnection* Input, FConnection* Output)
+	{
+		if (ensure(Input && Output))
+		{
+			Input->AddConnection(Output);
+			Output->AddConnection(Input);
+			Connections.Add(FLink(Input->GetGuid(), Output->GetGuid()));
+		}
+	}
+
+	void FGraph::Disconnect(FConnection* Input, FConnection* Output)
+	{
+		Input->RemoveConnection(Output);
+		Output->RemoveConnection(Input);
+		Connections.RemoveSwap(FLink(Input->GetGuid(), Output->GetGuid()));
+	}
+
+	void FGraph::Serialize(FArchive& Ar)
+	{
+		Ar << Guid;
+		if (Ar.IsSaving())
+		{
+			FGuid ArGuid;
+			FName ArType, ArName;
+			int32 ArNum = Nodes.Num();
+			Ar << ArNum;
+			for (TSharedPtr<FNode> Node : Nodes)
+			{
+				ArGuid = Node->GetGuid();
+				ArType = Node->GetType();
+				ArName = Node->GetName();
+				Ar << ArGuid << ArType << ArName;
+
+
+				TArray< FConnection* > IO = Node->GetOutputs();  IO.Append(Node->GetInputs());
+				ArNum = IO.Num();
+				Ar << ArNum;
+				for (FConnection* Conn : IO)
+				{
+					ArGuid = Conn->GetGuid();
+					ArType = Conn->GetType();
+					ArName = Conn->GetName();
+					Ar << ArGuid << ArType << ArName;
+				}
+
+				Node->SerializeInternal(Ar);
+			}
+
+			Ar << Connections;
+
+		}
+		else if( Ar.IsLoading())
+		{
+			FGuid ArGuid;
+			FName ArType, ArName;
+			int32 ArNum = 0;
+
+			TMap<FGuid, FConnection* > ConnectionGuidMap;
+
+			Ar << ArNum;
+			for (int32 Ndx = ArNum; Ndx > 0; Ndx--)
+			{
+				Ar << ArGuid << ArType << ArName;
+				if (TSharedPtr<FNode> Node = FNodeFactory::GetInstance()->NewNodeFromRegisteredType(*this, { ArGuid,ArType,ArName }))
+				{
+					int ArNumInner;
+					Ar << ArNumInner;
+					TArray< FConnection* > IO = Node->GetOutputs();  IO.Append(Node->GetInputs());
+					for (int Cdx=0;Cdx< ArNumInner; Cdx++)
+					{
+						Ar << ArGuid << ArType << ArName;
+						if (Cdx < IO.Num())
+						{
+							//if ((EGraphConnectionType)ArIntType!=IO[Cdx]->GetType())
+							//{
+							//	FName TypeName = GraphConnectionTypeName((EGraphConnectionType)ArIntType);
+							//	FName FromType = GraphConnectionTypeName(IO[Cdx]->GetType());
+							//	UE_LOG(DATAFLOW_LOG, Fatal, TEXT("Type mismatch in input file for Ed::FGraphNode Input/Output (%s,%s) to type (%s)"),
+							//		*TypeName.ToString(), *ArName.ToString(), *FromType.ToString());
+							//}
+
+							IO[Cdx]->SetGuid(ArGuid);
+							ensure(!ConnectionGuidMap.Contains(ArGuid));
+							ConnectionGuidMap.Add(ArGuid, IO[Cdx]);
+						}
+					}
+
+					Node->SerializeInternal(Ar);
+				}
+				else
+				{
+					ensureMsgf(false, TEXT("Warning: Missing registered node type on load. Graph will fail to evaluate (%s %s)"), *ArType.ToString(), *ArName.ToString());
+				}
+				
+			}
+
+			TArray< FLink > LocalConnections;
+			Ar << LocalConnections;
+			for (const FLink& Con : LocalConnections)
+			{
+				ensure(ConnectionGuidMap.Contains(Con.Input));
+				ensure(ConnectionGuidMap.Contains(Con.Output));
+				ensure(ConnectionGuidMap[Con.Input]->GetType() == ConnectionGuidMap[Con.Output]->GetType());
+				Connect(ConnectionGuidMap[Con.Input], ConnectionGuidMap[Con.Output]);
+			}
+
+		}
+	}
+}
+
diff --git a/Engine/Source/Runtime/Experimental/Dataflow/Core/Private/Dataflow/DataflowInputOutput.cpp b/Engine/Source/Runtime/Experimental/Dataflow/Core/Private/Dataflow/DataflowInputOutput.cpp
new file mode 100644
--- /dev/null
+++ b/Engine/Source/Runtime/Experimental/Dataflow/Core/Private/Dataflow/DataflowInputOutput.cpp
@@ -0,0 +1,18 @@
+// Copyright Epic Games, Inc. All Rights Reserved.
+
+#include "Dataflow/DataflowInputOutput.h"
+
+#include "Dataflow/DataflowNodeParameters.h"
+#include "Dataflow/DataflowNode.h"
+
+namespace Dataflow
+{
+	// Base Types
+	EVAL_GRAPH_CONNECTION_TYPE(bool, Bool)
+	EVAL_GRAPH_CONNECTION_TYPE(char, Char)
+	EVAL_GRAPH_CONNECTION_TYPE(int, Integer)	
+	EVAL_GRAPH_CONNECTION_TYPE(uint8, UInt8)
+	EVAL_GRAPH_CONNECTION_TYPE(float, Float)
+	EVAL_GRAPH_CONNECTION_TYPE(double, Double)
+}
+
diff --git a/Engine/Source/Runtime/Experimental/Dataflow/Core/Private/Dataflow/DataflowNode.cpp b/Engine/Source/Runtime/Experimental/Dataflow/Core/Private/Dataflow/DataflowNode.cpp
new file mode 100644
--- /dev/null
+++ b/Engine/Source/Runtime/Experimental/Dataflow/Core/Private/Dataflow/DataflowNode.cpp
@@ -0,0 +1,72 @@
+// Copyright Epic Games, Inc. All Rights Reserved.
+
+#include "Dataflow/DataflowNode.h"
+
+#include "Dataflow/DataflowNode.h"
+
+namespace Dataflow
+{
+	void FNode::AddInput(FConnection* InPtr)
+	{ 
+		for (FConnection* In : Inputs)
+		{
+			ensureMsgf(!In->GetName().IsEqual(InPtr->GetName()), TEXT("Add Input Failed: Existing Node input already defined with name (%s)"), *InPtr->GetName().ToString());
+		}
+		Inputs.Add(InPtr); 
+	}
+
+	void FNode::AddOutput(FConnection* InPtr)
+	{ 
+		for (FConnection* Out : Outputs)
+		{
+			ensureMsgf(!Out->GetName().IsEqual(InPtr->GetName()), TEXT("Add Output Failed: Existing Node output already defined with name (%s)"), *InPtr->GetName().ToString());
+		}
+		Outputs.Add(InPtr);
+	}
+
+
+
+	TArray<FPin> FNode::GetPins() const
+	{
+		TArray<FPin> RetVal;
+		for (FConnection* Con : Inputs)
+			RetVal.Add({ FPin::EDirection::INPUT,Con->GetType(), Con->GetName() });
+		for (FConnection* Con : Outputs)
+			RetVal.Add({ FPin::EDirection::OUTPUT,Con->GetType(), Con->GetName() });
+		return RetVal;
+	}
+
+	void FNode::InvalidateOutputs()
+	{
+		for(FConnection * Output :  Outputs)
+		{
+			Output->Invalidate();
+		}
+	}
+
+	FConnection* FNode::FindInput(FName InName) const
+	{
+		for (FConnection* Input : Inputs)
+		{
+			if (Input->GetName().IsEqual(InName))
+			{
+				return Input;
+			}
+		}
+		return nullptr;
+	}
+
+	FConnection* FNode::FindOutput(FName InName) const
+	{
+		for (FConnection* Output : Outputs)
+		{
+			if (Output->GetName().IsEqual(InName))
+			{
+				return Output;
+			}
+		}
+		return nullptr;
+	}
+
+}
+
diff --git a/Engine/Source/Runtime/Experimental/Dataflow/Core/Private/Dataflow/DataflowNodeFactory.cpp b/Engine/Source/Runtime/Experimental/Dataflow/Core/Private/Dataflow/DataflowNodeFactory.cpp
new file mode 100644
--- /dev/null
+++ b/Engine/Source/Runtime/Experimental/Dataflow/Core/Private/Dataflow/DataflowNodeFactory.cpp
@@ -0,0 +1,20 @@
+// Copyright Epic Games, Inc. All Rights Reserved.
+
+#include "Dataflow/DataflowNodeFactory.h"
+
+#include "Dataflow/DataflowNode.h"
+
+namespace Dataflow
+{
+	FNodeFactory* FNodeFactory::Instance = nullptr;
+
+	TSharedPtr<FNode> FNodeFactory::NewNodeFromRegisteredType(FGraph& Graph, const FNewNodeParameters& Param)
+	{ 
+		if (ClassMap.Contains(Param.Type))
+		{
+			return Graph.AddNode(ClassMap[Param.Type](Param));
+		}
+		return TSharedPtr<FNode>(nullptr);
+	}
+}
+
diff --git a/Engine/Source/Runtime/Experimental/Dataflow/Core/Private/Dataflow/DataflowNodeParameters.cpp b/Engine/Source/Runtime/Experimental/Dataflow/Core/Private/Dataflow/DataflowNodeParameters.cpp
new file mode 100644
--- /dev/null
+++ b/Engine/Source/Runtime/Experimental/Dataflow/Core/Private/Dataflow/DataflowNodeParameters.cpp
@@ -0,0 +1,4 @@
+// Copyright Epic Games, Inc. All Rights Reserved.
+
+#include "Dataflow/DataflowNodeParameters.h"
+
diff --git a/Engine/Source/Runtime/Experimental/Dataflow/Core/Public/Dataflow/Dataflow.h b/Engine/Source/Runtime/Experimental/Dataflow/Core/Public/Dataflow/Dataflow.h
new file mode 100644
--- /dev/null
+++ b/Engine/Source/Runtime/Experimental/Dataflow/Core/Public/Dataflow/Dataflow.h
@@ -0,0 +1,10 @@
+// Copyright Epic Games, Inc. All Rights Reserved.
+
+#pragma once
+
+#include "Dataflow/DataflowGraph.h"
+#include "Dataflow/DataflowConnection.h"
+#include "Dataflow/DataflowInputOutput.h"
+#include "Dataflow/DataflowNode.h"
+#include "Dataflow/DataflowNodeFactory.h"
+#include "Dataflow/DataflowNodeParameters.h"
\ No newline at end of file
diff --git a/Engine/Source/Runtime/Experimental/Dataflow/Core/Public/Dataflow/DataflowConnection.h b/Engine/Source/Runtime/Experimental/Dataflow/Core/Public/Dataflow/DataflowConnection.h
new file mode 100644
--- /dev/null
+++ b/Engine/Source/Runtime/Experimental/Dataflow/Core/Public/Dataflow/DataflowConnection.h
@@ -0,0 +1,71 @@
+// Copyright Epic Games, Inc. All Rights Reserved.
+
+#pragma once
+
+#include "CoreMinimal.h"
+#include "Dataflow/DataflowNodeParameters.h"
+
+namespace Dataflow
+{
+	class FNode;
+
+	template<class T> inline DATAFLOWCORE_API FName GraphConnectionTypeName();
+
+	struct FPin
+	{
+		enum class EDirection : uint8 {
+			NONE = 0,
+			INPUT,
+			OUTPUT
+		};
+		EDirection Direction;
+		FName Type;
+		FName Name;
+	};
+
+	//
+	// Input Output Base
+	//
+
+	class DATAFLOWCORE_API FConnection
+	{
+
+	protected:
+		FPin::EDirection Direction;
+		FName Type;
+		FName Name;
+		FGuid  Guid;
+		FNode* OwningNode = nullptr;
+
+		friend class FNode;
+		static void BindInput(FNode* InNode, FConnection*);
+		static void BindOutput(FNode* InNode, FConnection*);
+
+
+	public:
+		FConnection(FPin::EDirection Direction, FName InType, FName InName, FNode* OwningNode = nullptr, FGuid InGuid = FGuid::NewGuid());
+		virtual ~FConnection() {};
+
+		FNode* GetOwningNode() { return OwningNode; }
+		const FNode* GetOwningNode() const { return OwningNode; }
+
+		FPin::EDirection GetDirection() const { return Direction; }
+
+		FName GetType() const { return Type; }
+
+		FGuid GetGuid() const { return Guid; }
+		void SetGuid(FGuid InGuid) { Guid = InGuid; }
+
+		FName GetName() const { return Name; }
+		void SetName(FName InName) { Name = InName; }
+
+
+		virtual bool AddConnection(FConnection* In) { return false; };
+		virtual bool RemoveConnection(FConnection* In) { return false; }
+
+		virtual TArray< FConnection* > GetConnectedInputs() { return TArray<FConnection* >(); }
+		virtual TArray< FConnection* > GetConnectedOutputs() { return TArray<FConnection* >(); }
+		virtual void Invalidate() {};
+
+	};
+}
diff --git a/Engine/Source/Runtime/Experimental/Dataflow/Core/Public/Dataflow/DataflowCoreModule.h b/Engine/Source/Runtime/Experimental/Dataflow/Core/Public/Dataflow/DataflowCoreModule.h
new file mode 100644
--- /dev/null
+++ b/Engine/Source/Runtime/Experimental/Dataflow/Core/Public/Dataflow/DataflowCoreModule.h
@@ -0,0 +1,30 @@
+// Copyright Epic Games, Inc. All Rights Reserved.
+
+#pragma once
+
+#include "CoreMinimal.h"
+#include "Dataflow/DataflowNodeFactory.h"
+#include "Modules/ModuleInterface.h"
+#include "Modules/ModuleManager.h"
+
+
+/**
+ * The public interface to this module
+ */
+class DATAFLOWCORE_API IDataflowCoreModule : public IModuleInterface
+{
+
+public:
+
+	// IModuleInterface interface
+
+	virtual void StartupModule() override;
+
+	virtual void ShutdownModule() override;
+
+	virtual bool SupportsDynamicReloading() override
+	{
+		return false;
+	}
+
+};
diff --git a/Engine/Source/Runtime/Experimental/Dataflow/Core/Public/Dataflow/DataflowGraph.h b/Engine/Source/Runtime/Experimental/Dataflow/Core/Public/Dataflow/DataflowGraph.h
new file mode 100644
--- /dev/null
+++ b/Engine/Source/Runtime/Experimental/Dataflow/Core/Public/Dataflow/DataflowGraph.h
@@ -0,0 +1,127 @@
+// Copyright Epic Games, Inc. All Rights Reserved.
+
+#pragma once
+
+#include "CoreMinimal.h"
+#include "Chaos/ChaosArchive.h"
+#include "Dataflow/DataflowNode.h"
+#include "Serialization/Archive.h"
+
+namespace Dataflow
+{
+	class FConnection;
+
+	struct FLink {
+		FGuid Input;
+		FGuid Output;
+
+		FLink() {}
+
+		FLink(FGuid InInput, FGuid InOutput)
+			: Input(InInput), Output(InOutput) {}
+
+		FLink(const FLink& Other)
+			: Input(Other.Input), Output(Other.Output) {}
+
+		bool operator==(const FLink& Other) const
+		{
+			return Equals(Other);
+		}
+
+		bool Equals(const FLink& Other) const
+		{
+			return Input == Other.Input
+				&& Output == Other.Output;
+		}
+	};
+
+
+	//
+	//
+	//
+	class DATAFLOWCORE_API FGraph
+	{
+
+		FGuid  Guid;
+		TArray< TSharedPtr<FNode> > Nodes;
+		TArray< FLink > Connections;
+
+	public:
+		FGraph(FGuid InGuid = FGuid::NewGuid());
+		virtual ~FGraph() {}
+
+		int NumNodes() { return Nodes.Num(); }
+
+		template<class T> TSharedPtr<T> AddNode(T* InNode)
+		{
+			TSharedPtr<T> NewNode(InNode);
+			Nodes.AddUnique(NewNode);
+			return NewNode;
+		}
+
+		template<class T> T* FindNode(FName InName)
+		{
+			for (TSharedPtr<FNode> Node : Nodes)
+			{
+				if (Node->GetName().ToString().Equals(InName.ToString()))
+				{
+					if (Node->GetType().ToString().Equals(T::Type.ToString()))
+					{
+						return reinterpret_cast<T*>(Node.Get()); // @todo(eg) : Can we do better here?
+					}
+				}
+			}
+			return nullptr;
+		}
+
+		TSharedPtr<FNode> FindBaseNode(FGuid InGuid)
+		{
+			for (TSharedPtr<FNode> Node : Nodes)
+			{
+				if (Node->GetGuid() == InGuid)
+				{
+					return Node;
+				}
+			}
+			return TSharedPtr<FNode>(nullptr);
+		}
+
+
+		void RemoveNode(TSharedPtr<FNode> Node);
+
+		void ClearConnections(FConnection*);
+		void Connect(FConnection* Input, FConnection* Output);
+		void Disconnect(FConnection* Input, FConnection* Output);
+
+		virtual void Serialize(FArchive& Ar);
+
+	};
+}
+
+FORCEINLINE FArchive& operator<<(FArchive& Ar, Dataflow::FLink& Value)
+{
+	Ar << Value.Input << Value.Output;
+	return Ar;
+}
+
+FORCEINLINE FArchive& operator<<(Chaos::FChaosArchive& Ar, Dataflow::FLink& Value)
+{
+	Ar << Value.Input << Value.Output;
+	return Ar;
+}
+
+FORCEINLINE FArchive& operator<<(FArchive& Ar, Dataflow::FGraph& Value)
+{
+	Value.Serialize(Ar);
+	return Ar;
+}
+
+FORCEINLINE FArchive& operator<<(Chaos::FChaosArchive& Ar, Dataflow::FGraph& Value)
+{
+	Value.Serialize(Ar);
+	return Ar;
+}
+
+
+
+
diff --git a/Engine/Source/Runtime/Experimental/Dataflow/Core/Public/Dataflow/DataflowInputOutput.h b/Engine/Source/Runtime/Experimental/Dataflow/Core/Public/Dataflow/DataflowInputOutput.h
new file mode 100644
--- /dev/null
+++ b/Engine/Source/Runtime/Experimental/Dataflow/Core/Public/Dataflow/DataflowInputOutput.h
@@ -0,0 +1,207 @@
+// Copyright Epic Games, Inc. All Rights Reserved.
+
+#pragma once
+
+#include "CoreMinimal.h"
+#include "Dataflow/DataflowNodeParameters.h"
+#include "Dataflow/DataflowNode.h"
+#include "Dataflow/DataflowConnection.h"
+
+namespace Dataflow
+{
+	template<typename T> class TOutput;
+
+	//
+	//  Input
+	//
+
+	template<class T>
+	struct DATAFLOWCORE_API TInputParameters {
+		TInputParameters(FName InName, FNode* InOwner, T InDefault = T())
+			: Type(GraphConnectionTypeName<T>())
+			, Name(InName)
+			, Owner(InOwner)
+			, Default(InDefault) {}
+		FName Type;
+		FName Name;
+		FNode* Owner = nullptr;
+		T Default;
+	};
+
+	template<class T>
+	class DATAFLOWCORE_API TInput : public FConnection
+	{
+		typedef FConnection Super;
+		friend class FConnection;
+
+		T Default;
+		TOutput<T>* Connection;
+	public:
+		TInput(const TInputParameters<T>& Param, FGuid InGuid = FGuid::NewGuid())
+			: FConnection(FPin::EDirection::INPUT, Param.Type, Param.Name, Param.Owner, InGuid)
+			, Default(Param.Default)
+			, Connection(nullptr)
+		{
+			Super::BindInput(Param.Owner, this);
+		}
+	
+		const T& GetDefault() const { return Default; }
+		
+		const TOutput<T>* GetConnection() const { return Connection; }
+		TOutput<T>* GetConnection() { return Connection; }
+
+		virtual bool AddConnection(FConnection* InOutput) override
+		{ 
+			ensure(Connection == nullptr);
+			if (ensure(InOutput->GetType()==this->GetType()))
+			{
+				Connection = (TOutput<T>*)InOutput;
+				return true;
+			}
+			return false;
+		}
+
+		virtual bool RemoveConnection(FConnection* InOutput) override
+		{ 
+			if (ensure(Connection == (TOutput<T>*)InOutput))
+			{
+				Connection = nullptr;
+				return true;
+			}
+			return false;
+		}
+
+		T GetValue(const FContext& Context)
+		{
+			if (Connection)
+			{
+				return Connection->Evaluate(Context);
+			}
+			return Default;
+		}
+
+		void SetValue(const T& Value, const FContext& Context)
+		{
+			Default = Value;
+			if (!Connection)
+			{
+				OwningNode->InvalidateOutputs();
+			}
+		}
+
+		virtual TArray< FConnection* > GetConnectedOutputs() override
+		{ 
+			TArray<FConnection* > RetList;
+			if (GetConnection())
+			{
+				RetList.Add(GetConnection());
+			}
+			return RetList;
+		}
+
+		virtual void Invalidate() override
+		{
+			OwningNode->InvalidateOutputs();
+		}
+
+	};
+
+
+	//
+	// Output
+	//
+
+	template<class T>
+	struct DATAFLOWCORE_API TOutputParameters {
+		TOutputParameters(FName InName, FNode * InOwner) 
+			: Type(GraphConnectionTypeName<T>())
+			, Name(InName)
+			, Owner(InOwner) {}
+		FName Type;
+		FName Name;
+		FNode* Owner = nullptr;
+	};
+
+	template<class T>
+	class DATAFLOWCORE_API TOutput : public FConnection
+	{
+		typedef FConnection Super;
+		friend class FConnection;
+
+		uint32 CacheKey = UINT_MAX;
+		TCacheValue<T> Cache;
+		TArray< TInput<T>* > Connections;
+	
+	public:
+		TOutput(const TOutputParameters<T>& Param, FGuid InGuid = FGuid::NewGuid())
+			: FConnection(FPin::EDirection::OUTPUT, Param.Type, Param.Name, Param.Owner, InGuid)
+		{
+			Super::BindOutput(Param.Owner, this);
+		}
+
+
+		const TArray<TInput<T>*>& GetConnections() const { return Connections; }
+		TArray<TInput<T>* >& GetConnections() { return Connections; }
+
+		virtual TArray<FConnection*> GetConnectedInputs() override
+		{ 
+			TArray<FConnection*> RetList;
+			RetList.Reserve(Connections.Num());
+			for (TInput<T>* Ptr : Connections) { RetList.Add(Ptr); }
+			return RetList;
+		}
+
+
+		virtual bool AddConnection(FConnection* InOutput) override
+		{
+			if (ensure(InOutput->GetType() == this->GetType()))
+			{
+				Connections.Add((TInput<T>*)InOutput);
+				return true;
+			}
+			return false;
+		}
+		
+		virtual bool RemoveConnection(FConnection* InInput) override { Connections.RemoveSwap((TInput<T>*)InInput); return true; }
+
+		void SetValue(T InVal, const FContext& Context)
+		{
+			CacheKey = Context.GetTypeHash();
+			Cache.Data = InVal;
+		}
+
+		T Evaluate(const FContext& Context)
+		{
+			if (CacheKey != Context.GetTypeHash())
+			{
+				OwningNode->Evaluate(Context, this);
+			}
+			ensure(CacheKey == Context.GetTypeHash());
+			return Cache.Data;
+		}
+
+		virtual void Invalidate() override
+		{
+			if (CacheKey != UINT_MAX)
+			{
+				CacheKey = UINT_MAX;
+				Cache = T();
+
+				for (FConnection* Con : GetConnections())
+				{
+					Con->Invalidate();
+				}
+			}
+		}
+
+	};
+
+
+#define EVAL_GRAPH_CONNECTION_TYPE(a,A) \
+template<> FName GraphConnectionTypeName<a>() { return FName(TEXT(#A)); }\
+template class TOutput<a>;\
+template class TInput<a>;
+
+
+}
+
diff --git a/Engine/Source/Runtime/Experimental/Dataflow/Core/Public/Dataflow/DataflowNode.h b/Engine/Source/Runtime/Experimental/Dataflow/Core/Public/Dataflow/DataflowNode.h
new file mode 100644
--- /dev/null
+++ b/Engine/Source/Runtime/Experimental/Dataflow/Core/Public/Dataflow/DataflowNode.h
@@ -0,0 +1,75 @@
+// Copyright Epic Games, Inc. All Rights Reserved.
+
+#pragma once
+
+#include "CoreMinimal.h"
+#include "Dataflow/DataflowConnection.h"
+#include "Dataflow/DataflowNodeParameters.h"
+
+namespace Dataflow
+{
+
+	struct DATAFLOWCORE_API FNodeParameters {
+		FName Name;
+	};
+
+
+	//
+	// FNode
+	//
+	class DATAFLOWCORE_API FNode
+	{
+		friend class FGraph;
+		friend class FConnection;
+
+		FGuid Guid;
+		FName Name;
+
+		TArray< FConnection* > Inputs;
+		TArray< FConnection* > Outputs;
+	public:
+
+		FNode(const FNodeParameters& Param, FGuid InGuid = FGuid::NewGuid())
+			: Guid(InGuid)
+			, Name(Param.Name)
+		{}
+
+		virtual ~FNode() {}
+
+		FName GetName() const { return Name; }
+		void SetName(FName InName) { Name = InName; }
+
+		virtual FName GetType() const { check(true); return FName("invalid"); }
+		FGuid GetGuid() const { return Guid; }
+
+		TArray<FPin> GetPins() const;
+		FConnection* FindInput(FName Name) const;
+		FConnection* FindOutput(FName Name) const;
+
+		virtual void Evaluate(const FContext& Context, FConnection*) { ensure(false); }
+		void InvalidateOutputs();
+
+		virtual void SerializeInternal(FArchive& Ar) {};
+
+		void AddInput(FConnection* InPtr);
+		const TArray< FConnection* >& GetInputs() const { return Inputs; }
+		TArray< FConnection* >& GetInputs() { return Inputs; }
+
+		void AddOutput(FConnection* InPtr);
+		const TArray< FConnection* >& GetOutputs() const { return Outputs; }
+		TArray< FConnection* >& GetOutputs() { return Outputs; }
+
+	};
+}
+
+#define DATAFLOW_REGISTER_CREATION_FACTORY(A) Dataflow::FNodeFactory::GetInstance()->RegisterNode(Dataflow::A::Type, [](const Dataflow::FNewNodeParameters& InParam) {return new Dataflow::A({InParam.Name}, InParam.Guid); });
+
+#define DATAFLOW_DEFINE_EXTERNAL(A) FName A::Type = #A;
+
+#define DATAFLOW_DEFINE_INTERNAL(A) public:\
+ static FName Type;\
+ virtual FName GetType() const { return Type; }\
+ private:
+
+
+
diff --git a/Engine/Source/Runtime/Experimental/Dataflow/Core/Public/Dataflow/DataflowNodeFactory.h b/Engine/Source/Runtime/Experimental/Dataflow/Core/Public/Dataflow/DataflowNodeFactory.h
new file mode 100644
--- /dev/null
+++ b/Engine/Source/Runtime/Experimental/Dataflow/Core/Public/Dataflow/DataflowNodeFactory.h
@@ -0,0 +1,67 @@
+// Copyright Epic Games, Inc. All Rights Reserved.
+
+#pragma once
+
+#include "CoreMinimal.h"
+#include "Dataflow/Dataflow.h"
+
+namespace Dataflow
+{
+	class FNode;
+	class FConnection;
+
+	struct DATAFLOWCORE_API FNewNodeParameters {
+		FGuid Guid;
+		FName Type;
+		FName Name; 
+	};
+
+	//
+	//
+	//
+	class DATAFLOWCORE_API FNodeFactory
+	{
+		typedef TFunction<FNode* (const FNewNodeParameters&)> FNewNodeFunction;
+
+		TMap<FName, FNewNodeFunction > ClassMap;
+
+		static FNodeFactory* Instance;
+		FNodeFactory() {}
+
+	public:
+		~FNodeFactory() { delete Instance; }
+
+		static FNodeFactory* GetInstance()
+		{
+			if (!Instance)
+			{
+				Instance = new FNodeFactory();
+			}
+			return Instance;
+		}
+
+		void RegisterNode(const FName& Type, FNewNodeFunction NewFunction)
+		{
+			if (ensure(!ClassMap.Contains(Type)))
+			{
+				ClassMap.Add(Type, NewFunction);
+			}
+		}
+
+		TSharedPtr<FNode> NewNodeFromRegisteredType(FGraph& Graph, const FNewNodeParameters& Param);
+
+		template<class T> TSharedPtr<T> NewNode(FGraph& Graph, const FNewNodeParameters& Param)
+		{
+			return Graph.AddNode(new T(Param.Name, Param.Guid));
+		}
+
+		TArray<FName> RegisteredNodes() const
+		{
+			TArray<FName> ReturnNames;
+			for (auto Elem : ClassMap) ReturnNames.Add(Elem.Key);
+			return ReturnNames;
+		}
+	};
+
+}
+
diff --git a/Engine/Source/Runtime/Experimental/Dataflow/Core/Public/Dataflow/DataflowNodeParameters.h b/Engine/Source/Runtime/Experimental/Dataflow/Core/Public/Dataflow/DataflowNodeParameters.h
new file mode 100644
--- /dev/null
+++ b/Engine/Source/Runtime/Experimental/Dataflow/Core/Public/Dataflow/DataflowNodeParameters.h
@@ -0,0 +1,34 @@
+// Copyright Epic Games, Inc. All Rights Reserved.
+
+#pragma once
+
+#include "CoreMinimal.h"
+
+class  UDataflow;
+class  UEdGraphPin;
+
+namespace Dataflow
+{
+	struct DATAFLOWCORE_API FContext
+	{
+		float Timestamp = 0.f;
+
+		uint32 GetTypeHash() const
+		{
+			return ::GetTypeHash(Timestamp);
+		}
+	};
+
+
+	template<class T>
+	struct TCacheValue  {
+		TCacheValue(T InData = T())
+			: Data(InData) {}
+		T Data;
+	};
+}
+
+FORCEINLINE uint32 GetTypeHash(const Dataflow::FContext& Context)
+{
+	return ::GetTypeHash(Context.Timestamp);
+}
diff --git a/Engine/Source/Runtime/Experimental/Dataflow/Engine/DataflowEngine.Build.cs b/Engine/Source/Runtime/Experimental/Dataflow/Engine/DataflowEngine.Build.cs
new file mode 100644
--- /dev/null
+++ b/Engine/Source/Runtime/Experimental/Dataflow/Engine/DataflowEngine.Build.cs
@@ -0,0 +1,21 @@
+// Copyright Epic Games, Inc. All Rights Reserved.
+
+namespace UnrealBuildTool.Rules
+{
+	public class DataflowEngine : ModuleRules
+	{
+        public DataflowEngine(ReadOnlyTargetRules Target) : base(Target)
+		{
+			bTreatAsEngineModule = true;
+			PublicDependencyModuleNames.AddRange(
+				new string[]
+				{
+					"Core",
+					"CoreUObject",
+					"Engine",
+					"DataflowCore"
+				}
+				);
+		}
+	}
+}
diff --git a/Engine/Source/Runtime/Experimental/Dataflow/Engine/Private/Dataflow/DataflowEdNode.cpp b/Engine/Source/Runtime/Experimental/Dataflow/Engine/Private/Dataflow/DataflowEdNode.cpp
new file mode 100644
--- /dev/null
+++ b/Engine/Source/Runtime/Experimental/Dataflow/Engine/Private/Dataflow/DataflowEdNode.cpp
@@ -0,0 +1,113 @@
+// Copyright Epic Games, Inc. All Rights Reserved.
+
+#include "Dataflow/DataflowEdNode.h"
+
+#include "Dataflow/DataflowNode.h"
+#include "Dataflow/Dataflow.h"
+#include "Logging/LogMacros.h"
+
+#define LOCTEXT_NAMESPACE "DataflowEdNode"
+
+DEFINE_LOG_CATEGORY_STATIC(DATAFLOWNODE_LOG, Error, All);
+
+
+void UDataflowEdNode::AllocateDefaultPins()
+{
+	UE_LOG(DATAFLOWNODE_LOG, Verbose, TEXT("UDataflowEdNode::AllocateDefaultPins()"));
+	// called on node creation from UI. 
+#if WITH_EDITOR && !(UE_BUILD_SHIPPING || UE_BUILD_TEST)
+	if (DataflowGraph)
+	{
+		if (DataflowNodeGuid.IsValid())
+		{
+			if (TSharedPtr<Dataflow::FNode> DataflowNode = DataflowGraph->FindBaseNode(DataflowNodeGuid))
+			{
+				for (const Dataflow::FPin& Pin : DataflowNode->GetPins())
+				{
+					if (Pin.Direction == Dataflow::FPin::EDirection::INPUT)
+					{
+						CreatePin(EEdGraphPinDirection::EGPD_Input, Pin.Type, Pin.Name);
+					}					
+					if (Pin.Direction == Dataflow::FPin::EDirection::OUTPUT)
+					{
+						CreatePin(EEdGraphPinDirection::EGPD_Output, Pin.Type, Pin.Name);
+					}
+				}
+			}
+		}
+	}
+#endif // WITH_EDITOR && !UE_BUILD_SHIPPING
+}
+
+
+FText UDataflowEdNode::GetNodeTitle(ENodeTitleType::Type TitleType) const
+{
+	return FText::FromString(GetName());
+}
+
+#if WITH_EDITOR && !(UE_BUILD_SHIPPING || UE_BUILD_TEST)
+void UDataflowEdNode::PinConnectionListChanged(UEdGraphPin* Pin)
+{
+	if (ensure(IsBound()))
+	{
+		if (TSharedPtr<Dataflow::FNode> DataflowNode = DataflowGraph->FindBaseNode(DataflowNodeGuid))
+		{
+			if (Dataflow::FConnection* ConnectionInput = DataflowNode->FindInput(FName(Pin->GetName())) )
+			{
+				DataflowGraph->ClearConnections(ConnectionInput);
+				for (UEdGraphPin* LinkedCon : Pin->LinkedTo)
+				{
+					if (UDataflowEdNode* LinkedNode = Cast<UDataflowEdNode>(LinkedCon->GetOwningNode()))
+					{
+						if (ensure(LinkedNode->IsBound()))
+						{
+							if (TSharedPtr<Dataflow::FNode> LinkedDataflowNode = DataflowGraph->FindBaseNode(LinkedNode->GetDataflowNodeGuid()))
+							{
+								if (Dataflow::FConnection* LinkedConBase = LinkedDataflowNode->FindOutput(FName(LinkedCon->GetName())))
+								{
+									DataflowGraph->Connect(ConnectionInput, LinkedConBase);
+								}
+							}
+						}
+					}
+				}
+			}
+			else if (Dataflow::FConnection* Con = DataflowNode->FindOutput(FName(Pin->GetName())))
+			{
+				if (Dataflow::FConnection* ConnectionOutput = DataflowNode->FindOutput(FName(Pin->GetName())))
+				{
+					DataflowGraph->ClearConnections(ConnectionOutput);
+					for (UEdGraphPin* LinkedCon : Pin->LinkedTo)
+					{
+						if (UDataflowEdNode* LinkedNode = Cast<UDataflowEdNode>(LinkedCon->GetOwningNode()))
+						{
+							if (ensure(LinkedNode->IsBound()))
+							{
+								if (TSharedPtr<Dataflow::FNode> LinkedDataflowNode = DataflowGraph->FindBaseNode(LinkedNode->GetDataflowNodeGuid()))
+								{
+									if (Dataflow::FConnection* LinkedConBase = LinkedDataflowNode->FindInput(FName(LinkedCon->GetName())))
+									{
+										DataflowGraph->Connect(LinkedConBase, ConnectionOutput);
+									}
+								}
+							}
+						}
+					}
+				}
+			}
+		}
+	}
+
+	Super::PinConnectionListChanged(Pin);
+}
+#endif // WITH_EDITOR && !UE_BUILD_SHIPPING
+
+void UDataflowEdNode::Serialize(FArchive& Ar)
+{
+	Super::Serialize(Ar);
+	Ar << DataflowNodeGuid;
+}
+
+
+#undef LOCTEXT_NAMESPACE
+
diff --git a/Engine/Source/Runtime/Experimental/Dataflow/Engine/Private/Dataflow/DataflowEngineLogPrivate.h b/Engine/Source/Runtime/Experimental/Dataflow/Engine/Private/Dataflow/DataflowEngineLogPrivate.h
new file mode 100644
--- /dev/null
+++ b/Engine/Source/Runtime/Experimental/Dataflow/Engine/Private/Dataflow/DataflowEngineLogPrivate.h
@@ -0,0 +1,11 @@
+// Copyright Epic Games, Inc. All Rights Reserved.
+
+#pragma once
+
+#include "CoreMinimal.h"
+#include "Stats/Stats.h"
+
+// You should place include statements to your module's private header files here.  You only need to
+// add includes for headers that are used in most of your module's source files though.
+
+
diff --git a/Engine/Source/Runtime/Experimental/Dataflow/Engine/Private/Dataflow/DataflowEngineModule.cpp b/Engine/Source/Runtime/Experimental/Dataflow/Engine/Private/Dataflow/DataflowEngineModule.cpp
new file mode 100644
--- /dev/null
+++ b/Engine/Source/Runtime/Experimental/Dataflow/Engine/Private/Dataflow/DataflowEngineModule.cpp
@@ -0,0 +1,30 @@
+// Copyright Epic Games, Inc. All Rights Reserved.
+
+#include "CoreMinimal.h"
+#include "DataflowEngineLogPrivate.h"
+#include "Modules/ModuleInterface.h"
+#include "Modules/ModuleManager.h"
+
+
+/**
+ * The public interface to this module
+ */
+class FDataflowEngineModule : public IModuleInterface
+{
+public:
+
+	// IModuleInterface interface
+
+	virtual void StartupModule() override { }
+	virtual void ShutdownModule() override { }
+
+	virtual bool SupportsDynamicReloading() override
+	{
+		return false;
+	}
+};
+
+IMPLEMENT_MODULE( FDataflowEngineModule, DataflowEngine )
+
+
+
diff --git a/Engine/Source/Runtime/Experimental/Dataflow/Engine/Private/Dataflow/DataflowObject.cpp b/Engine/Source/Runtime/Experimental/Dataflow/Engine/Private/Dataflow/DataflowObject.cpp
new file mode 100644
--- /dev/null
+++ b/Engine/Source/Runtime/Experimental/Dataflow/Engine/Private/Dataflow/DataflowObject.cpp
@@ -0,0 +1,34 @@
+// Copyright Epic Games, Inc. All Rights Reserved.
+
+#include "Dataflow/DataflowObject.h"
+#include "Dataflow/Dataflow.h"
+
+
+#define LOCTEXT_NAMESPACE "UDataflow")
+
+UDataflow::UDataflow(const FObjectInitializer& ObjectInitializer)
+	: UEdGraph(ObjectInitializer)
+	, Dataflow(new Dataflow::FGraph())
+{}
+
+#if WITH_EDITOR
+
+void UDataflow::PostEditChangeProperty(struct FPropertyChangedEvent& PropertyChangedEvent)
+{
+	UObject::PostEditChangeProperty(PropertyChangedEvent);
+}
+
+#endif
+
+void UDataflow::PostLoad()
+{
+	UObject::PostLoad();
+}
+
+void UDataflow::Serialize(FArchive& Ar)
+{
+	Super::Serialize(Ar);
+	Ar << *Dataflow.Get();
+}
+
+#undef LOCTEXT_NAMESPACE
diff --git a/Engine/Source/Runtime/Experimental/Dataflow/Engine/Public/Dataflow/DataflowEdNode.h b/Engine/Source/Runtime/Experimental/Dataflow/Engine/Public/Dataflow/DataflowEdNode.h
new file mode 100644
--- /dev/null
+++ b/Engine/Source/Runtime/Experimental/Dataflow/Engine/Public/Dataflow/DataflowEdNode.h
@@ -0,0 +1,43 @@
+// Copyright Epic Games, Inc. All Rights Reserved.
+
+#pragma once
+
+#include "CoreMinimal.h"
+#include "EdGraph/EdGraph.h"
+#include "Dataflow/Dataflow.h"
+
+#include "DataflowEdNode.generated.h"
+
+
+
+UCLASS()
+class DATAFLOWENGINE_API UDataflowEdNode : public UEdGraphNode
+{
+	GENERATED_BODY()
+
+	FGuid DataflowNodeGuid;
+	TSharedPtr<Dataflow::FGraph> DataflowGraph;
+public:
+
+	// UEdGraphNode interface
+	virtual void AllocateDefaultPins();
+	virtual FText GetNodeTitle(ENodeTitleType::Type TitleType) const;
+#if WITH_EDITOR && !(UE_BUILD_SHIPPING || UE_BUILD_TEST)
+	virtual void PinConnectionListChanged(UEdGraphPin* Pin) override;
+#endif // WITH_EDITOR && !UE_BUILD_SHIPPING
+	// End of UEdGraphNode interface
+
+	bool IsBound() { return DataflowGraph && DataflowNodeGuid.IsValid(); }
+
+	TSharedPtr<Dataflow::FGraph> GetDataflowGraph() { return DataflowGraph;}
+	void SetDataflowGraph(TSharedPtr<Dataflow::FGraph> InDataflowGraph) { DataflowGraph = InDataflowGraph; }
+
+	FGuid GetDataflowNodeGuid() const { return DataflowNodeGuid; }
+	void SetDataflowNodeGuid(FGuid InGuid) { DataflowNodeGuid = InGuid; }
+
+	// UObject interface
+	void Serialize(FArchive& Ar);
+	// End UObject interface
+
+};
+
diff --git a/Engine/Source/Runtime/Experimental/Dataflow/Engine/Public/Dataflow/DataflowObject.h b/Engine/Source/Runtime/Experimental/Dataflow/Engine/Public/Dataflow/DataflowObject.h
new file mode 100644
--- /dev/null
+++ b/Engine/Source/Runtime/Experimental/Dataflow/Engine/Public/Dataflow/DataflowObject.h
@@ -0,0 +1,43 @@
+// Copyright Epic Games, Inc. All Rights Reserved.
+
+#pragma once
+
+#include "CoreMinimal.h"
+#include "EdGraph/EdGraph.h"
+#include "Dataflow/Dataflow.h"
+
+#include "DataflowObject.generated.h"
+
+/**
+* UDataflow (UObject)
+*
+* UObject wrapper for the Dataflow::FGraph
+*
+*/
+UCLASS(BlueprintType, customconstructor)
+class DATAFLOWENGINE_API UDataflow : public UEdGraph
+{
+	GENERATED_UCLASS_BODY()
+
+	TSharedPtr<Dataflow::FGraph, ESPMode::ThreadSafe> Dataflow;
+
+public:
+	UDataflow(const FObjectInitializer& ObjectInitializer = FObjectInitializer::Get());
+
+	/** UObject Interface */
+#if WITH_EDITOR
+	virtual void PostEditChangeProperty(struct FPropertyChangedEvent& PropertyChangedEvent) override;
+#endif
+	virtual void PostLoad() override;
+	/** End UObject Interface */
+
+	void Serialize(FArchive& Ar);
+
+	/** Accessors for internal geometry collection */
+	void SetDataflow(TSharedPtr<Dataflow::FGraph, ESPMode::ThreadSafe> DataflowIn) { Dataflow = DataflowIn; }
+	TSharedPtr<Dataflow::FGraph, ESPMode::ThreadSafe>       GetDataflow() { return Dataflow; }
+	const TSharedPtr<Dataflow::FGraph, ESPMode::ThreadSafe> GetDataflow() const { return Dataflow; }
+
+
+};
+
diff --git a/Engine/Source/Runtime/Experimental/EvalGraph/Core/EvalGraphCore.Build.cs b/Engine/Source/Runtime/Experimental/EvalGraph/Core/EvalGraphCore.Build.cs
deleted file mode 100644
diff --git a/Engine/Source/Runtime/Experimental/EvalGraph/Core/Private/EvalGraph/EvalGraph.cpp b/Engine/Source/Runtime/Experimental/EvalGraph/Core/Private/EvalGraph/EvalGraph.cpp
deleted file mode 100644
diff --git a/Engine/Source/Runtime/Experimental/EvalGraph/Core/Private/EvalGraph/EvalGraphConnectionBase.cpp b/Engine/Source/Runtime/Experimental/EvalGraph/Core/Private/EvalGraph/EvalGraphConnectionBase.cpp
deleted file mode 100644
diff --git a/Engine/Source/Runtime/Experimental/EvalGraph/Core/Private/EvalGraph/EvalGraphCoreLogPrivate.h b/Engine/Source/Runtime/Experimental/EvalGraph/Core/Private/EvalGraph/EvalGraphCoreLogPrivate.h
deleted file mode 100644
diff --git a/Engine/Source/Runtime/Experimental/EvalGraph/Core/Private/EvalGraph/EvalGraphCoreModule.cpp b/Engine/Source/Runtime/Experimental/EvalGraph/Core/Private/EvalGraph/EvalGraphCoreModule.cpp
deleted file mode 100644
diff --git a/Engine/Source/Runtime/Experimental/EvalGraph/Core/Private/EvalGraph/EvalGraphInputOutput.cpp b/Engine/Source/Runtime/Experimental/EvalGraph/Core/Private/EvalGraph/EvalGraphInputOutput.cpp
deleted file mode 100644
diff --git a/Engine/Source/Runtime/Experimental/EvalGraph/Core/Private/EvalGraph/EvalGraphNode.cpp b/Engine/Source/Runtime/Experimental/EvalGraph/Core/Private/EvalGraph/EvalGraphNode.cpp
deleted file mode 100644
diff --git a/Engine/Source/Runtime/Experimental/EvalGraph/Core/Private/EvalGraph/EvalGraphNodeFactory.cpp b/Engine/Source/Runtime/Experimental/EvalGraph/Core/Private/EvalGraph/EvalGraphNodeFactory.cpp
deleted file mode 100644
diff --git a/Engine/Source/Runtime/Experimental/EvalGraph/Core/Private/EvalGraph/EvalGraphNodeParameters.cpp b/Engine/Source/Runtime/Experimental/EvalGraph/Core/Private/EvalGraph/EvalGraphNodeParameters.cpp
deleted file mode 100644
diff --git a/Engine/Source/Runtime/Experimental/EvalGraph/Core/Public/EvalGraph/EvalGraph.h b/Engine/Source/Runtime/Experimental/EvalGraph/Core/Public/EvalGraph/EvalGraph.h
deleted file mode 100644
diff --git a/Engine/Source/Runtime/Experimental/EvalGraph/Core/Public/EvalGraph/EvalGraphConnectionBase.h b/Engine/Source/Runtime/Experimental/EvalGraph/Core/Public/EvalGraph/EvalGraphConnectionBase.h
deleted file mode 100644
diff --git a/Engine/Source/Runtime/Experimental/EvalGraph/Core/Public/EvalGraph/EvalGraphCoreModule.h b/Engine/Source/Runtime/Experimental/EvalGraph/Core/Public/EvalGraph/EvalGraphCoreModule.h
deleted file mode 100644
diff --git a/Engine/Source/Runtime/Experimental/EvalGraph/Core/Public/EvalGraph/EvalGraphInputOutput.h b/Engine/Source/Runtime/Experimental/EvalGraph/Core/Public/EvalGraph/EvalGraphInputOutput.h
deleted file mode 100644
diff --git a/Engine/Source/Runtime/Experimental/EvalGraph/Core/Public/EvalGraph/EvalGraphNode.h b/Engine/Source/Runtime/Experimental/EvalGraph/Core/Public/EvalGraph/EvalGraphNode.h
deleted file mode 100644
diff --git a/Engine/Source/Runtime/Experimental/EvalGraph/Core/Public/EvalGraph/EvalGraphNodeFactory.h b/Engine/Source/Runtime/Experimental/EvalGraph/Core/Public/EvalGraph/EvalGraphNodeFactory.h
deleted file mode 100644
diff --git a/Engine/Source/Runtime/Experimental/EvalGraph/Core/Public/EvalGraph/EvalGraphNodeParameters.h b/Engine/Source/Runtime/Experimental/EvalGraph/Core/Public/EvalGraph/EvalGraphNodeParameters.h
deleted file mode 100644
diff --git a/Engine/Source/Runtime/Experimental/EvalGraph/Engine/EvalGraphEngine.Build.cs b/Engine/Source/Runtime/Experimental/EvalGraph/Engine/EvalGraphEngine.Build.cs
deleted file mode 100644
diff --git a/Engine/Source/Runtime/Experimental/EvalGraph/Engine/Private/EvalGraph/EvalGraphEdNode.cpp b/Engine/Source/Runtime/Experimental/EvalGraph/Engine/Private/EvalGraph/EvalGraphEdNode.cpp
deleted file mode 100644
diff --git a/Engine/Source/Runtime/Experimental/EvalGraph/Engine/Private/EvalGraph/EvalGraphEngineLogPrivate.h b/Engine/Source/Runtime/Experimental/EvalGraph/Engine/Private/EvalGraph/EvalGraphEngineLogPrivate.h
deleted file mode 100644
diff --git a/Engine/Source/Runtime/Experimental/EvalGraph/Engine/Private/EvalGraph/EvalGraphEngineModule.cpp b/Engine/Source/Runtime/Experimental/EvalGraph/Engine/Private/EvalGraph/EvalGraphEngineModule.cpp
deleted file mode 100644
diff --git a/Engine/Source/Runtime/Experimental/EvalGraph/Engine/Private/EvalGraph/EvalGraphObject.cpp b/Engine/Source/Runtime/Experimental/EvalGraph/Engine/Private/EvalGraph/EvalGraphObject.cpp
deleted file mode 100644
diff --git a/Engine/Source/Runtime/Experimental/EvalGraph/Engine/Public/EvalGraph/EvalGraphEdNode.h b/Engine/Source/Runtime/Experimental/EvalGraph/Engine/Public/EvalGraph/EvalGraphEdNode.h
deleted file mode 100644
diff --git a/Engine/Source/Runtime/Experimental/EvalGraph/Engine/Public/EvalGraph/EvalGraphObject.h b/Engine/Source/Runtime/Experimental/EvalGraph/Engine/Public/EvalGraph/EvalGraphObject.h
deleted file mode 100644
