]Dataflow - ManagedArrayCollection Node Example

diff --git a/Engine/Plugins/Experimental/Dataflow/Dataflow.uplugin b/Engine/Plugins/Experimental/Dataflow/Dataflow.uplugin
--- a/Engine/Plugins/Experimental/Dataflow/Dataflow.uplugin
+++ b/Engine/Plugins/Experimental/Dataflow/Dataflow.uplugin
@@ -21,6 +21,11 @@
 			"Type": "Editor",
 			"LoadingPhase": "Default"
 		},
+		{
+			"Name": "DataflowNodes",
+			"Type": "Runtime",
+			"LoadingPhase": "Default"
+		},
 	],
 	"Plugins": [
 	]
diff --git a/Engine/Plugins/Experimental/Dataflow/Source/DataflowNodes/DataflowNodes.Build.cs b/Engine/Plugins/Experimental/Dataflow/Source/DataflowNodes/DataflowNodes.Build.cs
new file mode 100644
--- /dev/null
+++ b/Engine/Plugins/Experimental/Dataflow/Source/DataflowNodes/DataflowNodes.Build.cs
@@ -0,0 +1,23 @@
+// Copyright Epic Games, Inc. All Rights Reserved.
+
+namespace UnrealBuildTool.Rules
+{
+	public class DataflowNodes : ModuleRules
+	{
+        public DataflowNodes(ReadOnlyTargetRules Target) : base(Target)
+		{
+			PrivateIncludePaths.Add("DataflowNodes/Private");
+            PublicIncludePaths.Add(ModuleDirectory + "/Public");
+
+			PublicDependencyModuleNames.AddRange(
+				new string[]
+				{
+					"Core",
+					"CoreUObject",
+					"Chaos",
+					"DataflowCore"
+				}
+			);
+		}
+	}
+}
diff --git a/Engine/Plugins/Experimental/Dataflow/Source/DataflowNodes/Private/Dataflow/DataflowConnectionTypes.cpp b/Engine/Plugins/Experimental/Dataflow/Source/DataflowNodes/Private/Dataflow/DataflowConnectionTypes.cpp
new file mode 100644
--- /dev/null
+++ b/Engine/Plugins/Experimental/Dataflow/Source/DataflowNodes/Private/Dataflow/DataflowConnectionTypes.cpp
@@ -0,0 +1,10 @@
+// Copyright Epic Games, Inc. All Rights Reserved.
+
+#include "Dataflow/DataflowConnectionTypes.h"
+#include "Dataflow/DataflowInputOutput.h"
+
+
+namespace Dataflow
+{
+}
+
diff --git a/Engine/Plugins/Experimental/Dataflow/Source/DataflowNodes/Private/Dataflow/DataflowNodesPlugin.cpp b/Engine/Plugins/Experimental/Dataflow/Source/DataflowNodes/Private/Dataflow/DataflowNodesPlugin.cpp
new file mode 100644
--- /dev/null
+++ b/Engine/Plugins/Experimental/Dataflow/Source/DataflowNodes/Private/Dataflow/DataflowNodesPlugin.cpp
@@ -0,0 +1,31 @@
+// Copyright Epic Games, Inc. All Rights Reserved.
+
+#include "Dataflow/DataflowNodesPlugin.h"
+
+#include "CoreMinimal.h"
+
+#include "Dataflow/DataflowNodeFactory.h"
+#include "Dataflow/ManagedArrayCollectionNodes.h"
+
+
+#define LOCTEXT_NAMESPACE "DataflowNodes"
+
+
+void IDataflowNodesPlugin::StartupModule()
+{
+	Dataflow::RegisterManagedArrayCollectionNodes();
+}
+
+void IDataflowNodesPlugin::ShutdownModule()
+{
+	if (UObjectInitialized())
+	{
+		//Dataflow::UnregisterManagedArrayCollectionNodes(); ?
+	}
+}
+
+
+IMPLEMENT_MODULE(IDataflowNodesPlugin, DataflowNodes)
+
+
+#undef LOCTEXT_NAMESPACE
diff --git a/Engine/Plugins/Experimental/Dataflow/Source/DataflowNodes/Private/Dataflow/ManagedArrayCollectionNodes.cpp b/Engine/Plugins/Experimental/Dataflow/Source/DataflowNodes/Private/Dataflow/ManagedArrayCollectionNodes.cpp
new file mode 100644
--- /dev/null
+++ b/Engine/Plugins/Experimental/Dataflow/Source/DataflowNodes/Private/Dataflow/ManagedArrayCollectionNodes.cpp
@@ -0,0 +1,24 @@
+// Copyright Epic Games, Inc. All Rights Reserved.
+
+#include "Dataflow/ManagedArrayCollectionNodes.h"
+
+#include "Dataflow/DataflowConnectionTypes.h"
+#include "Dataflow/DataflowInputOutput.h"
+#include "Dataflow/DataflowNodeParameters.h"
+#include "Dataflow/DataflowNodeFactory.h"
+
+#include "GeometryCollection/ManagedArray.h"
+#include "GeometryCollection/ManagedArrayCollection.h"
+
+
+namespace Dataflow
+{
+	
+void RegisterManagedArrayCollectionNodes()
+{
+	DATAFLOW_NODE_REGISTER_CREATION_FACTORY(FNewManagedArrayCollectionNode);
+	DATAFLOW_NODE_REGISTER_CREATION_FACTORY(FAddAttributeNode);
+}
+
+}
+
diff --git a/Engine/Plugins/Experimental/Dataflow/Source/DataflowNodes/Private/Dataflow/ManagedArrayNodes.cpp b/Engine/Plugins/Experimental/Dataflow/Source/DataflowNodes/Private/Dataflow/ManagedArrayNodes.cpp
new file mode 100644
--- /dev/null
+++ b/Engine/Plugins/Experimental/Dataflow/Source/DataflowNodes/Private/Dataflow/ManagedArrayNodes.cpp
@@ -0,0 +1,10 @@
+// Copyright Epic Games, Inc. All Rights Reserved.
+
+#include "Dataflow/ManagedArrayNodes.h"
+
+
+namespace Dataflow
+{
+
+}
+
diff --git a/Engine/Plugins/Experimental/Dataflow/Source/DataflowNodes/Public/Dataflow/DataflowConnectionTypes.h b/Engine/Plugins/Experimental/Dataflow/Source/DataflowNodes/Public/Dataflow/DataflowConnectionTypes.h
new file mode 100644
--- /dev/null
+++ b/Engine/Plugins/Experimental/Dataflow/Source/DataflowNodes/Public/Dataflow/DataflowConnectionTypes.h
@@ -0,0 +1,17 @@
+// Copyright Epic Games, Inc. All Rights Reserved.
+
+#pragma once
+
+#include "CoreMinimal.h"
+
+#include "Dataflow/DataflowInputOutput.h"
+
+#include "GeometryCollection/ManagedArray.h"
+#include "GeometryCollection/ManagedArrayCollection.h"
+
+namespace Dataflow
+{
+	typedef TSharedPtr<FManagedArrayCollection> FManagedArrayCollectionSharedPtr;
+
+	DATAFLOW_CONNECTION_TYPE(DATAFLOWNODES_API, FManagedArrayCollectionSharedPtr, ManagedArrayCollectionSharedPtr)
+}
diff --git a/Engine/Plugins/Experimental/Dataflow/Source/DataflowNodes/Public/Dataflow/DataflowNodesPlugin.h b/Engine/Plugins/Experimental/Dataflow/Source/DataflowNodes/Public/Dataflow/DataflowNodesPlugin.h
new file mode 100644
--- /dev/null
+++ b/Engine/Plugins/Experimental/Dataflow/Source/DataflowNodes/Public/Dataflow/DataflowNodesPlugin.h
@@ -0,0 +1,18 @@
+// Copyright Epic Games, Inc. All Rights Reserved.
+
+#pragma once
+
+#include "CoreMinimal.h"
+#include "Modules/ModuleInterface.h"
+#include "Modules/ModuleManager.h"
+
+/**
+ * The public interface to this module
+ */
+class IDataflowNodesPlugin : public IModuleInterface
+{
+public:
+	virtual void StartupModule();
+	virtual void ShutdownModule();
+};
+
diff --git a/Engine/Plugins/Experimental/Dataflow/Source/DataflowNodes/Public/Dataflow/ManagedArrayCollectionNodes.h b/Engine/Plugins/Experimental/Dataflow/Source/DataflowNodes/Public/Dataflow/ManagedArrayCollectionNodes.h
new file mode 100644
--- /dev/null
+++ b/Engine/Plugins/Experimental/Dataflow/Source/DataflowNodes/Public/Dataflow/ManagedArrayCollectionNodes.h
@@ -0,0 +1,127 @@
+// Copyright Epic Games, Inc. All Rights Reserved.
+
+#pragma once
+
+#include "CoreMinimal.h"
+
+#include "Dataflow/DataflowInputOutput.h"
+#include "Dataflow/DataflowNode.h"
+
+#include "GeometryCollection/ManagedArray.h"
+#include "GeometryCollection/ManagedArrayCollection.h"
+
+namespace Dataflow
+{
+	typedef TSharedPtr<FManagedArrayCollection> FManagedArrayCollectionSharedPtr;
+
+	enum class EManagedArrayType : uint8
+	{
+		EgManagedArrayTypeNone = 0,
+		EgManagedArrayTypeBool,
+		EgManagedArrayTypeInt,
+		EgManagedArrayTypeFloat,
+		EgManagedArrayTypeVector
+	};
+
+	class DATAFLOWNODES_API FNewManagedArrayCollectionNode : public FNode
+	{
+		DATAFLOW_NODE_DEFINE_INTERNAL(FNewManagedArrayCollectionNode)
+
+	public:
+		typedef TSharedPtr<FManagedArrayCollection> DataType;
+		DataType Value;
+		TSharedPtr< class TOutput<DataType> > Output;
+
+		FNewManagedArrayCollectionNode(const FNodeParameters& InParam, FGuid InGuid = FGuid::NewGuid())
+			: FNode(InParam, InGuid)
+			, Value(new FManagedArrayCollection())
+			, Output(new TOutput<DataType>(TOutputParameters<DataType>({FName("CollectionOut"), this})))
+		{}
+
+		virtual void Evaluate(const FContext& Context, FConnection* Out) override
+		{
+			if (Output.Get() == Out)
+			{
+				Output->SetValue(Value, Context);
+			}
+		}
+	};
+
+	class DATAFLOWNODES_API FAddAttributeNode : public FNode
+	{
+		DATAFLOW_NODE_DEFINE_INTERNAL(FAddAttributeNode)
+
+	public:
+
+		typedef TSharedPtr<FManagedArrayCollection> DataType;
+
+		FName AttributeName;
+		FName GroupName;
+		EManagedArrayType AttributeType;
+
+		TSharedPtr< class TInput<DataType> > Input;
+		TSharedPtr< class TInput<int32> > SizeInput;
+		TSharedPtr< class TOutput<DataType> > Output;
+
+		FAddAttributeNode(const FNodeParameters& InParam, FGuid InGuid = FGuid::NewGuid())
+			: FNode(InParam, InGuid)
+			, AttributeName("Position")
+			, GroupName("Particle")
+			, AttributeType(EManagedArrayType::EgManagedArrayTypeVector)
+			, Input(new TInput<DataType>(TInputParameters<DataType>({FName("CollectionIn"), this})))
+			, SizeInput(new TInput<int32>(TInputParameters<int32>({FName("SizeIn"), this })))
+			, Output(new TOutput<DataType>(TOutputParameters<DataType>({FName("CollectionOut"), this })))
+		{}
+
+		virtual void Evaluate(const FContext& Context, FConnection* Out) override
+		{
+			if (Output.Get() == Out)
+			{
+				DataType Collection(new FManagedArrayCollection());
+
+				// its an error to not set the requested output
+				if (DataType Val = Input->GetValue(Context))
+				{
+					Collection = Val;
+
+					int NumElements = 100;
+					if (SizeInput->GetConnection())
+					{
+						NumElements = SizeInput->GetValue(Context);
+					}
+
+					if (!Collection->HasGroup(GroupName))
+					{
+						Collection->AddGroup(GroupName);
+					}
+					if (!Collection->HasAttribute(AttributeName, GroupName))
+					{
+						Collection->AddAttribute<FVector>(AttributeName, GroupName);
+					}
+
+					if (Collection->NumElements(GroupName) < NumElements)
+					{
+						Collection->AddElements(NumElements, GroupName);
+
+						TManagedArray<FVector>& Pos = Collection->ModifyAttribute<FVector>(AttributeName, GroupName);
+						for (int i = 0; i < NumElements; i++)
+						{
+							Pos[i] = FVector(i);
+						}
+					}
+				}
+				Output->SetValue(Collection, Context);
+			}
+		}
+
+		virtual void SerializeInternal(FArchive& Ar) override { 
+			Ar << AttributeName << GroupName << AttributeType;
+		};
+
+	};
+
+
+
+	void DATAFLOWNODES_API RegisterManagedArrayCollectionNodes();
+}
+
diff --git a/Engine/Plugins/Experimental/Dataflow/Source/DataflowNodes/Public/Dataflow/ManagedArrayNodes.h b/Engine/Plugins/Experimental/Dataflow/Source/DataflowNodes/Public/Dataflow/ManagedArrayNodes.h
new file mode 100644
--- /dev/null
+++ b/Engine/Plugins/Experimental/Dataflow/Source/DataflowNodes/Public/Dataflow/ManagedArrayNodes.h
@@ -0,0 +1,11 @@
+// Copyright Epic Games, Inc. All Rights Reserved.
+
+#pragma once
+
+#include "CoreMinimal.h"
+
+namespace Dataflow
+{
+
+}
+
diff --git a/Engine/Source/Runtime/Experimental/Dataflow/Core/Private/Dataflow/DataflowInputOutput.cpp b/Engine/Source/Runtime/Experimental/Dataflow/Core/Private/Dataflow/DataflowInputOutput.cpp
--- a/Engine/Source/Runtime/Experimental/Dataflow/Core/Private/Dataflow/DataflowInputOutput.cpp
+++ b/Engine/Source/Runtime/Experimental/Dataflow/Core/Private/Dataflow/DataflowInputOutput.cpp
@@ -8,11 +8,6 @@
 namespace Dataflow
 {
 	// Base Types
-	EVAL_GRAPH_CONNECTION_TYPE(bool, Bool)
-	EVAL_GRAPH_CONNECTION_TYPE(char, Char)
-	EVAL_GRAPH_CONNECTION_TYPE(int, Integer)	
-	EVAL_GRAPH_CONNECTION_TYPE(uint8, UInt8)
-	EVAL_GRAPH_CONNECTION_TYPE(float, Float)
-	EVAL_GRAPH_CONNECTION_TYPE(double, Double)
+
 }
 
diff --git a/Engine/Source/Runtime/Experimental/Dataflow/Core/Public/Dataflow/DataflowInputOutput.h b/Engine/Source/Runtime/Experimental/Dataflow/Core/Public/Dataflow/DataflowInputOutput.h
--- a/Engine/Source/Runtime/Experimental/Dataflow/Core/Public/Dataflow/DataflowInputOutput.h
+++ b/Engine/Source/Runtime/Experimental/Dataflow/Core/Public/Dataflow/DataflowInputOutput.h
@@ -197,11 +197,25 @@ namespace Dataflow
 	};
 
 
-#define EVAL_GRAPH_CONNECTION_TYPE(a,A) \
-template<> FName GraphConnectionTypeName<a>() { return FName(TEXT(#A)); }\
-template class TOutput<a>;\
-template class TInput<a>;
-
-
+//
+// Used this macros to defined dataflow connection types 
+//
+
+#define DATAFLOW_CONNECTION_TYPE(D, a,A)			\
+	template<> inline FName D GraphConnectionTypeName<a>()	\
+			{ return FName(TEXT(#A)); }				\
+	template struct D TOutputParameters<a>;			\
+	template class D TOutput<a>;					\
+	template struct D TInputParameters<a>;			\
+	template class D TInput<a>;
+
+
+
+DATAFLOW_CONNECTION_TYPE(DATAFLOWCORE_API, bool, Bool)
+DATAFLOW_CONNECTION_TYPE(DATAFLOWCORE_API, char, Char)
+DATAFLOW_CONNECTION_TYPE(DATAFLOWCORE_API, int, Integer)
+DATAFLOW_CONNECTION_TYPE(DATAFLOWCORE_API, uint8, UInt8)
+DATAFLOW_CONNECTION_TYPE(DATAFLOWCORE_API, float, Float)
+DATAFLOW_CONNECTION_TYPE(DATAFLOWCORE_API, double, Double)
 }
 
diff --git a/Engine/Source/Runtime/Experimental/Dataflow/Core/Public/Dataflow/DataflowNode.h b/Engine/Source/Runtime/Experimental/Dataflow/Core/Public/Dataflow/DataflowNode.h
--- a/Engine/Source/Runtime/Experimental/Dataflow/Core/Public/Dataflow/DataflowNode.h
+++ b/Engine/Source/Runtime/Experimental/Dataflow/Core/Public/Dataflow/DataflowNode.h
@@ -60,16 +60,23 @@ namespace Dataflow
 		TArray< FConnection* >& GetOutputs() { return Outputs; }
 
 	};
-}
 
-#define DATAFLOW_REGISTER_CREATION_FACTORY(A) Dataflow::FNodeFactory::GetInstance()->RegisterNode(Dataflow::A::Type, [](const Dataflow::FNewNodeParameters& InParam) {return new Dataflow::A({InParam.Name}, InParam.Guid); });
+//
+// Used these macros to register dataflow nodes. 
+//
 
-#define DATAFLOW_DEFINE_EXTERNAL(A) FName A::Type = #A;
+#define DATAFLOW_NODE_REGISTER_CREATION_FACTORY(A)					\
+	FNodeFactory::GetInstance()->RegisterNode(					\
+		A::StaticType(),										\
+		[](const FNewNodeParameters& InParam)					\
+			{return new A({InParam.Name}, InParam.Guid); });
 
-#define DATAFLOW_DEFINE_INTERNAL(A) public:\
- static FName Type;\
- virtual FName GetType() const { return Type; }\
- private:
+#define DATAFLOW_NODE_DEFINE_INTERNAL(A)								\
+public:															\
+   static FName StaticType() {return #A;}						\
+   virtual FName GetType() const { return #A; }					\
+private:
 
 
+}
 
