Dataflow : Updates
- UniqueNames
- PostLoad bindings.

diff --git a/Engine/Plugins/Experimental/Dataflow/Source/DataflowEditor/Private/Dataflow/DataflowSNode.cpp b/Engine/Plugins/Experimental/Dataflow/Source/DataflowEditor/Private/Dataflow/DataflowSNode.cpp
--- a/Engine/Plugins/Experimental/Dataflow/Source/DataflowEditor/Private/Dataflow/DataflowSNode.cpp
+++ b/Engine/Plugins/Experimental/Dataflow/Source/DataflowEditor/Private/Dataflow/DataflowSNode.cpp
@@ -47,7 +47,7 @@ UEdGraphNode* FAssetSchemaAction_Dataflow_CreateNode_DataflowEdNode::PerformActi
 {
 	if (UDataflow* Dataflow = Cast<UDataflow>(ParentGraph))
 	{
-		const FName NodeName(GetMenuDescription().ToString());
+		const FName NodeName = MakeUniqueObjectName(Dataflow, UDataflowEdNode::StaticClass(), FName(GetMenuDescription().ToString()));
 		if (UDataflowEdNode* EdNode = NewObject<UDataflowEdNode>(Dataflow, UDataflowEdNode::StaticClass(), NodeName))
 		{
 			//const FScopedTransaction Transaction(LOCTEXT("DataflowNewNode", "Generic Graph Editor: New Node"));
diff --git a/Engine/Source/Runtime/Experimental/Dataflow/Core/Public/Dataflow/DataflowGraph.h b/Engine/Source/Runtime/Experimental/Dataflow/Core/Public/Dataflow/DataflowGraph.h
--- a/Engine/Source/Runtime/Experimental/Dataflow/Core/Public/Dataflow/DataflowGraph.h
+++ b/Engine/Source/Runtime/Experimental/Dataflow/Core/Public/Dataflow/DataflowGraph.h
@@ -50,6 +50,8 @@ namespace Dataflow
 		FGraph(FGuid InGuid = FGuid::NewGuid());
 		virtual ~FGraph() {}
 
+		const TArray< TSharedPtr<FNode> >& GetNodes() const {return Nodes;}
+		TArray< TSharedPtr<FNode> >& GetNodes() { return Nodes; }
 		int NumNodes() { return Nodes.Num(); }
 
 		template<class T> TSharedPtr<T> AddNode(T* InNode)
diff --git a/Engine/Source/Runtime/Experimental/Dataflow/Engine/Private/Dataflow/DataflowObject.cpp b/Engine/Source/Runtime/Experimental/Dataflow/Engine/Private/Dataflow/DataflowObject.cpp
--- a/Engine/Source/Runtime/Experimental/Dataflow/Engine/Private/Dataflow/DataflowObject.cpp
+++ b/Engine/Source/Runtime/Experimental/Dataflow/Engine/Private/Dataflow/DataflowObject.cpp
@@ -2,6 +2,7 @@
 
 #include "Dataflow/DataflowObject.h"
 #include "Dataflow/Dataflow.h"
+#include "Dataflow/DataflowEdNode.h"
 
 
 #define LOCTEXT_NAMESPACE "UDataflow")
@@ -22,6 +23,17 @@ void UDataflow::PostEditChangeProperty(struct FPropertyChangedEvent& PropertyCha
 
 void UDataflow::PostLoad()
 {
+#if WITH_EDITOR
+	for (UEdGraphNode* EdNode : Nodes)
+	{
+		UDataflowEdNode* DataflowEdNode = Cast<UDataflowEdNode>(EdNode);
+		if(ensure(DataflowEdNode))
+		{
+			DataflowEdNode->SetDataflowGraph(Dataflow);
+		}
+	}
+#endif
+
 	UObject::PostLoad();
 }
 
