Dataflow - DeleteNode

diff --git a/Engine/Plugins/Experimental/Dataflow/Source/DataflowEditor/Private/Dataflow/DataflowEditorToolkit.cpp b/Engine/Plugins/Experimental/Dataflow/Source/DataflowEditor/Private/Dataflow/DataflowEditorToolkit.cpp
--- a/Engine/Plugins/Experimental/Dataflow/Source/DataflowEditor/Private/Dataflow/DataflowEditorToolkit.cpp
+++ b/Engine/Plugins/Experimental/Dataflow/Source/DataflowEditor/Private/Dataflow/DataflowEditorToolkit.cpp
@@ -16,6 +16,7 @@
 #include "Modules/ModuleManager.h"
 #include "PropertyEditorModule.h"
 #include "Styling/SlateStyleRegistry.h"
+#include "Framework/Commands/GenericCommands.h"
 
 #define LOCTEXT_NAMESPACE "DataflowEditorToolkit"
 
@@ -96,6 +97,27 @@ void FDataflowEditorToolkit::EvaluateNode()
 	}
 }
 
+void FDataflowEditorToolkit::DeleteNode()
+{
+	if (UDataflow* Graph = dynamic_cast<UDataflow*>(GraphEditor->GetCurrentGraph()))
+	{
+		if (TSharedPtr<Dataflow::FGraph> DataflowGraph = Graph->GetDataflow())
+		{
+			for (UObject* Ode : GetSelectedNodes())
+			{
+				if (UDataflowEdNode* EdNode = dynamic_cast<UDataflowEdNode*>(Ode))
+				{
+					if (TSharedPtr<Dataflow::FNode> DataflowNode = DataflowGraph->FindBaseNode(EdNode->GetDataflowNodeGuid()))
+					{
+						Graph->RemoveNode(EdNode);
+						DataflowGraph->RemoveNode(DataflowNode);
+					}
+				}
+			}
+		}
+	}
+}
+
 FGraphPanelSelectionSet FDataflowEditorToolkit::GetSelectedNodes() const
 {
 	if (GraphEditor.IsValid())
@@ -123,6 +145,9 @@ TSharedRef<SGraphEditor> FDataflowEditorToolkit::CreateGraphEditorWidget(UDatafl
 	{
 		GraphEditorCommands = MakeShareable(new FUICommandList);
 		{
+			GraphEditorCommands->MapAction(FGenericCommands::Get().Delete,
+				FExecuteAction::CreateSP(this, &FDataflowEditorToolkit::DeleteNode)
+			);
 			GraphEditorCommands->MapAction(FDataflowEditorCommands::Get().EvaluateNode,
 				FExecuteAction::CreateSP(this, &FDataflowEditorToolkit::EvaluateNode)
 			);
diff --git a/Engine/Plugins/Experimental/Dataflow/Source/DataflowEditor/Private/Dataflow/DataflowEditorToolkit.h b/Engine/Plugins/Experimental/Dataflow/Source/DataflowEditor/Private/Dataflow/DataflowEditorToolkit.h
--- a/Engine/Plugins/Experimental/Dataflow/Source/DataflowEditor/Private/Dataflow/DataflowEditorToolkit.h
+++ b/Engine/Plugins/Experimental/Dataflow/Source/DataflowEditor/Private/Dataflow/DataflowEditorToolkit.h
@@ -36,6 +36,7 @@ public:
 
 	// Graph Editor Operations
 	void OnSelectedNodesChanged(const TSet<UObject*>& NewSelection);
+	void DeleteNode();
 	void EvaluateNode();
 	FGraphPanelSelectionSet GetSelectedNodes() const;
 
diff --git a/Engine/Plugins/Experimental/Dataflow/Source/DataflowEditor/Private/Dataflow/DataflowSchema.cpp b/Engine/Plugins/Experimental/Dataflow/Source/DataflowEditor/Private/Dataflow/DataflowSchema.cpp
--- a/Engine/Plugins/Experimental/Dataflow/Source/DataflowEditor/Private/Dataflow/DataflowSchema.cpp
+++ b/Engine/Plugins/Experimental/Dataflow/Source/DataflowEditor/Private/Dataflow/DataflowSchema.cpp
@@ -28,9 +28,6 @@ void UDataflowSchema::GetContextMenuActions(class UToolMenu* Menu, class UGraphN
 		FToolMenuSection& Section = Menu->AddSection("TestGraphSchemaNodeActions", LOCTEXT("ClassActionsMenuHeader", "Node Actions"));
 		{
 			Section.AddMenuEntry(FGenericCommands::Get().Delete);
-			Section.AddMenuEntry(FGenericCommands::Get().Cut);
-			Section.AddMenuEntry(FGenericCommands::Get().Copy);
-			Section.AddMenuEntry(FGenericCommands::Get().Duplicate);
 			Section.AddMenuEntry(FGraphEditorCommands::Get().BreakNodeLinks);
 			Section.AddMenuEntry(FDataflowEditorCommands::Get().EvaluateNode);
 		}
