Dataflow - Disable installed

diff --git a/Engine/Plugins/Experimental/Dataflow/Dataflow.uplugin b/Engine/Plugins/Experimental/Dataflow/Dataflow.uplugin
--- a/Engine/Plugins/Experimental/Dataflow/Dataflow.uplugin
+++ b/Engine/Plugins/Experimental/Dataflow/Dataflow.uplugin
@@ -14,7 +14,7 @@
     "CanContainContent" : true,
     "IsBetaVersion" : true,
 	"IsExperimentalVersion": true,
-    "Installed" : true,
+    "Installed" : false,
 	"Modules": [
 		{
 			"Name": "DataflowEditor",
