ChaosFlesh : Dataflow
- Serialization ensure
- Connection checks
- AssetNodes function.

diff --git a/Engine/Plugins/Experimental/Dataflow/Source/DataflowEditor/Private/Dataflow/DataflowSchema.cpp b/Engine/Plugins/Experimental/Dataflow/Source/DataflowEditor/Private/Dataflow/DataflowSchema.cpp
--- a/Engine/Plugins/Experimental/Dataflow/Source/DataflowEditor/Private/Dataflow/DataflowSchema.cpp
+++ b/Engine/Plugins/Experimental/Dataflow/Source/DataflowEditor/Private/Dataflow/DataflowSchema.cpp
@@ -50,19 +50,50 @@ void UDataflowSchema::GetGraphContextActions(FGraphContextMenuBuilder& ContextMe
 	}
 }
 
-const FPinConnectionResponse UDataflowSchema::CanCreateConnection(const UEdGraphPin* PinA, const UEdGraphPin* PinB) const
+const FPinConnectionResponse UDataflowSchema::CanCreateConnection(const UEdGraphPin* InPinA, const UEdGraphPin* InPinB) const
 {
-	// Make sure the pins are not on the same node
-	if (PinA->GetOwningNode() != PinB->GetOwningNode())
+	bool bSwapped = false;
+	const UEdGraphPin* PinA = InPinA;
+	const UEdGraphPin* PinB = InPinB;
+	if (PinA->Direction == EEdGraphPinDirection::EGPD_Input &&
+		PinB->Direction == EEdGraphPinDirection::EGPD_Output)
 	{
-		// Make sure types match. 
-		if (PinA->PinType == PinB->PinType)
+		bSwapped = true;
+		PinA = InPinB; PinB = InPinA;
+	}
+
+
+	if (PinA->Direction == EEdGraphPinDirection::EGPD_Output)
+	{
+		if (PinB->Direction == EEdGraphPinDirection::EGPD_Input)
 		{
-			return FPinConnectionResponse(CONNECT_RESPONSE_MAKE, LOCTEXT("PinConnect", "Connect nodes"));
+			// Make sure the pins are not on the same node
+			if (PinA->GetOwningNode() != PinB->GetOwningNode())
+			{
+				// Make sure types match. 
+				if (PinA->PinType == PinB->PinType)
+				{
+					if (PinB->LinkedTo.Num())
+					{
+						return (bSwapped) ?
+							FPinConnectionResponse(CONNECT_RESPONSE_BREAK_OTHERS_A, LOCTEXT("PinSteal", "Disconnect existing input and connect new input."))
+							:
+							FPinConnectionResponse(CONNECT_RESPONSE_BREAK_OTHERS_B, LOCTEXT("PinSteal", "Disconnect existing input and connect new input."));
+
+					}
+					return FPinConnectionResponse(CONNECT_RESPONSE_MAKE, LOCTEXT("PinConnect", "Connect input to output."));
+				}
+			}
 		}
 	}
-
-	return FPinConnectionResponse(CONNECT_RESPONSE_DISALLOW, LOCTEXT("PinErrorSameNode", "Both are on the same node"));
+	TArray<FText> NoConnectionResponse = {
+		LOCTEXT("PinErrorSameNode", "Nope"), 
+		LOCTEXT("PinErrorSameNode", "Sorry :("),
+		LOCTEXT("PinErrorSameNode", "Not gonna work."),
+		LOCTEXT("PinErrorSameNode", "Still no!"),
+		LOCTEXT("PinErrorSameNode", "Try again?"),
+	};
+	return FPinConnectionResponse(CONNECT_RESPONSE_DISALLOW, NoConnectionResponse[FMath::RandRange(0, NoConnectionResponse.Num()-1)]);
 }
 
 
diff --git a/Engine/Source/Runtime/Experimental/Dataflow/Core/Private/Dataflow/DataflowGraph.cpp b/Engine/Source/Runtime/Experimental/Dataflow/Core/Private/Dataflow/DataflowGraph.cpp
--- a/Engine/Source/Runtime/Experimental/Dataflow/Core/Private/Dataflow/DataflowGraph.cpp
+++ b/Engine/Source/Runtime/Experimental/Dataflow/Core/Private/Dataflow/DataflowGraph.cpp
@@ -173,10 +173,13 @@ namespace Dataflow
 			Ar << LocalConnections;
 			for (const FLink& Con : LocalConnections)
 			{
-				ensure(ConnectionGuidMap.Contains(Con.Input));
-				ensure(ConnectionGuidMap.Contains(Con.Output));
-				ensure(ConnectionGuidMap[Con.Input]->GetType() == ConnectionGuidMap[Con.Output]->GetType());
-				Connect(ConnectionGuidMap[Con.Input], ConnectionGuidMap[Con.Output]);
+				if (ensure(ConnectionGuidMap.Contains(Con.Input) && ConnectionGuidMap.Contains(Con.Output)))
+				{
+					if (ensure(ConnectionGuidMap[Con.Input]->GetType() == ConnectionGuidMap[Con.Output]->GetType()))
+					{
+						Connect(ConnectionGuidMap[Con.Input], ConnectionGuidMap[Con.Output]);
+					}
+				}
 			}
 
 		}
