Dataflow : Properties

diff --git a/Engine/Plugins/Experimental/Dataflow/Source/DataflowEditor/Private/Dataflow/DataflowEPropertyCustomizations.cpp b/Engine/Plugins/Experimental/Dataflow/Source/DataflowEditor/Private/Dataflow/DataflowEPropertyCustomizations.cpp
new file mode 100644
--- /dev/null
+++ b/Engine/Plugins/Experimental/Dataflow/Source/DataflowEditor/Private/Dataflow/DataflowEPropertyCustomizations.cpp
@@ -0,0 +1,70 @@
+// Copyright Epic Games, Inc. All Rights Reserved.
+#include "Dataflow/DataflowEPropertyCustomizations.h"
+
+#include "Dataflow/DataflowNode.h"
+#include "Dataflow/DataflowProperty.h"
+#include "DetailLayoutBuilder.h"
+#include "DetailCategoryBuilder.h"
+#include "DetailWidgetRow.h"
+#include "Dataflow/DataflowEPropertyFactories.h"
+
+#define LOCTEXT_NAMESPACE "DataflowEProprty_Customizations"
+
+DEFINE_LOG_CATEGORY_STATIC(LogCDataflowEProprtyCustomizations, Log, All);
+
+void UDataflowSEditorObject::PostEditChangeProperty(struct FPropertyChangedEvent& PropertyChangedEvent)
+{
+	UE_LOG(LogCDataflowEProprtyCustomizations, Verbose, TEXT("Here"));
+}
+
+TSharedRef<IDetailCustomization> FDataflowSEditorCustomization::MakeInstance()
+{
+	return MakeShareable(new FDataflowSEditorCustomization);
+}
+
+void FDataflowSEditorCustomization::CustomizeDetails(IDetailLayoutBuilder& DetailBuilder)
+{
+	IDetailCategoryBuilder& Category = DetailBuilder.EditCategory("Dataflow");
+
+	TArray< TWeakObjectPtr<UObject> > OutObjects;
+	DetailBuilder.GetObjectsBeingCustomized(OutObjects);
+	for (TWeakObjectPtr<UObject> Obj : OutObjects)
+	{
+		if (UDataflowSEditorObject* DataflowSObj = Cast< UDataflowSEditorObject>(Obj))
+		{
+			if (TSharedPtr<Dataflow::FNode> Node = DataflowSObj->Node)
+			{
+				for (Dataflow::FProperty* Property : Node->GetProperties())
+				{
+					IDetailPropertyRow& PropertyRow = Category.AddProperty(FName("CustomProperty"));
+					TSharedPtr<IPropertyHandle> PropertyHandle = PropertyRow.GetPropertyHandle();
+
+					switch (Property->GetType())
+					{
+					case Dataflow::FProperty::EType::BOOL:
+						Dataflow::PropertyWidgetFactory<bool>(DetailBuilder, PropertyHandle, Node, (Dataflow::TProperty<bool>*)Property);
+						break;
+					case Dataflow::FProperty::EType::INT:
+						Dataflow::PropertyWidgetFactory<int>(DetailBuilder, PropertyHandle, Node, (Dataflow::TProperty<int>*)Property);
+						break;
+					case Dataflow::FProperty::EType::FLOAT:
+						Dataflow::PropertyWidgetFactory<float>(DetailBuilder, PropertyHandle, Node, (Dataflow::TProperty<float>*)Property);
+						break;
+					case Dataflow::FProperty::EType::DOUBLE:
+						Dataflow::PropertyWidgetFactory<double>(DetailBuilder, PropertyHandle, Node, (Dataflow::TProperty<double>*)Property);
+						break;
+					case Dataflow::FProperty::EType::STRING:
+						Dataflow::PropertyWidgetFactory<FString>(DetailBuilder, PropertyHandle, Node, (Dataflow::TProperty<FString>*)Property);
+						break;
+					default:
+						ensureMsgf(true,TEXT("Missing slate property convert."));
+					}
+					UE_LOG(LogCDataflowEProprtyCustomizations, Verbose, TEXT("FFloatObjectDetails::CustomizeDetails"));
+				}
+			}
+		}
+	}
+}
+
+
+#undef LOCTEXT_NAMESPACE
diff --git a/Engine/Plugins/Experimental/Dataflow/Source/DataflowEditor/Private/Dataflow/DataflowEPropertyFactories.cpp b/Engine/Plugins/Experimental/Dataflow/Source/DataflowEditor/Private/Dataflow/DataflowEPropertyFactories.cpp
new file mode 100644
--- /dev/null
+++ b/Engine/Plugins/Experimental/Dataflow/Source/DataflowEditor/Private/Dataflow/DataflowEPropertyFactories.cpp
@@ -0,0 +1,13 @@
+// Copyright Epic Games, Inc. All Rights Reserved.
+
+#include "Dataflow/DataflowEPropertyFactories.h"
+#include "Dataflow/DataflowEditorPlugin.h"
+
+#include "Dataflow/DataflowObject.h"
+
+#define LOCTEXT_NAMESPACE "/DataflowEProperty_Factories"
+namespace Dataflow
+{
+
+}
+#undef LOCTEXT_NAMESPACE
diff --git a/Engine/Plugins/Experimental/Dataflow/Source/DataflowEditor/Private/Dataflow/DataflowEditorPlugin.cpp b/Engine/Plugins/Experimental/Dataflow/Source/DataflowEditor/Private/Dataflow/DataflowEditorPlugin.cpp
--- a/Engine/Plugins/Experimental/Dataflow/Source/DataflowEditor/Private/Dataflow/DataflowEditorPlugin.cpp
+++ b/Engine/Plugins/Experimental/Dataflow/Source/DataflowEditor/Private/Dataflow/DataflowEditorPlugin.cpp
@@ -8,6 +8,7 @@
 #include "DataflowEditorToolkit.h"
 #include "Dataflow/DataflowNodeFactory.h"
 #include "Dataflow/DataflowAssetActions.h"
+#include "Dataflow/DataflowEPropertyCustomizations.h"
 #include "Dataflow/DataflowSNodeFactories.h"
 
 #define LOCTEXT_NAMESPACE "DataflowEditor"
@@ -23,6 +24,8 @@
 	IAssetTools& AssetTools = AssetToolsModule.Get();
 	AssetTools.RegisterAssetTypeActions(MakeShareable(DataflowAssetActions));
 
+	FPropertyEditorModule& PropertyEditorModule = FModuleManager::LoadModuleChecked<FPropertyEditorModule>(TEXT("PropertyEditor"));
+	PropertyEditorModule.RegisterCustomClassLayout(UDataflowSEditorObject::StaticClass()->GetFName(), FOnGetDetailCustomizationInstance::CreateStatic(&FDataflowSEditorCustomization::MakeInstance));
 
 	DataflowSNodeFactory = MakeShareable(new FDataflowSNodeFactory());
 	FEdGraphUtilities::RegisterVisualNodeFactory(DataflowSNodeFactory);
diff --git a/Engine/Plugins/Experimental/Dataflow/Source/DataflowEditor/Public/Dataflow/DataflowEPropertyCustomizations.h b/Engine/Plugins/Experimental/Dataflow/Source/DataflowEditor/Public/Dataflow/DataflowEPropertyCustomizations.h
new file mode 100644
--- /dev/null
+++ b/Engine/Plugins/Experimental/Dataflow/Source/DataflowEditor/Public/Dataflow/DataflowEPropertyCustomizations.h
@@ -0,0 +1,40 @@
+// Copyright Epic Games, Inc. All Rights Reserved.
+
+#pragma once
+#include "CoreMinimal.h"
+#include "IDetailCustomization.h"
+#include "IPropertyTypeCustomization.h"
+#include "PropertyHandle.h"
+#include "UObject/ObjectMacros.h"
+
+#include "DataflowEPropertyCustomizations.generated.h"
+
+namespace Dataflow {
+	class FNode;
+}
+
+
+UCLASS()
+class DATAFLOWEDITOR_API UDataflowSEditorObject : public UObject
+{
+	GENERATED_BODY()
+public:
+	TSharedPtr<Dataflow::FNode> Node = nullptr;
+
+	virtual void PostEditChangeProperty(struct FPropertyChangedEvent& PropertyChangedEvent) override;
+};
+
+/**
+* Detail customization for PS5 target settings panel
+*/
+class DATAFLOWEDITOR_API FDataflowSEditorCustomization : public IDetailCustomization
+{
+public:
+	/** Makes a new instance of this detail layout class for a specific detail view requesting it */
+	static TSharedRef<IDetailCustomization> MakeInstance();
+
+	/** IDetailCustomization interface */
+	virtual void CustomizeDetails(IDetailLayoutBuilder& DetailBuilder) override;
+
+};
+
diff --git a/Engine/Plugins/Experimental/Dataflow/Source/DataflowEditor/Public/Dataflow/DataflowEPropertyFactories.h b/Engine/Plugins/Experimental/Dataflow/Source/DataflowEditor/Public/Dataflow/DataflowEPropertyFactories.h
new file mode 100644
--- /dev/null
+++ b/Engine/Plugins/Experimental/Dataflow/Source/DataflowEditor/Public/Dataflow/DataflowEPropertyFactories.h
@@ -0,0 +1,107 @@
+// Copyright Epic Games, Inc. All Rights Reserved.
+
+#pragma once
+
+#include "CoreMinimal.h"
+#include "Dataflow/DataflowProperty.h"
+#include "DetailCategoryBuilder.h"
+#include "DetailLayoutBuilder.h"
+#include "DetailWidgetRow.h"
+#include "IDetailCustomization.h"
+#include "Misc/Optional.h"
+#include "Widgets/Input/SNumericEntryBox.h"
+
+namespace Dataflow
+{
+	// todo(Dataflow) : Add support for FFilename as SFilePathPicker
+
+	//
+	//  T
+	//
+	template<class T> DATAFLOWEDITOR_API void PropertyWidgetFactory(
+		IDetailLayoutBuilder& InDetailBuilder,
+		TSharedPtr<IPropertyHandle> PropertyHandle,
+		TSharedPtr<Dataflow::FNode> InNode,
+		TProperty<T>* InProperty) 
+	{
+		InDetailBuilder.AddCustomRowToCategory(PropertyHandle, FText::FromString("PropertyHandel<T>"))
+			.NameContent()[
+				SNew(STextBlock).Font(IDetailLayoutBuilder::GetDetailFont())
+					.Text(FText::FromName(InProperty->GetName()))
+			].ValueContent()[
+				SNew(SNumericEntryBox<T>)
+					.Font(IDetailLayoutBuilder::GetDetailFont())
+					.MinValue(TOptional<T>())
+					.MaxValue(TOptional<T>())
+					.MinSliderValue(TOptional<T>()) // No lower limit
+					.MaxSliderValue(TOptional<T>()) // No upper limit
+					.Value_Lambda([InNode, InProperty]()
+						{
+							return InProperty->GetValue();
+						})
+					.OnValueCommitted_Lambda([InNode, InProperty](T NewValue, ETextCommit::Type CommitType)
+						{
+							InProperty->SetValue(NewValue);
+						})
+			];
+
+	}
+
+	//
+	// bool
+	//
+	template<> inline
+		void PropertyWidgetFactory<bool>(
+			IDetailLayoutBuilder& InDetailBuilder,
+			TSharedPtr<IPropertyHandle> PropertyHandle,
+			TSharedPtr<Dataflow::FNode> InNode,
+			TProperty<bool>* InProperty)
+	{
+		InDetailBuilder.AddCustomRowToCategory(PropertyHandle, FText::FromString("PropertyHandel<bool>"))
+			.NameContent()[
+				SNew(STextBlock).Font(IDetailLayoutBuilder::GetDetailFont())
+					.Text(FText::FromName(InProperty->GetName()))
+			].ValueContent()[
+				SNew(SCheckBox)
+					.IsChecked_Lambda([InNode, InProperty]() -> ECheckBoxState
+						{
+							return InProperty->GetValue()?ECheckBoxState::Checked:ECheckBoxState::Unchecked;
+						})
+					.OnCheckStateChanged_Lambda([InNode, InProperty](ECheckBoxState NewState)
+						{
+							InProperty->SetValue(NewState == ECheckBoxState::Checked);
+						})
+
+			];
+	}
+
+	//
+	// FString
+	//
+	template<> inline
+	void PropertyWidgetFactory<FString>(
+		IDetailLayoutBuilder& InDetailBuilder,
+		TSharedPtr<IPropertyHandle> PropertyHandle,
+		TSharedPtr<Dataflow::FNode> InNode,
+		TProperty<FString>* InProperty)
+	{
+		InDetailBuilder.AddCustomRowToCategory(PropertyHandle, FText::FromString("PropertyHandel<FString>"))
+			.NameContent()[
+				SNew(STextBlock).Font(IDetailLayoutBuilder::GetDetailFont())
+					.Text(FText::FromName(InProperty->GetName()))
+			].ValueContent()[
+				SNew(SEditableTextBox)
+					.Font(IDetailLayoutBuilder::GetDetailFont())
+					.Text_Lambda([InNode, InProperty]()
+						{
+							return FText::FromString(InProperty->GetValue());
+						})
+					.OnTextCommitted_Lambda([InNode, InProperty](const FText& InText, ETextCommit::Type TextCommitType)
+						{
+							InProperty->SetValue(InText.ToString());
+						})
+
+			];
+	}
+
+}
diff --git a/Engine/Source/Runtime/Experimental/Dataflow/Core/Private/Dataflow/DataflowNode.cpp b/Engine/Source/Runtime/Experimental/Dataflow/Core/Private/Dataflow/DataflowNode.cpp
--- a/Engine/Source/Runtime/Experimental/Dataflow/Core/Private/Dataflow/DataflowNode.cpp
+++ b/Engine/Source/Runtime/Experimental/Dataflow/Core/Private/Dataflow/DataflowNode.cpp
@@ -6,6 +6,16 @@
 
 namespace Dataflow
 {
+
+	void FNode::AddProperty(FProperty* InPtr)
+	{
+		for (FProperty* In : Properties)
+		{
+			ensureMsgf(!In->GetName().IsEqual(InPtr->GetName()), TEXT("Add Property Failed: Existing Node property already defined with name (%s)"), *InPtr->GetName().ToString());
+		}
+		Properties.Add(InPtr);
+	}
+
 	void FNode::AddInput(FConnection* InPtr)
 	{ 
 		for (FConnection* In : Inputs)
@@ -68,5 +78,14 @@
 		return nullptr;
 	}
 
+
+	void FNode::SerializeInternal(FArchive& Ar) 
+	{
+		for (FProperty* Prop : GetProperties())
+		{
+			Prop->Serialize(Ar);
+		}
+	}
+
 }
 
diff --git a/Engine/Source/Runtime/Experimental/Dataflow/Core/Private/Dataflow/DataflowProperty.cpp b/Engine/Source/Runtime/Experimental/Dataflow/Core/Private/Dataflow/DataflowProperty.cpp
new file mode 100644
--- /dev/null
+++ b/Engine/Source/Runtime/Experimental/Dataflow/Core/Private/Dataflow/DataflowProperty.cpp
@@ -0,0 +1,13 @@
+// Copyright Epic Games, Inc. All Rights Reserved.
+
+#include "Dataflow/DataflowProperty.h"
+#include "Dataflow/DataflowNode.h"
+
+namespace Dataflow
+{
+	void FProperty::BindProperty(FNode* InNode, FProperty* That)
+	{
+		InNode->AddProperty(That);
+	}
+}
+
diff --git a/Engine/Source/Runtime/Experimental/Dataflow/Core/Public/Dataflow/DataflowNode.h b/Engine/Source/Runtime/Experimental/Dataflow/Core/Public/Dataflow/DataflowNode.h
--- a/Engine/Source/Runtime/Experimental/Dataflow/Core/Public/Dataflow/DataflowNode.h
+++ b/Engine/Source/Runtime/Experimental/Dataflow/Core/Public/Dataflow/DataflowNode.h
@@ -4,6 +4,7 @@
 
 #include "CoreMinimal.h"
 #include "Dataflow/DataflowConnection.h"
+#include "Dataflow/DataflowProperty.h"
 #include "Dataflow/DataflowNodeParameters.h"
 
 namespace Dataflow
@@ -21,12 +22,14 @@
 	{
 		friend class FGraph;
 		friend class FConnection;
+		friend class FProperty;
 
 		FGuid Guid;
 		FName Name;
 
 		TArray< FConnection* > Inputs;
 		TArray< FConnection* > Outputs;
+		TArray< FProperty* > Properties;
 	public:
 
 		FNode(const FNodeParameters& Param, FGuid InGuid = FGuid::NewGuid())
@@ -49,8 +52,12 @@
 		virtual void Evaluate(const FContext& Context, FConnection*) const { ensure(false); }
 		void InvalidateOutputs();
 
-		virtual void SerializeInternal(FArchive& Ar) {};
+		virtual void SerializeInternal(FArchive& Ar);
 
+		void AddProperty(FProperty* InPtr);
+		const TArray< FProperty* >& GetProperties() const { return Properties; }
+		TArray< FProperty* >& GetProperties() { return Properties; }
+
 		void AddInput(FConnection* InPtr);
 		const TArray< FConnection* >& GetInputs() const { return Inputs; }
 		TArray< FConnection* >& GetInputs() { return Inputs; }
diff --git a/Engine/Source/Runtime/Experimental/Dataflow/Core/Public/Dataflow/DataflowProperty.h b/Engine/Source/Runtime/Experimental/Dataflow/Core/Public/Dataflow/DataflowProperty.h
new file mode 100644
--- /dev/null
+++ b/Engine/Source/Runtime/Experimental/Dataflow/Core/Public/Dataflow/DataflowProperty.h
@@ -0,0 +1,115 @@
+// Copyright Epic Games, Inc. All Rights Reserved.
+
+#pragma once
+
+#include "CoreMinimal.h"
+#include "Chaos/ChaosArchive.h"
+#include "Serialization/Archive.h"
+
+namespace Dataflow
+{
+	class FNode;
+
+	struct FFilename
+	{
+		FString Filename;
+		FMD5Hash FileHash;
+
+		FFilename() {}
+		FFilename(const FString& File) : Filename(File) {}
+	};
+
+
+	template<class T>
+	struct DATAFLOWCORE_API TPropertyParameters {
+		FName Name = FName("");
+		T Value = T();
+		FNode* Node = nullptr;
+	};
+
+	//
+	// FParameter
+	//
+	class  DATAFLOWCORE_API FProperty 
+	{
+	public:
+		enum class EType : uint8 {
+			NONE = 0,
+			BOOL,
+			INT,
+			FLOAT,
+			DOUBLE,
+			STRING,
+			FILE,
+			// <new types go here>
+			MAX
+		};
+
+
+		FProperty(EType InType, FName InName, FNode* Node)
+			: Type(InType)
+			, Name(InName)
+			, Category("")
+		{
+			BindProperty(Node, this);
+		}
+		virtual ~FProperty() {}
+
+
+		const FName& GetName() const { return Name; }
+		const EType GetType() const { return Type; }
+
+		void GetCategory(const FName& InName) { Name = InName; }
+		const FName& GetCategory() const { return Category; }
+
+		void BindProperty(FNode* InNode, FProperty* That);
+
+		virtual void Serialize(FArchive& Ar) {}
+
+	private:
+		EType Type;
+		FName Name;
+		FName Category;
+	};
+
+	//
+	// TParameter
+	//
+	template<class T>
+	class DATAFLOWCORE_API TProperty : public FProperty
+	{
+		T Value = T();
+	public:
+
+		TProperty(TPropertyParameters<T> InParam)
+			: FProperty(StaticType(), InParam.Name, InParam.Node)
+			, Value(InParam.Value)
+		{}
+
+		static FProperty::EType StaticType();
+
+		const T& GetValue() const { return Value; }
+		void SetValue(const T& InValue) { Value = InValue; }
+
+		virtual void Serialize(FArchive& Ar) override
+		{
+			Ar << Value;
+		}
+
+	};
+
+} // Dataflow
+
+#define DATAFLOW_PROPERTY(DECL, TYPE,ETYPE)																			\
+	template<> inline Dataflow::FProperty::EType Dataflow::TProperty<TYPE>::StaticType() { return EType::ETYPE; }	\
+	template struct DECL Dataflow::TPropertyParameters<TYPE>;														\
+	template class Dataflow::TProperty<TYPE>;
+
+DATAFLOW_PROPERTY(DATAFLOWCORE_API, bool, BOOL)
+DATAFLOW_PROPERTY(DATAFLOWCORE_API, int, INT)
+DATAFLOW_PROPERTY(DATAFLOWCORE_API, float, FLOAT)
+DATAFLOW_PROPERTY(DATAFLOWCORE_API, double, DOUBLE)
+DATAFLOW_PROPERTY(DATAFLOWCORE_API, FString, STRING)
+//DATAFLOW_PROPERTY(DATAFLOWCORE_API, Dataflow::FFilename, FILE)
+
+
