Fix - NoPCH CIS

diff --git a/Engine/Plugins/Experimental/Dataflow/Source/DataflowEditor/Public/Dataflow/DataflowEPropertyFactories.h b/Engine/Plugins/Experimental/Dataflow/Source/DataflowEditor/Public/Dataflow/DataflowEPropertyFactories.h
--- a/Engine/Plugins/Experimental/Dataflow/Source/DataflowEditor/Public/Dataflow/DataflowEPropertyFactories.h
+++ b/Engine/Plugins/Experimental/Dataflow/Source/DataflowEditor/Public/Dataflow/DataflowEPropertyFactories.h
@@ -10,6 +10,9 @@
 #include "IDetailCustomization.h"
 #include "Misc/Optional.h"
 #include "Widgets/Input/SNumericEntryBox.h"
+#include "Widgets/Input/SEditableTextBox.h"
+#include "Widgets/Text/STextBlock.h"
+#include "Widgets/Input/SCheckBox.h"
 
 namespace Dataflow
 {
