Dataflow - Moved functionality to Dataflow modules.

diff --git a/Engine/Plugins/Experimental/Dataflow/Source/DataflowEditor/Private/Dataflow/DataflowEditorActions.cpp b/Engine/Plugins/Experimental/Dataflow/Source/DataflowEditor/Private/Dataflow/DataflowEditorActions.cpp
--- a/Engine/Plugins/Experimental/Dataflow/Source/DataflowEditor/Private/Dataflow/DataflowEditorActions.cpp
+++ b/Engine/Plugins/Experimental/Dataflow/Source/DataflowEditor/Private/Dataflow/DataflowEditorActions.cpp
@@ -1,8 +1,12 @@
 // Copyright Epic Games, Inc. All Rights Reserved.
 
 #include "Dataflow/DataflowEditorActions.h"
+#include "Dataflow/DataflowEdNode.h"
+#include "Dataflow/DataflowEPropertyCustomizations.h"
 #include "Dataflow/DataflowNodeFactory.h"
+#include "Dataflow/DataflowObject.h"
 
+
 #define LOCTEXT_NAMESPACE "DataflowEditorCommands"
 
 void FDataflowEditorCommandsImpl::RegisterCommands()
@@ -29,19 +33,95 @@
 	}
 }
 
-void FDataflowEditorCommands::Register()
-{
-	return FDataflowEditorCommandsImpl::Register();
-}
-
 const FDataflowEditorCommandsImpl& FDataflowEditorCommands::Get()
 {
 	return FDataflowEditorCommandsImpl::Get();
 }
 
+void FDataflowEditorCommands::Register()
+{
+	return FDataflowEditorCommandsImpl::Register();
+}
+
 void FDataflowEditorCommands::Unregister()
 {
 	return FDataflowEditorCommandsImpl::Unregister();
 }
 
+
+
+void FDataflowEditorCommands::EvaluateNodes(const FGraphPanelSelectionSet& SelectedNodes, const Dataflow::FContext& InContext)
+{
+	for (UObject* Ode : SelectedNodes)
+	{
+		if (UDataflowEdNode* EdNode = dynamic_cast<UDataflowEdNode*>(Ode))
+		{
+			if (const TSharedPtr<Dataflow::FGraph> DataflowGraph = EdNode->GetDataflowGraph())
+			{
+				if (const TSharedPtr<Dataflow::FNode> DataflowNode = DataflowGraph->FindBaseNode(EdNode->GetDataflowNodeGuid()))
+				{
+					if (DataflowNode->GetOutputs().Num())
+					{
+						for (Dataflow::FConnection* NodeOutput : DataflowNode->GetOutputs())
+						{
+							DataflowNode->Evaluate(InContext, NodeOutput);
+						}
+					}
+					else
+					{
+						DataflowNode->Evaluate(InContext, nullptr);
+					}
+				}
+			}
+		}
+	}
+}
+
+void FDataflowEditorCommands::DeleteNodes(UDataflow* Graph, const FGraphPanelSelectionSet& SelectedNodes)
+{
+	for (UObject* Ode : SelectedNodes)
+	{
+		if (UDataflowEdNode* EdNode = dynamic_cast<UDataflowEdNode*>(Ode))
+		{
+			if (const TSharedPtr<Dataflow::FGraph> DataflowGraph = EdNode->GetDataflowGraph())
+			{
+				if (TSharedPtr<Dataflow::FNode> DataflowNode = DataflowGraph->FindBaseNode(EdNode->GetDataflowNodeGuid()))
+				{
+					Graph->RemoveNode(EdNode);
+					DataflowGraph->RemoveNode(DataflowNode);
+				}
+			}
+		}
+	}
+}
+
+void FDataflowEditorCommands::OnSelectedNodesChanged(TSharedPtr<IDetailsView> PropertiesEditor, UObject* Asset, UDataflow* Graph, const TSet<UObject*>& NewSelection)
+{
+	if (Graph && PropertiesEditor)
+	{
+		if (const TSharedPtr<Dataflow::FGraph> DataflowGraph = Graph->GetDataflow())
+		{
+			PropertiesEditor->SetObject(Asset);
+			FGraphPanelSelectionSet SelectedNodes = NewSelection;//GetSelectedNodes();
+			if (SelectedNodes.Num())
+			{
+				TArray<UObject*> Objects;
+				for (UObject* SelectedObject : SelectedNodes)
+				{
+					if (UDataflowEdNode* EdNode = Cast<UDataflowEdNode>(SelectedObject))
+					{
+						if (TSharedPtr<Dataflow::FNode> DataflowNode = DataflowGraph->FindBaseNode(EdNode->GetDataflowNodeGuid()))
+						{
+							UDataflowSEditorObject* Object = NewObject<UDataflowSEditorObject>(Asset, DataflowNode->GetName());
+							Object->Node = DataflowNode;
+							Objects.Add(Object);
+						}
+					}
+				}
+				PropertiesEditor->SetObjects(Objects);
+			}
+		}
+	}
+}
+
 #undef LOCTEXT_NAMESPACE
diff --git a/Engine/Plugins/Experimental/Dataflow/Source/DataflowEditor/Private/Dataflow/DataflowEditorToolkit.cpp b/Engine/Plugins/Experimental/Dataflow/Source/DataflowEditor/Private/Dataflow/DataflowEditorToolkit.cpp
--- a/Engine/Plugins/Experimental/Dataflow/Source/DataflowEditor/Private/Dataflow/DataflowEditorToolkit.cpp
+++ b/Engine/Plugins/Experimental/Dataflow/Source/DataflowEditor/Private/Dataflow/DataflowEditorToolkit.cpp
@@ -74,47 +74,19 @@
 	}
 }
 
+
+
 void FDataflowEditorToolkit::EvaluateNode()
 {
-	if (UDataflow* Graph = dynamic_cast<UDataflow*>(GraphEditor->GetCurrentGraph()))
-	{
-		if (TSharedPtr<Dataflow::FGraph> DataflowGraph = Graph->GetDataflow())
-		{
-			for (UObject* Ode : GetSelectedNodes())
-			{
-				if (UDataflowEdNode* Node = dynamic_cast<UDataflowEdNode*>(Ode))
-				{
-					if (TSharedPtr<Dataflow::FNode> DataflowNode = DataflowGraph->FindBaseNode(Node->GetDataflowNodeGuid()))
-					{
-						for (Dataflow::FConnection* NodeOutput : DataflowNode->GetOutputs())
-						{
-							DataflowNode->Evaluate({0.f}, NodeOutput);
-						}
-					}
-				}
-			}
-		}
-	}
+	Dataflow::FContext Context(FGameTime::GetTimeSinceAppStart().GetRealTimeSeconds());
+	FDataflowEditorCommands::EvaluateNodes(GetSelectedNodes(), Context);
 }
 
 void FDataflowEditorToolkit::DeleteNode()
 {
 	if (UDataflow* Graph = dynamic_cast<UDataflow*>(GraphEditor->GetCurrentGraph()))
 	{
-		if (TSharedPtr<Dataflow::FGraph> DataflowGraph = Graph->GetDataflow())
-		{
-			for (UObject* Ode : GetSelectedNodes())
-			{
-				if (UDataflowEdNode* EdNode = dynamic_cast<UDataflowEdNode*>(Ode))
-				{
-					if (TSharedPtr<Dataflow::FNode> DataflowNode = DataflowGraph->FindBaseNode(EdNode->GetDataflowNodeGuid()))
-					{
-						Graph->RemoveNode(EdNode);
-						DataflowGraph->RemoveNode(DataflowNode);
-					}
-				}
-			}
-		}
+		FDataflowEditorCommands::DeleteNodes(Graph, GetSelectedNodes());
 	}
 }
 
@@ -129,10 +101,13 @@
 
 void FDataflowEditorToolkit::OnSelectedNodesChanged(const TSet<UObject*>& NewSelection)
 {
+	if (UDataflow* Graph = dynamic_cast<UDataflow*>(GraphEditor->GetCurrentGraph()))
+	{
+		FDataflowEditorCommands::OnSelectedNodesChanged(PropertiesEditor, Dataflow, Graph, NewSelection);
+	}
 }
 
 
-
 TSharedRef<SGraphEditor> FDataflowEditorToolkit::CreateGraphEditorWidget(UDataflow* DataflowToEdit)
 {
 	ensure(DataflowToEdit);
diff --git a/Engine/Plugins/Experimental/Dataflow/Source/DataflowEditor/Public/Dataflow/DataflowEditorActions.h b/Engine/Plugins/Experimental/Dataflow/Source/DataflowEditor/Public/Dataflow/DataflowEditorActions.h
--- a/Engine/Plugins/Experimental/Dataflow/Source/DataflowEditor/Public/Dataflow/DataflowEditorActions.h
+++ b/Engine/Plugins/Experimental/Dataflow/Source/DataflowEditor/Public/Dataflow/DataflowEditorActions.h
@@ -1,36 +1,54 @@
-// Copyright Epic Games, Inc. All Rights Reserved.
-
-#pragma once
-
-#include "CoreMinimal.h"
-#include "Framework/Commands/Commands.h"
-#include "Styling/AppStyle.h"
-
-class FDataflowEditorCommandsImpl : public TCommands<FDataflowEditorCommandsImpl>
-{
-public:
-
-	FDataflowEditorCommandsImpl()
-		: TCommands<FDataflowEditorCommandsImpl>( TEXT("DataflowEditor"), NSLOCTEXT("Contexts", "DataflowEditor", "Scene Graph Editor"), NAME_None, FAppStyle::GetAppStyleSetName() )
-	{
-	}	
-
-	virtual ~FDataflowEditorCommandsImpl()
-	{
-	}
-
-	DATAFLOWEDITOR_API virtual void RegisterCommands() override;
-
-	TSharedPtr< FUICommandInfo > EvaluateNode;
-	TMap< FName, TSharedPtr<FUICommandInfo> > CreateNodesMap;
-};
-
-class DATAFLOWEDITOR_API FDataflowEditorCommands
-{
-public:
-	static void Register();
-
-	static const FDataflowEditorCommandsImpl& Get();
-
-	static void Unregister();
-};
+// Copyright Epic Games, Inc. All Rights Reserved.
+
+#pragma once
+
+#include "CoreMinimal.h"
+#include "Framework/Commands/Commands.h"
+#include "Styling/AppStyle.h"
+#include "Dataflow/DataflowNodeParameters.h"
+
+class UDataflowEdNode;
+
+class FDataflowEditorCommandsImpl : public TCommands<FDataflowEditorCommandsImpl>
+{
+public:
+
+	FDataflowEditorCommandsImpl()
+		: TCommands<FDataflowEditorCommandsImpl>( TEXT("DataflowEditor"), NSLOCTEXT("Contexts", "DataflowEditor", "Scene Graph Editor"), NAME_None, FAppStyle::GetAppStyleSetName() )
+	{
+	}	
+
+	virtual ~FDataflowEditorCommandsImpl()
+	{
+	}
+
+	DATAFLOWEDITOR_API virtual void RegisterCommands() override;
+
+	TSharedPtr< FUICommandInfo > EvaluateNode;
+	TMap< FName, TSharedPtr<FUICommandInfo> > CreateNodesMap;
+};
+
+class DATAFLOWEDITOR_API FDataflowEditorCommands
+{
+public:
+	static void Register();
+	static void Unregister();
+
+	static const FDataflowEditorCommandsImpl& Get();
+
+	/*
+	*  EvaluateNodes
+	*/
+	static void EvaluateNodes(const FGraphPanelSelectionSet& SelectedNodes, const Dataflow::FContext& InContext);
+	
+	/*
+	*  DeleteNodes
+	*/
+	static void DeleteNodes(UDataflow* Graph, const FGraphPanelSelectionSet& SelectedNodes);
+	
+	/*
+	*  OnSelectedNodesChanged
+	*/
+	static void OnSelectedNodesChanged(TSharedPtr<IDetailsView> PropertiesEditor, UObject* Asset, UDataflow* Graph, const TSet<UObject*>& NewSelection);
+
+};
