Dataflow : Update

diff --git a/Engine/Plugins/Experimental/Dataflow/Source/DataflowEditor/Private/Dataflow/DataflowEditorActions.cpp b/Engine/Plugins/Experimental/Dataflow/Source/DataflowEditor/Private/Dataflow/DataflowEditorActions.cpp
--- a/Engine/Plugins/Experimental/Dataflow/Source/DataflowEditor/Private/Dataflow/DataflowEditorActions.cpp
+++ b/Engine/Plugins/Experimental/Dataflow/Source/DataflowEditor/Private/Dataflow/DataflowEditorActions.cpp
@@ -112,7 +112,8 @@ void FDataflowEditorCommands::OnSelectedNodesChanged(TSharedPtr<IDetailsView> Pr
 					{
 						if (TSharedPtr<Dataflow::FNode> DataflowNode = DataflowGraph->FindBaseNode(EdNode->GetDataflowNodeGuid()))
 						{
-							UDataflowSEditorObject* Object = NewObject<UDataflowSEditorObject>(Asset, DataflowNode->GetName());
+							const FName NodeName = MakeUniqueObjectName(Graph, UDataflow::StaticClass(), DataflowNode->GetName());
+							UDataflowSEditorObject* Object = NewObject<UDataflowSEditorObject>(Asset, NodeName);
 							Object->Node = DataflowNode;
 							Objects.Add(Object);
 						}
