Dataflow - Viewport cleanup

diff --git a/Engine/Plugins/Experimental/Dataflow/Source/DataflowEditor/Private/Dataflow/DataflowEditorToolkit.cpp b/Engine/Plugins/Experimental/Dataflow/Source/DataflowEditor/Private/Dataflow/DataflowEditorToolkit.cpp
--- a/Engine/Plugins/Experimental/Dataflow/Source/DataflowEditor/Private/Dataflow/DataflowEditorToolkit.cpp
+++ b/Engine/Plugins/Experimental/Dataflow/Source/DataflowEditor/Private/Dataflow/DataflowEditorToolkit.cpp
@@ -9,9 +9,6 @@
 #include "Dataflow/DataflowSchema.h"
 #include "Dataflow/DataflowCore.h"
 #include "EditorStyleSet.h"
-#include "EditorViewportTabContent.h"
-#include "EditorViewportLayout.h"
-#include "EditorViewportCommands.h"
 #include "GraphEditorActions.h"
 #include "Modules/ModuleManager.h"
 #include "PropertyEditorModule.h"
diff --git a/Engine/Plugins/Experimental/Dataflow/Source/DataflowEditor/Private/Dataflow/DataflowEditorToolkit.h b/Engine/Plugins/Experimental/Dataflow/Source/DataflowEditor/Private/Dataflow/DataflowEditorToolkit.h
--- a/Engine/Plugins/Experimental/Dataflow/Source/DataflowEditor/Private/Dataflow/DataflowEditorToolkit.h
+++ b/Engine/Plugins/Experimental/Dataflow/Source/DataflowEditor/Private/Dataflow/DataflowEditorToolkit.h
@@ -7,12 +7,10 @@
 #include "Misc/NotifyHook.h"
 #include "GraphEditor.h"
 
-class FEditorViewportTabContent;
 class IDetailsView;
 class FTabManager;
 class IToolkitHost;
 class UDataflow;
-class SDataflowEditorViewport;
 
 class FDataflowEditorToolkit : public FAssetEditorToolkit, public FNotifyHook, public FGCObject
 {
