Dataflow : DragAndDrop from SkeletonTree

diff --git a/Engine/Plugins/Experimental/Dataflow/Source/DataflowEditor/Private/Dataflow/DataflowGraphEditor.cpp b/Engine/Plugins/Experimental/Dataflow/Source/DataflowEditor/Private/Dataflow/DataflowGraphEditor.cpp
--- a/Engine/Plugins/Experimental/Dataflow/Source/DataflowEditor/Private/Dataflow/DataflowGraphEditor.cpp
+++ b/Engine/Plugins/Experimental/Dataflow/Source/DataflowEditor/Private/Dataflow/DataflowGraphEditor.cpp
@@ -5,7 +5,9 @@
 #include "DataflowEditorToolkit.h"
 #include "Dataflow/DataflowEngine.h"
 #include "Dataflow/DataflowSNodeFactories.h"
+#include "BoneDragDropOp.h"
 
+DEFINE_LOG_CATEGORY_STATIC(SDataflowGraphEditorLog, Log, All);
 
 #define LOCTEXT_NAMESPACE "DataflowGraphEditor"
 
@@ -17,6 +19,7 @@ void SDataflowGraphEditor::Construct(const FArguments& InArgs, UObject* InAssetO
 	DataflowAsset = Cast<UDataflow>(InArgs._GraphToEdit);
 	DetailsView = InArgs._DetailsView;
 	EvaluateGraphCallback = InArgs._EvaluateGraph;
+	OnDragDropEventCallback = InArgs._OnDragDropEvent;
 
 	FGraphAppearanceInfo AppearanceInfo;
 	AppearanceInfo.CornerText = FText::FromString("Dataflow");
@@ -86,15 +89,53 @@ void SDataflowGraphEditor::OnSelectedNodesChanged(const TSet<UObject*>& NewSelec
 	}
 }
 
-FReply SDataflowGraphEditor::OnDrop(const FGeometry& MyGeometry, const FDragDropEvent& DragDropEvent)
+FReply SDataflowGraphEditor::OnDragOver(const FGeometry& MyGeometry, const FDragDropEvent& DragDropEvent)
 {
-	//TSharedPtr<FDataprepDragDropOp> DragNodeOp = DragDropEvent.GetOperationAs<FDataprepDragDropOp>();
-	//if (TrackGraphNodePtr.IsValid() && DragNodeOp.IsValid())
-	//{
-	//	return FReply::Handled().EndDragDrop();
-	//}
+	UE_LOG(SDataflowGraphEditorLog, All, TEXT("SDataflowGraphEditor::OnDragOver"));
+	if (TSharedPtr<FBoneDragDropOp> SchemeDragDropOp = DragDropEvent.GetOperationAs<FBoneDragDropOp>())
+	{
+		return FReply::Handled();
+	}
+	return FReply::Unhandled();
+}
+
 
+FReply SDataflowGraphEditor::OnDrop(const FGeometry& MyGeometry, const FDragDropEvent& DragDropEvent)
+{
+	if (TSharedPtr<FBoneDragDropOp> SchemeDragDropOp = DragDropEvent.GetOperationAs<FBoneDragDropOp>())
+	{
+		if (OnDragDropEventCallback)
+		{
+			UE_LOG(SDataflowGraphEditorLog, All, TEXT("SDataflowGraphEditor::OnDrop"));
+			OnDragDropEventCallback(MyGeometry,DragDropEvent);
+		}
+	}
 	return SGraphEditor::OnDrop(MyGeometry, DragDropEvent);
 }
 
+/*
+void SDataflowGraphEditor::OnDragEnter(const FGeometry& MyGeometry, const FDragDropEvent& DragDropEvent)
+{
+
+	TSharedPtr<FGraphSchemaActionDragDropAction> SchemeDragDropOp = DragDropEvent.GetOperationAs<FGraphSchemaActionDragDropAction>();
+	{
+	//	FDataprepSchemaActionContext ActionContext;
+	//	ActionContext.DataprepActionPtr = StepData->DataprepActionPtr;
+	//	ActionContext.DataprepActionStepPtr = StepData->DataprepActionStepPtr;
+	//	ActionContext.StepIndex = StepData->StepIndex;
+	//	DataprepDragDropOp->SetHoveredDataprepActionContext(ActionContext);
+	}
+	UE_LOG(SDataflowGraphEditorLog, All, TEXT("SDataflowGraphEditor::OnDragEnter"));
+
+}
+void SDataflowGraphEditor::OnDragLeave(const FDragDropEvent& DragDropEvent)
+{
+	TSharedPtr<FGraphSchemaActionDragDropAction> SchemeDragDropOp = DragDropEvent.GetOperationAs<FGraphSchemaActionDragDropAction>();
+	{
+		// DataflowDragDropOp->SetHoveredDataprepActionContext(TOptional<FDataprepSchemaActionContext>());
+	}
+	UE_LOG(SDataflowGraphEditorLog, All, TEXT("SDataflowGraphEditor::OnDragLeave"));
+}
+*/
+
 #undef LOCTEXT_NAMESPACE
diff --git a/Engine/Plugins/Experimental/Dataflow/Source/DataflowEditor/Public/Dataflow/DataflowEditorActions.h b/Engine/Plugins/Experimental/Dataflow/Source/DataflowEditor/Public/Dataflow/DataflowEditorActions.h
--- a/Engine/Plugins/Experimental/Dataflow/Source/DataflowEditor/Public/Dataflow/DataflowEditorActions.h
+++ b/Engine/Plugins/Experimental/Dataflow/Source/DataflowEditor/Public/Dataflow/DataflowEditorActions.h
@@ -44,6 +44,7 @@ class DATAFLOWEDITOR_API FDataflowEditorCommands
 {
 public:
 	typedef TFunction<void(Dataflow::FNode*, Dataflow::FConnection*)> FGraphEvaluationCallback;
+	typedef TFunction<void(const FGeometry& MyGeometry, const FDragDropEvent& DragDropEvent)> FOnDragDropEventCallback;
 
 	static void Register();
 	static void Unregister();
diff --git a/Engine/Plugins/Experimental/Dataflow/Source/DataflowEditor/Public/Dataflow/DataflowGraphEditor.h b/Engine/Plugins/Experimental/Dataflow/Source/DataflowEditor/Public/Dataflow/DataflowGraphEditor.h
--- a/Engine/Plugins/Experimental/Dataflow/Source/DataflowEditor/Public/Dataflow/DataflowGraphEditor.h
+++ b/Engine/Plugins/Experimental/Dataflow/Source/DataflowEditor/Public/Dataflow/DataflowGraphEditor.h
@@ -46,10 +46,15 @@ public:
 	SLATE_ARGUMENT(FGraphEditorEvents, GraphEvents)
 	SLATE_ARGUMENT(TSharedPtr<IDetailsView>, DetailsView)
 	SLATE_ARGUMENT(FDataflowEditorCommands::FGraphEvaluationCallback, EvaluateGraph)
+	SLATE_ARGUMENT(FDataflowEditorCommands::FOnDragDropEventCallback, OnDragDropEvent)
 	SLATE_END_ARGS()
 
 	// SWidget overrides
+	virtual FReply OnDragOver(const FGeometry& MyGeometry, const FDragDropEvent& DragDropEvent) override;
 	virtual FReply OnDrop(const FGeometry& MyGeometry, const FDragDropEvent& DragDropEvent) override;
+	//virtual void OnDragEnter(const FGeometry& MyGeometry, const FDragDropEvent& DragDropEvent) override;
+	//virtual void OnDragLeave(const FDragDropEvent& DragDropEvent) override;
+	// end SWidget
 
 	/** */
 	void Construct(const FArguments& InArgs, UObject* AssetOwner);
@@ -67,6 +72,7 @@ public:
 
 
 private:
+	FDataflowEditorCommands::FOnDragDropEventCallback OnDragDropEventCallback;
 	FDataflowEditorCommands::FGraphEvaluationCallback EvaluateGraphCallback;
 
 	/** The asset that ownes this dataflow graph */
