Dataflow - DragAndDrop support

diff --git a/Engine/Plugins/Experimental/Dataflow/Dataflow.uplugin b/Engine/Plugins/Experimental/Dataflow/Dataflow.uplugin
--- a/Engine/Plugins/Experimental/Dataflow/Dataflow.uplugin
+++ b/Engine/Plugins/Experimental/Dataflow/Dataflow.uplugin
@@ -20,6 +20,11 @@
 			"Name": "DataflowEditor",
 			"Type": "Editor",
 			"LoadingPhase": "Default"
+		},
+		{
+			"Name": "DataflowNodes",
+			"Type": "Runtime",
+			"LoadingPhase": "Default"
 		}
 	],
 	"Plugins": [
diff --git a/Engine/Plugins/Experimental/Dataflow/Source/DataflowNodes/DataflowNodes.Build.cs b/Engine/Plugins/Experimental/Dataflow/Source/DataflowNodes/DataflowNodes.Build.cs
new file mode 100644
--- /dev/null
+++ b/Engine/Plugins/Experimental/Dataflow/Source/DataflowNodes/DataflowNodes.Build.cs
@@ -0,0 +1,26 @@
+// Copyright Epic Games, Inc. All Rights Reserved.
+
+namespace UnrealBuildTool.Rules
+{
+	public class DataflowNodes : ModuleRules
+	{
+        public DataflowNodes(ReadOnlyTargetRules Target) : base(Target)
+		{
+			PrivateIncludePaths.Add("DataflowNodes/Private");
+            PublicIncludePaths.Add(ModuleDirectory + "/Public");
+
+			PublicDependencyModuleNames.AddRange(
+				new string[]
+				{
+					"Core",
+					"CoreUObject",
+					"ChaosCore",
+					"Chaos",
+					"DataflowCore",
+					"DataflowEngine",
+					"Engine"
+				}
+			);
+		}
+	}
+}
diff --git a/Engine/Plugins/Experimental/Dataflow/Source/DataflowNodes/Private/Dataflow/DataflowNodesConnectionTypes.cpp b/Engine/Plugins/Experimental/Dataflow/Source/DataflowNodes/Private/Dataflow/DataflowNodesConnectionTypes.cpp
new file mode 100644
--- /dev/null
+++ b/Engine/Plugins/Experimental/Dataflow/Source/DataflowNodes/Private/Dataflow/DataflowNodesConnectionTypes.cpp
@@ -0,0 +1,9 @@
+// Copyright Epic Games, Inc. All Rights Reserved.
+
+#include "Dataflow/DataflowNodesConnectionTypes.h"
+
+
+namespace Dataflow
+{
+}
+
diff --git a/Engine/Plugins/Experimental/Dataflow/Source/DataflowNodes/Private/Dataflow/DataflowNodesPlugin.cpp b/Engine/Plugins/Experimental/Dataflow/Source/DataflowNodes/Private/Dataflow/DataflowNodesPlugin.cpp
new file mode 100644
--- /dev/null
+++ b/Engine/Plugins/Experimental/Dataflow/Source/DataflowNodes/Private/Dataflow/DataflowNodesPlugin.cpp
@@ -0,0 +1,28 @@
+// Copyright Epic Games, Inc. All Rights Reserved.
+
+#include "Dataflow/DataflowNodesPlugin.h"
+
+#include "CoreMinimal.h"
+
+#include "Dataflow/DataflowSkeletalMeshNodes.h"
+#include "Dataflow/DataflowStaticMeshNodes.h"
+#include "Dataflow/DataflowNodeFactory.h"
+
+#define LOCTEXT_NAMESPACE "DataflowNodes"
+
+
+void IDataflowNodesPlugin::StartupModule()
+{
+	Dataflow::RegisterSkeletalMeshNodes();
+	Dataflow::RegisterStaticMeshNodes();
+}
+
+void IDataflowNodesPlugin::ShutdownModule()
+{
+}
+
+
+IMPLEMENT_MODULE(IDataflowNodesPlugin, DataflowNodes)
+
+
+#undef LOCTEXT_NAMESPACE
diff --git a/Engine/Plugins/Experimental/Dataflow/Source/DataflowNodes/Private/Dataflow/DataflowSkeletalMeshNodes.cpp b/Engine/Plugins/Experimental/Dataflow/Source/DataflowNodes/Private/Dataflow/DataflowSkeletalMeshNodes.cpp
new file mode 100644
--- /dev/null
+++ b/Engine/Plugins/Experimental/Dataflow/Source/DataflowNodes/Private/Dataflow/DataflowSkeletalMeshNodes.cpp
@@ -0,0 +1,14 @@
+// Copyright Epic Games, Inc. All Rights Reserved.
+
+#include "Dataflow/DataflowSkeletalMeshNodes.h"
+#include "Dataflow/DataflowCore.h"
+
+namespace Dataflow
+{
+	void RegisterSkeletalMeshNodes()
+	{
+		//DATAFLOW_NODE_REGISTER_CREATION_FACTORY(SkeletalMesh);
+		DATAFLOW_NODE_REGISTER_CREATION_FACTORY(SkeletalMeshBone);
+	}
+}
+
diff --git a/Engine/Plugins/Experimental/Dataflow/Source/DataflowNodes/Private/Dataflow/DataflowStaticMeshNodes.cpp b/Engine/Plugins/Experimental/Dataflow/Source/DataflowNodes/Private/Dataflow/DataflowStaticMeshNodes.cpp
new file mode 100644
--- /dev/null
+++ b/Engine/Plugins/Experimental/Dataflow/Source/DataflowNodes/Private/Dataflow/DataflowStaticMeshNodes.cpp
@@ -0,0 +1,12 @@
+// Copyright Epic Games, Inc. All Rights Reserved.
+
+#include "Dataflow/DataflowStaticMeshNodes.h"
+
+namespace Dataflow
+{
+	void RegisterStaticMeshNodes()
+	{
+		//DATAFLOW_NODE_REGISTER_CREATION_FACTORY(ExampleNode);
+	}
+}
+
diff --git a/Engine/Plugins/Experimental/Dataflow/Source/DataflowNodes/Public/Dataflow/DataflowNodesConnectionTypes.h b/Engine/Plugins/Experimental/Dataflow/Source/DataflowNodes/Public/Dataflow/DataflowNodesConnectionTypes.h
new file mode 100644
--- /dev/null
+++ b/Engine/Plugins/Experimental/Dataflow/Source/DataflowNodes/Public/Dataflow/DataflowNodesConnectionTypes.h
@@ -0,0 +1,14 @@
+// Copyright Epic Games, Inc. All Rights Reserved.
+
+#pragma once
+
+#include "CoreMinimal.h"
+#include "Dataflow/DataflowInputOutput.h"
+
+
+namespace Dataflow
+{
+	//typedef TSharedPtr<FManagedArrayCollection> FManagedArrayCollectionSharedPtr;
+	DATAFLOW_CONNECTION_TYPE_PRIMITIVE(DATAFLOWNODES_API, FString, String)
+	DATAFLOW_CONNECTION_TYPE_PRIMITIVE(DATAFLOWNODES_API, FVector3f, Vector3f)
+}
diff --git a/Engine/Plugins/Experimental/Dataflow/Source/DataflowNodes/Public/Dataflow/DataflowNodesPlugin.h b/Engine/Plugins/Experimental/Dataflow/Source/DataflowNodes/Public/Dataflow/DataflowNodesPlugin.h
new file mode 100644
--- /dev/null
+++ b/Engine/Plugins/Experimental/Dataflow/Source/DataflowNodes/Public/Dataflow/DataflowNodesPlugin.h
@@ -0,0 +1,18 @@
+// Copyright Epic Games, Inc. All Rights Reserved.
+
+#pragma once
+
+#include "CoreMinimal.h"
+#include "Modules/ModuleInterface.h"
+#include "Modules/ModuleManager.h"
+
+/**
+ * The public interface to this module
+ */
+class IDataflowNodesPlugin : public IModuleInterface
+{
+public:
+	virtual void StartupModule();
+	virtual void ShutdownModule();
+};
+
diff --git a/Engine/Plugins/Experimental/Dataflow/Source/DataflowNodes/Public/Dataflow/DataflowSkeletalMeshNodes.h b/Engine/Plugins/Experimental/Dataflow/Source/DataflowNodes/Public/Dataflow/DataflowSkeletalMeshNodes.h
new file mode 100644
--- /dev/null
+++ b/Engine/Plugins/Experimental/Dataflow/Source/DataflowNodes/Public/Dataflow/DataflowSkeletalMeshNodes.h
@@ -0,0 +1,109 @@
+// Copyright Epic Games, Inc. All Rights Reserved.
+
+#pragma once 
+
+#include "CoreMinimal.h"
+
+#include "Dataflow/DataflowNodesConnectionTypes.h"
+#include "Dataflow/DataflowObjectInterface.h"
+#include "Dataflow/DataflowEngine.h"
+#include "Dataflow/DataflowProperty.h"
+#include "Logging/LogMacros.h"
+#include "UObject/UnrealTypePrivate.h"
+#include "Animation/Skeleton.h"
+#include "Animation/Rig.h"
+#include "Engine/SkeletalMesh.h"
+
+DEFINE_LOG_CATEGORY_STATIC(LogDataflowSkeletalMeshNodes, Log, All);
+
+namespace Dataflow
+{
+	
+	class DATAFLOWNODES_API SkeletalMeshBone : public FNode
+	{
+		DATAFLOW_NODE_DEFINE_INTERNAL(SkeletalMeshBone)
+
+	public:
+		FGuid SkeletalMeshGuid;
+		TProperty<FName> SkeletalMeshAttributeName;
+		TProperty<FName> BoneName;
+
+		TSharedPtr< class TOutput<int> > BoneIndexOut;
+		TSharedPtr< class TOutput<FString> > BoneNameOut;
+		TSharedPtr< class TOutput<FVector3f> > BonePositionOut;
+
+		SkeletalMeshBone(const FNodeParameters& InParam, FGuid InGuid = FGuid::NewGuid())
+			: FNode(InParam, InGuid)
+			, SkeletalMeshAttributeName({ FName("SkeletalMeshAttributeName"), FName("SkeletalMesh"), this })
+			, BoneName({ FName("BoneName"), FName(""), this})
+			, BoneIndexOut(new TOutput<int>(TOutputParameters<int>({ FName("BoneIndexOut"), this })))
+			, BoneNameOut(new TOutput<FString>(TOutputParameters<FString>({ FName("BoneNameOut"), this })))
+			, BonePositionOut(new TOutput<FVector3f>(TOutputParameters<FVector3f>({ FName("BonePositionOut"), this })))
+		{}
+
+
+		template<class T>
+		const T* FindObjectPtrProperty(const UObject* Owner, FName Name) const
+		{
+			if (Owner && Owner->GetClass())
+			{
+				if (const ::FProperty* UEProperty = Owner->GetClass()->FindPropertyByName(Name))
+				{
+					if (const FObjectProperty* ObjectProperty = CastField<FObjectProperty>(UEProperty))
+					{
+						if (const void* ObjectContainer = ObjectProperty->ContainerPtrToValuePtr<void>(Owner))
+						{
+							if (const UObject* Value = ObjectProperty->GetObjectPropertyValue(ObjectContainer))
+							{
+								return Cast<T>(Value);
+							}
+						}
+					}
+				}
+			}
+			return nullptr;
+		}
+
+
+		virtual void Evaluate(const FContext& Context, FConnection* Out) const override
+		{
+			BoneIndexOut->SetValue(INDEX_NONE, Context);
+			BoneNameOut->SetValue(BoneName.GetValue().ToString(), Context);
+			BonePositionOut->SetValue(FVector3f(0, 0, 0), Context);
+
+			if (const FEngineContext* EngineContext = (const FEngineContext*)(&Context))
+			{
+				if (const USkeletalMesh* SkeletalMesh = FindObjectPtrProperty<USkeletalMesh>(
+					EngineContext->Owner, SkeletalMeshAttributeName.GetValue()))
+				{
+					if (const USkeleton* Skeleton = SkeletalMesh->GetSkeleton())
+					{
+						BoneNameOut->SetValue(BoneName.GetValue().ToString(), Context);
+						int32 Index = SkeletalMesh->GetRefSkeleton().FindBoneIndex(BoneName.GetValue());
+						BoneIndexOut->SetValue(Index, Context);
+						if (Index != INDEX_NONE)
+						{
+							const TArray<FTransform>& LocalPose = Skeleton->GetRefLocalPoses();
+							if (0 < Index && Index < LocalPose.Num())
+							{
+								BonePositionOut->SetValue(FVector3f(LocalPose[Index].GetTranslation()), Context);
+							}
+						}
+					}
+				}
+			}
+		}
+
+		virtual void SerializeInternal(FArchive& Ar) override
+		{
+			FNode::SerializeInternal(Ar);
+			Ar << SkeletalMeshGuid;
+		}
+
+	};
+	
+	void RegisterSkeletalMeshNodes();
+
+
+}
+
diff --git a/Engine/Plugins/Experimental/Dataflow/Source/DataflowNodes/Public/Dataflow/DataflowStaticMeshNodes.h b/Engine/Plugins/Experimental/Dataflow/Source/DataflowNodes/Public/Dataflow/DataflowStaticMeshNodes.h
new file mode 100644
--- /dev/null
+++ b/Engine/Plugins/Experimental/Dataflow/Source/DataflowNodes/Public/Dataflow/DataflowStaticMeshNodes.h
@@ -0,0 +1,35 @@
+// Copyright Epic Games, Inc. All Rights Reserved.
+
+#pragma once 
+
+#include "CoreMinimal.h"
+
+namespace Dataflow
+{
+
+	/*
+	class CHAOSFLESHENGINENODES_API ExampleNode : public FNode
+	{
+		DATAFLOW_NODE_DEFINE_INTERNAL(ExampleNode)
+
+	public:
+		typede ExampleDataSharedPtr DataType;
+
+		TSharedPtr< class TOutput<DataType> > Output;
+
+		ExampleNode(const FNodeParameters& InParam, FGuid InGuid = FGuid::NewGuid())
+			: FNode(InParam, InGuid)
+			, Output(new TOutput<DataType>(TOutputParameters<DataType>({ FName("ExampleOut"), this })))
+		{}
+
+		virtual void Evaluate(const FContext& Context, FConnection* Out) const override
+		{
+			Output->(<set example data>, Context);
+		}
+	};
+	*/
+
+	void RegisterStaticMeshNodes();
+
+}
+
diff --git a/Engine/Plugins/Experimental/GeometryCollectionPlugin/Source/GeometryCollectionEditor/Private/GeometryCollection/Dataflow/GeometryCollectionEditorToolkit.cpp b/Engine/Plugins/Experimental/GeometryCollectionPlugin/Source/GeometryCollectionEditor/Private/GeometryCollection/Dataflow/GeometryCollectionEditorToolkit.cpp
--- a/Engine/Plugins/Experimental/GeometryCollectionPlugin/Source/GeometryCollectionEditor/Private/GeometryCollection/Dataflow/GeometryCollectionEditorToolkit.cpp
+++ b/Engine/Plugins/Experimental/GeometryCollectionPlugin/Source/GeometryCollectionEditor/Private/GeometryCollection/Dataflow/GeometryCollectionEditorToolkit.cpp
@@ -8,6 +8,7 @@
 #include "Dataflow/DataflowGraphEditor.h"
 #include "Dataflow/DataflowNodeFactory.h"
 #include "Dataflow/DataflowObject.h"
+#include "Dataflow/DataflowObjectInterface.h"
 #include "Dataflow/DataflowSchema.h"
 #include "Dataflow/DataflowCore.h"
 #include "EditorStyleSet.h"
@@ -93,7 +94,7 @@ TSharedRef<SGraphEditor> FGeometryCollectionEditorToolkit::CreateGraphEditorWidg
 	FDataflowEditorCommands::FGraphEvaluationCallback Evaluate = [&](Dataflow::FNode* Node, Dataflow::FConnection* Out)
 	{
 		float EvalTime = FGameTime::GetTimeSinceAppStart().GetRealTimeSeconds();
-		return Node->Evaluate(FGeometryCollectionContext(GeometryCollection, EvalTime), Out);
+		return Node->Evaluate(FEngineContext(GeometryCollection, Dataflow, EvalTime, FName("UGeometryCollection")), Out);
 	};
 
 	return SNew(SDataflowGraphEditor, GeometryCollection)
diff --git a/Engine/Plugins/Experimental/GeometryCollectionPlugin/Source/GeometryCollectionEditor/Private/GeometryCollection/Dataflow/GeometryCollectionNodes.h b/Engine/Plugins/Experimental/GeometryCollectionPlugin/Source/GeometryCollectionEditor/Private/GeometryCollection/Dataflow/GeometryCollectionNodes.h
--- a/Engine/Plugins/Experimental/GeometryCollectionPlugin/Source/GeometryCollectionEditor/Private/GeometryCollection/Dataflow/GeometryCollectionNodes.h
+++ b/Engine/Plugins/Experimental/GeometryCollectionPlugin/Source/GeometryCollectionEditor/Private/GeometryCollection/Dataflow/GeometryCollectionNodes.h
@@ -35,13 +35,15 @@ namespace Dataflow
 		{
 			if (Output.Get() == Out)
 			{
-				const FGeometryCollectionContext* EngineContext = (const FGeometryCollectionContext*)(&Context);
-				if (UGeometryCollection* CollectionAsset = EngineContext->Asset)
+				if (const FEngineContext* EngineContext = Context.AsType<FEngineContext>(FName("UGeometryCollection")))
 				{
-					if (const TSharedPtr<FGeometryCollection, ESPMode::ThreadSafe> AssetCollection = CollectionAsset->GetGeometryCollection())
+					if (UGeometryCollection* CollectionAsset = Cast<UGeometryCollection>(EngineContext->Owner))
 					{
-						FManagedArrayCollection* NewCollection = AssetCollection->NewCopy<FManagedArrayCollection>();
-						Output->SetValue(FManagedArrayCollectionSharedPtr(NewCollection), Context);
+						if (const TSharedPtr<FGeometryCollection, ESPMode::ThreadSafe> AssetCollection = CollectionAsset->GetGeometryCollection())
+						{
+							FManagedArrayCollection* NewCollection = AssetCollection->NewCopy<FManagedArrayCollection>();
+							Output->SetValue(FManagedArrayCollectionSharedPtr(NewCollection), Context);
+						}
 					}
 				}
 			}
@@ -101,14 +103,16 @@ namespace Dataflow
 
 		virtual void Evaluate(const FContext& Context, FConnection* Out) const override
 		{
-			if (Out==nullptr)
+			if (Out == nullptr)
 			{
 				DataType Collection = Input->GetValue(Context);
-				const FGeometryCollectionContext* EngineContext = (FGeometryCollectionContext*)(&Context);
-				if (UGeometryCollection* CollectionAsset = EngineContext->Asset)
+				if (const FEngineContext* EngineContext = Context.AsType<FEngineContext>(FName("UGeometryCollection")))
 				{
-					TSharedPtr<FGeometryCollection, ESPMode::ThreadSafe> NewCollection(Collection->NewCopy<FGeometryCollection>());
-					CollectionAsset->SetGeometryCollection(NewCollection);
+					if (UGeometryCollection* CollectionAsset = Cast<UGeometryCollection>(EngineContext->Owner))
+					{
+						TSharedPtr<FGeometryCollection, ESPMode::ThreadSafe> NewCollection(Collection->NewCopy<FGeometryCollection>());
+						CollectionAsset->SetGeometryCollection(NewCollection);
+					}
 				}
 			}
 		}
diff --git a/Engine/Plugins/Experimental/GeometryCollectionPlugin/Source/GeometryCollectionEditor/Public/GeometryCollection/Dataflow/GeometryCollectionEditorToolkit.h b/Engine/Plugins/Experimental/GeometryCollectionPlugin/Source/GeometryCollectionEditor/Public/GeometryCollection/Dataflow/GeometryCollectionEditorToolkit.h
--- a/Engine/Plugins/Experimental/GeometryCollectionPlugin/Source/GeometryCollectionEditor/Public/GeometryCollection/Dataflow/GeometryCollectionEditorToolkit.h
+++ b/Engine/Plugins/Experimental/GeometryCollectionPlugin/Source/GeometryCollectionEditor/Public/GeometryCollection/Dataflow/GeometryCollectionEditorToolkit.h
@@ -15,20 +15,6 @@ class FTabManager;
 class IToolkitHost;
 class UDataflow;
 
-namespace Dataflow
-{
-	class GEOMETRYCOLLECTIONEDITOR_API FGeometryCollectionContext : public FContext
-	{
-	public:
-		UGeometryCollection* Asset;
-
-		FGeometryCollectionContext(UGeometryCollection* InAsset, float InTime)
-			: FContext(InTime)
-			, Asset(InAsset)
-		{}
-
-	};
-}
 
 class FGeometryCollectionEditorToolkit : public FAssetEditorToolkit, public FNotifyHook, public FGCObject
 {
diff --git a/Engine/Source/Runtime/Experimental/Dataflow/Core/Public/Dataflow/DataflowNodeParameters.h b/Engine/Source/Runtime/Experimental/Dataflow/Core/Public/Dataflow/DataflowNodeParameters.h
--- a/Engine/Source/Runtime/Experimental/Dataflow/Core/Public/Dataflow/DataflowNodeParameters.h
+++ b/Engine/Source/Runtime/Experimental/Dataflow/Core/Public/Dataflow/DataflowNodeParameters.h
@@ -11,14 +11,27 @@ namespace Dataflow
 	class DATAFLOWCORE_API FContext
 	{
 	public:
-		FContext(float InTime) : Timestamp(InTime) {}
+		FContext(float InTime, FName InType = FName("Unknown"))
+		: Timestamp(InTime) 
+		, Type(InType){}
 
 		float Timestamp = 0.f;
+		FName Type = FName("Unknown");
 
 		uint32 GetTypeHash() const
 		{
 			return ::GetTypeHash(Timestamp);
 		}
+
+		template<class T>
+		const T* AsType(FName InType) const
+		{
+			if (Type.IsEqual(InType))
+			{
+				return (T*)this;
+			}
+			return nullptr;
+		}
 	};
 
 
diff --git a/Engine/Source/Runtime/Experimental/Dataflow/Engine/Private/Dataflow/DataflowObjectInterface.cpp b/Engine/Source/Runtime/Experimental/Dataflow/Engine/Private/Dataflow/DataflowObjectInterface.cpp
--- a/Engine/Source/Runtime/Experimental/Dataflow/Engine/Private/Dataflow/DataflowObjectInterface.cpp
+++ b/Engine/Source/Runtime/Experimental/Dataflow/Engine/Private/Dataflow/DataflowObjectInterface.cpp
@@ -4,8 +4,9 @@
 
 namespace Dataflow
 {
-	FEngineContext::FEngineContext(UDataflow* InGraph, float InTime)
-			: FContext(InTime)
-			, Graph(InGraph) 
+	FEngineContext::FEngineContext(UObject* InOwner, UDataflow* InGraph, float InTime, FName InType)
+			: FContext(InTime, InType)
+			, Owner(InOwner)
+			, Graph(InGraph)
 		{}
 }
\ No newline at end of file
diff --git a/Engine/Source/Runtime/Experimental/Dataflow/Engine/Public/Dataflow/DataflowObjectInterface.h b/Engine/Source/Runtime/Experimental/Dataflow/Engine/Public/Dataflow/DataflowObjectInterface.h
--- a/Engine/Source/Runtime/Experimental/Dataflow/Engine/Public/Dataflow/DataflowObjectInterface.h
+++ b/Engine/Source/Runtime/Experimental/Dataflow/Engine/Public/Dataflow/DataflowObjectInterface.h
@@ -13,7 +13,9 @@ namespace Dataflow
 	class DATAFLOWENGINE_API FEngineContext : public FContext
 	{
 	public:
-		FEngineContext(UDataflow* InGraph, float InTime);
+		FEngineContext(UObject* Owner, UDataflow* InGraph, float InTime, FName InType=FName("Unknown"));
+
+		UObject* Owner = nullptr;
 		UDataflow* Graph = nullptr;
 	};
 }
\ No newline at end of file
