Dataflow : FName property
- Added FName support for Dataflow::TProperty

diff --git a/Engine/Plugins/Experimental/Dataflow/Source/DataflowEditor/Private/Dataflow/DataflowEPropertyCustomizations.cpp b/Engine/Plugins/Experimental/Dataflow/Source/DataflowEditor/Private/Dataflow/DataflowEPropertyCustomizations.cpp
--- a/Engine/Plugins/Experimental/Dataflow/Source/DataflowEditor/Private/Dataflow/DataflowEPropertyCustomizations.cpp
+++ b/Engine/Plugins/Experimental/Dataflow/Source/DataflowEditor/Private/Dataflow/DataflowEPropertyCustomizations.cpp
@@ -56,6 +56,9 @@
 					case Dataflow::FProperty::EType::STRING:
 						Dataflow::PropertyWidgetFactory<FString>(DetailBuilder, PropertyHandle, Node, (Dataflow::TProperty<FString>*)Property);
 						break;
+					case Dataflow::FProperty::EType::NAME:
+						Dataflow::PropertyWidgetFactory<FName>(DetailBuilder, PropertyHandle, Node, (Dataflow::TProperty<FName>*)Property);
+						break;
 					default:
 						ensureMsgf(true,TEXT("Missing slate property convert."));
 					}
diff --git a/Engine/Plugins/Experimental/Dataflow/Source/DataflowEditor/Public/Dataflow/DataflowEPropertyFactories.h b/Engine/Plugins/Experimental/Dataflow/Source/DataflowEditor/Public/Dataflow/DataflowEPropertyFactories.h
--- a/Engine/Plugins/Experimental/Dataflow/Source/DataflowEditor/Public/Dataflow/DataflowEPropertyFactories.h
+++ b/Engine/Plugins/Experimental/Dataflow/Source/DataflowEditor/Public/Dataflow/DataflowEPropertyFactories.h
@@ -107,4 +107,33 @@
 			];
 	}
 
+	//
+	// FName
+	//
+	template<> inline
+		void PropertyWidgetFactory<FName>(
+			IDetailLayoutBuilder& InDetailBuilder,
+			TSharedPtr<IPropertyHandle> PropertyHandle,
+			TSharedPtr<Dataflow::FNode> InNode,
+			TProperty<FName>* InProperty)
+	{
+		InDetailBuilder.AddCustomRowToCategory(PropertyHandle, FText::FromString("PropertyHandel<FName>"))
+			.NameContent()[
+				SNew(STextBlock).Font(IDetailLayoutBuilder::GetDetailFont())
+					.Text(FText::FromName(InProperty->GetName()))
+			].ValueContent()[
+				SNew(SEditableTextBox)
+					.Font(IDetailLayoutBuilder::GetDetailFont())
+					.Text_Lambda([InNode, InProperty]()
+						{
+							return FText::FromString(InProperty->GetValue().ToString());
+						})
+					.OnTextCommitted_Lambda([InNode, InProperty](const FText& InText, ETextCommit::Type TextCommitType)
+						{
+							InProperty->SetValue(FName(InText.ToString()));
+						})
+
+			];
+	}
+
 }
diff --git a/Engine/Source/Runtime/Experimental/Dataflow/Core/Public/Dataflow/DataflowProperty.h b/Engine/Source/Runtime/Experimental/Dataflow/Core/Public/Dataflow/DataflowProperty.h
--- a/Engine/Source/Runtime/Experimental/Dataflow/Core/Public/Dataflow/DataflowProperty.h
+++ b/Engine/Source/Runtime/Experimental/Dataflow/Core/Public/Dataflow/DataflowProperty.h
@@ -30,7 +30,7 @@
 			FLOAT,
 			DOUBLE,
 			STRING,
-			FILE,
+			NAME,
 			// <new types go here>
 			MAX
 		};
@@ -100,6 +100,7 @@
 DATAFLOW_PROPERTY(DATAFLOWCORE_API, float, FLOAT)
 DATAFLOW_PROPERTY(DATAFLOWCORE_API, double, DOUBLE)
 DATAFLOW_PROPERTY(DATAFLOWCORE_API, FString, STRING)
+DATAFLOW_PROPERTY(DATAFLOWCORE_API, FName, NAME)
 //DATAFLOW_PROPERTY(DATAFLOWCORE_API, Dataflow::FFilename, FILE)
 
 
