Dataflow : Cleanup

diff --git a/Engine/Plugins/Experimental/Dataflow/Source/DataflowNodes/Public/Dataflow/DataflowSkeletalMeshNodes.h b/Engine/Plugins/Experimental/Dataflow/Source/DataflowNodes/Public/Dataflow/DataflowSkeletalMeshNodes.h
--- a/Engine/Plugins/Experimental/Dataflow/Source/DataflowNodes/Public/Dataflow/DataflowSkeletalMeshNodes.h
+++ b/Engine/Plugins/Experimental/Dataflow/Source/DataflowNodes/Public/Dataflow/DataflowSkeletalMeshNodes.h
@@ -7,6 +7,7 @@
 #include "Dataflow/DataflowNodesConnectionTypes.h"
 #include "Dataflow/DataflowObjectInterface.h"
 #include "Dataflow/DataflowEngine.h"
+#include "Dataflow/DataflowEngineUtil.h"
 #include "Dataflow/DataflowProperty.h"
 #include "Logging/LogMacros.h"
 #include "UObject/UnrealTypePrivate.h"
@@ -24,7 +25,6 @@ namespace Dataflow
 		DATAFLOW_NODE_DEFINE_INTERNAL(SkeletalMeshBone)
 
 	public:
-		FGuid SkeletalMeshGuid;
 		TProperty<FName> SkeletalMeshAttributeName;
 		TProperty<FName> BoneName;
 
@@ -42,29 +42,6 @@ namespace Dataflow
 		{}
 
 
-		template<class T>
-		const T* FindObjectPtrProperty(const UObject* Owner, FName Name) const
-		{
-			if (Owner && Owner->GetClass())
-			{
-				if (const ::FProperty* UEProperty = Owner->GetClass()->FindPropertyByName(Name))
-				{
-					if (const FObjectProperty* ObjectProperty = CastField<FObjectProperty>(UEProperty))
-					{
-						if (const void* ObjectContainer = ObjectProperty->ContainerPtrToValuePtr<void>(Owner))
-						{
-							if (const UObject* Value = ObjectProperty->GetObjectPropertyValue(ObjectContainer))
-							{
-								return Cast<T>(Value);
-							}
-						}
-					}
-				}
-			}
-			return nullptr;
-		}
-
-
 		virtual void Evaluate(const FContext& Context, FConnection* Out) const override
 		{
 			BoneIndexOut->SetValue(INDEX_NONE, Context);
@@ -73,7 +50,7 @@ namespace Dataflow
 
 			if (const FEngineContext* EngineContext = (const FEngineContext*)(&Context))
 			{
-				if (const USkeletalMesh* SkeletalMesh = FindObjectPtrProperty<USkeletalMesh>(
+				if (const USkeletalMesh* SkeletalMesh = Reflection::FindObjectPtrProperty<USkeletalMesh>(
 					EngineContext->Owner, SkeletalMeshAttributeName.GetValue()))
 				{
 					if (const USkeleton* Skeleton = SkeletalMesh->GetSkeleton())
@@ -83,23 +60,17 @@ namespace Dataflow
 						BoneIndexOut->SetValue(Index, Context);
 						if (Index != INDEX_NONE)
 						{
-							const TArray<FTransform>& LocalPose = Skeleton->GetRefLocalPoses();
-							if (0 < Index && Index < LocalPose.Num())
+							TArray<FTransform> ComponentPose; 
+							Animation::GlobalTransforms(SkeletalMesh->GetRefSkeleton(), ComponentPose);
+							if (0 < Index && Index < ComponentPose.Num())
 							{
-								BonePositionOut->SetValue(FVector3f(LocalPose[Index].GetTranslation()), Context);
+								BonePositionOut->SetValue(FVector3f(ComponentPose[Index].GetTranslation()), Context);
 							}
 						}
 					}
 				}
 			}
 		}
-
-		virtual void SerializeInternal(FArchive& Ar) override
-		{
-			FNode::SerializeInternal(Ar);
-			Ar << SkeletalMeshGuid;
-		}
-
 	};
 	
 	void RegisterSkeletalMeshNodes();
diff --git a/Engine/Source/Runtime/Experimental/Dataflow/Engine/Private/Dataflow/DataflowEngineUtil.cpp b/Engine/Source/Runtime/Experimental/Dataflow/Engine/Private/Dataflow/DataflowEngineUtil.cpp
new file mode 100644
--- /dev/null
+++ b/Engine/Source/Runtime/Experimental/Dataflow/Engine/Private/Dataflow/DataflowEngineUtil.cpp
@@ -0,0 +1,42 @@
+// Copyright Epic Games, Inc. All Rights Reserved.
+
+#include "Dataflow/DataflowEngineUtil.h"
+#include "ReferenceSkeleton.h"
+
+namespace Dataflow
+{
+
+	namespace Animation
+	{
+		void GlobalTransformsInternal(int32 Index, const FReferenceSkeleton& Ref, TArray<FTransform>& Mat, TArray<bool>& Visited)
+		{
+			const TArray<FTransform>& RefMat = Ref.GetRefBonePose();
+
+			int32 ParentIndex = Ref.GetParentIndex(Index);
+			if (ParentIndex != INDEX_NONE && ParentIndex!=Index) // why self check?
+			{
+				GlobalTransformsInternal(ParentIndex, Ref, Mat, Visited);
+				Mat[Index] = Mat[ParentIndex].GetRelativeTransform(RefMat[Index]);
+			}
+			else
+			{
+				Mat[Index] = RefMat[Index];
+			}
+			Visited[Index] = true;
+		}
+
+		void GlobalTransforms(const FReferenceSkeleton& Ref, TArray<FTransform>& Mat)
+		{
+			TArray<bool> Visited;
+			Visited.Init(false, Ref.GetNum());
+			Mat.SetNum(Ref.GetNum());
+
+			int32 Index = Ref.GetNum() - 1;
+			while (Index >= 0 && !Visited[Index])
+			{
+				GlobalTransformsInternal(Index, Ref, Mat, Visited);
+				Index--;
+			}
+		}
+	}
+}
\ No newline at end of file
diff --git a/Engine/Source/Runtime/Experimental/Dataflow/Engine/Public/Dataflow/DataflowEngineUtil.h b/Engine/Source/Runtime/Experimental/Dataflow/Engine/Public/Dataflow/DataflowEngineUtil.h
new file mode 100644
--- /dev/null
+++ b/Engine/Source/Runtime/Experimental/Dataflow/Engine/Public/Dataflow/DataflowEngineUtil.h
@@ -0,0 +1,39 @@
+// Copyright Epic Games, Inc. All Rights Reserved.
+
+#pragma once
+
+#include "CoreMinimal.h"
+#include "UObject/UnrealTypePrivate.h"
+#include "ReferenceSkeleton.h"
+namespace Dataflow
+{
+	namespace Reflection
+	{
+		template<class T>
+		const T* FindObjectPtrProperty(const UObject* Owner, FName Name)
+		{
+			if (Owner && Owner->GetClass())
+			{
+				if (const ::FProperty* UEProperty = Owner->GetClass()->FindPropertyByName(Name))
+				{
+					if (const FObjectProperty* ObjectProperty = CastField<FObjectProperty>(UEProperty))
+					{
+						if (const void* ObjectContainer = ObjectProperty->ContainerPtrToValuePtr<void>(Owner))
+						{
+							if (const UObject* Value = ObjectProperty->GetObjectPropertyValue(ObjectContainer))
+							{
+								return Cast<T>(Value);
+							}
+						}
+					}
+				}
+			}
+			return nullptr;
+		}
+	}
+
+	namespace Animation
+	{
+		void DATAFLOWENGINE_API GlobalTransforms(const FReferenceSkeleton&, TArray<FTransform>& GlobalTransforms);
+	}
+}
\ No newline at end of file
