ChaosFlesh : Asset
- Automatic invalidation on reset.
- Disable solver when time steps are zero.

diff --git a/Engine/Plugins/Experimental/Dataflow/Source/DataflowNodes/Private/Dataflow/DataflowSkeletalMeshNodes.cpp b/Engine/Plugins/Experimental/Dataflow/Source/DataflowNodes/Private/Dataflow/DataflowSkeletalMeshNodes.cpp
--- a/Engine/Plugins/Experimental/Dataflow/Source/DataflowNodes/Private/Dataflow/DataflowSkeletalMeshNodes.cpp
+++ b/Engine/Plugins/Experimental/Dataflow/Source/DataflowNodes/Private/Dataflow/DataflowSkeletalMeshNodes.cpp
@@ -2,13 +2,43 @@
 
 #include "Dataflow/DataflowSkeletalMeshNodes.h"
 #include "Dataflow/DataflowCore.h"
+#include "Dataflow/DataflowEngine.h"
+#include "Logging/LogMacros.h"
+#include "UObject/UnrealTypePrivate.h"
+#include "Animation/Skeleton.h"
+#include "Animation/Rig.h"
+#include "Engine/SkeletalMesh.h"
 
 namespace Dataflow
 {
 	void RegisterSkeletalMeshNodes()
 	{
-		//DATAFLOW_NODE_REGISTER_CREATION_FACTORY(SkeletalMesh);
+		DATAFLOW_NODE_REGISTER_CREATION_FACTORY(GetSkeletalMesh);
 		DATAFLOW_NODE_REGISTER_CREATION_FACTORY(SkeletalMeshBone);
 	}
+
+	void GetSkeletalMesh::Evaluate(const FContext& Context, FConnection* Out) const
+	{
+		SkeletalMeshOut->SetValue(nullptr, Context);
+		if (const FEngineContext* EngineContext = (const FEngineContext*)(&Context))
+		{
+			if (const USkeletalMesh* SkeletalMesh = Reflection::FindObjectPtrProperty<USkeletalMesh>(
+				EngineContext->Owner, SkeletalMeshAttributeName.GetValue()))
+			{
+				SkeletalMeshOut->SetValue(SkeletalMesh, Context);
+			}
+		}
+	}
+
+	void SkeletalMeshBone::Evaluate(const FContext& Context, FConnection* Out) const
+	{
+		BoneIndexOut->SetValue(INDEX_NONE, Context);
+
+		if (const USkeletalMesh* SkeletalMesh = SkeletalMeshIn->GetValue(Context))
+		{
+			int32 Index = SkeletalMesh->GetRefSkeleton().FindBoneIndex(BoneName.GetValue());
+			BoneIndexOut->SetValue(Index, Context);
+		}
+	}
 }
 
diff --git a/Engine/Plugins/Experimental/Dataflow/Source/DataflowNodes/Public/Dataflow/DataflowNodesConnectionTypes.h b/Engine/Plugins/Experimental/Dataflow/Source/DataflowNodes/Public/Dataflow/DataflowNodesConnectionTypes.h
--- a/Engine/Plugins/Experimental/Dataflow/Source/DataflowNodes/Public/Dataflow/DataflowNodesConnectionTypes.h
+++ b/Engine/Plugins/Experimental/Dataflow/Source/DataflowNodes/Public/Dataflow/DataflowNodesConnectionTypes.h
@@ -5,10 +5,13 @@
 #include "CoreMinimal.h"
 #include "Dataflow/DataflowInputOutput.h"
 
+class USkeletalMesh;
 
 namespace Dataflow
 {
+	typedef const USkeletalMesh* USkeletalMeshPtr;	
 	//typedef TSharedPtr<FManagedArrayCollection> FManagedArrayCollectionSharedPtr;
-	DATAFLOW_CONNECTION_TYPE_PRIMITIVE(DATAFLOWNODES_API, FString, String)
-	DATAFLOW_CONNECTION_TYPE_PRIMITIVE(DATAFLOWNODES_API, FVector3f, Vector3f)
+	DATAFLOW_CONNECTION_TYPE_PRIMITIVE(DATAFLOWNODES_API, USkeletalMeshPtr, SkeletalMeshPtr);
+	DATAFLOW_CONNECTION_TYPE_PRIMITIVE(DATAFLOWNODES_API, FString, String);
+	DATAFLOW_CONNECTION_TYPE_PRIMITIVE(DATAFLOWNODES_API, FVector3f, Vector3f);
 }
diff --git a/Engine/Plugins/Experimental/Dataflow/Source/DataflowNodes/Public/Dataflow/DataflowSkeletalMeshNodes.h b/Engine/Plugins/Experimental/Dataflow/Source/DataflowNodes/Public/Dataflow/DataflowSkeletalMeshNodes.h
--- a/Engine/Plugins/Experimental/Dataflow/Source/DataflowNodes/Public/Dataflow/DataflowSkeletalMeshNodes.h
+++ b/Engine/Plugins/Experimental/Dataflow/Source/DataflowNodes/Public/Dataflow/DataflowSkeletalMeshNodes.h
@@ -5,72 +5,54 @@
 #include "CoreMinimal.h"
 
 #include "Dataflow/DataflowNodesConnectionTypes.h"
-#include "Dataflow/DataflowObjectInterface.h"
+#include "Dataflow/DataflowCore.h"
 #include "Dataflow/DataflowEngine.h"
-#include "Dataflow/DataflowEngineUtil.h"
-#include "Dataflow/DataflowProperty.h"
-#include "Logging/LogMacros.h"
-#include "UObject/UnrealTypePrivate.h"
-#include "Animation/Skeleton.h"
-#include "Animation/Rig.h"
-#include "Engine/SkeletalMesh.h"
 
 DEFINE_LOG_CATEGORY_STATIC(LogDataflowSkeletalMeshNodes, Log, All);
 
+class USkeletalMesh;
+
 namespace Dataflow
 {
+
+	class DATAFLOWNODES_API GetSkeletalMesh : public FNode
+	{
+		DATAFLOW_NODE_DEFINE_INTERNAL(GetSkeletalMesh)
+
+	public:
+		TProperty<FName> SkeletalMeshAttributeName;
+
+		TSharedPtr< class TOutput<USkeletalMeshPtr> > SkeletalMeshOut;
+
+		GetSkeletalMesh(const FNodeParameters& InParam, FGuid InGuid = FGuid::NewGuid())
+			: FNode(InParam, InGuid)
+			, SkeletalMeshAttributeName({ FName("SkeletalMeshAttributeName"), FName("SkeletalMesh"), this })
+			, SkeletalMeshOut(new TOutput<USkeletalMeshPtr>(TOutputParameters<USkeletalMeshPtr>({ FName("SkeletalMeshOut"), this })))
+		{}
+
+
+		virtual void Evaluate(const FContext& Context, FConnection* Out) const override;
+	};
 	
 	class DATAFLOWNODES_API SkeletalMeshBone : public FNode
 	{
 		DATAFLOW_NODE_DEFINE_INTERNAL(SkeletalMeshBone)
 
 	public:
-		TProperty<FName> SkeletalMeshAttributeName;
 		TProperty<FName> BoneName;
 
+		TSharedPtr< class TInput<USkeletalMeshPtr> > SkeletalMeshIn;
 		TSharedPtr< class TOutput<int> > BoneIndexOut;
-		TSharedPtr< class TOutput<FString> > BoneNameOut;
-		TSharedPtr< class TOutput<FVector3f> > BonePositionOut;
 
 		SkeletalMeshBone(const FNodeParameters& InParam, FGuid InGuid = FGuid::NewGuid())
 			: FNode(InParam, InGuid)
-			, SkeletalMeshAttributeName({ FName("SkeletalMeshAttributeName"), FName("SkeletalMesh"), this })
 			, BoneName({ FName("BoneName"), FName(""), this})
+			, SkeletalMeshIn(new TInput<USkeletalMeshPtr>(TInputParameters<USkeletalMeshPtr>({ FName("SkeletalMeshIn"), this })))
 			, BoneIndexOut(new TOutput<int>(TOutputParameters<int>({ FName("BoneIndexOut"), this })))
-			, BoneNameOut(new TOutput<FString>(TOutputParameters<FString>({ FName("BoneNameOut"), this })))
-			, BonePositionOut(new TOutput<FVector3f>(TOutputParameters<FVector3f>({ FName("BonePositionOut"), this })))
 		{}
 
 
-		virtual void Evaluate(const FContext& Context, FConnection* Out) const override
-		{
-			BoneIndexOut->SetValue(INDEX_NONE, Context);
-			BoneNameOut->SetValue(BoneName.GetValue().ToString(), Context);
-			BonePositionOut->SetValue(FVector3f(0, 0, 0), Context);
-
-			if (const FEngineContext* EngineContext = (const FEngineContext*)(&Context))
-			{
-				if (const USkeletalMesh* SkeletalMesh = Reflection::FindObjectPtrProperty<USkeletalMesh>(
-					EngineContext->Owner, SkeletalMeshAttributeName.GetValue()))
-				{
-					if (const USkeleton* Skeleton = SkeletalMesh->GetSkeleton())
-					{
-						BoneNameOut->SetValue(BoneName.GetValue().ToString(), Context);
-						int32 Index = SkeletalMesh->GetRefSkeleton().FindBoneIndex(BoneName.GetValue());
-						BoneIndexOut->SetValue(Index, Context);
-						if (Index != INDEX_NONE)
-						{
-							TArray<FTransform> ComponentPose; 
-							Animation::GlobalTransforms(SkeletalMesh->GetRefSkeleton(), ComponentPose);
-							if (0 < Index && Index < ComponentPose.Num())
-							{
-								BonePositionOut->SetValue(FVector3f(ComponentPose[Index].GetTranslation()), Context);
-							}
-						}
-					}
-				}
-			}
-		}
+		virtual void Evaluate(const FContext& Context, FConnection* Out) const override;
 	};
 	
 	void RegisterSkeletalMeshNodes();
