Datafflow : Invalidation
- Added modified response.

diff --git a/Engine/Plugins/Experimental/Dataflow/Source/DataflowEditor/Private/Dataflow/DataflowEPropertyCustomizations.cpp b/Engine/Plugins/Experimental/Dataflow/Source/DataflowEditor/Private/Dataflow/DataflowEPropertyCustomizations.cpp
--- a/Engine/Plugins/Experimental/Dataflow/Source/DataflowEditor/Private/Dataflow/DataflowEPropertyCustomizations.cpp
+++ b/Engine/Plugins/Experimental/Dataflow/Source/DataflowEditor/Private/Dataflow/DataflowEPropertyCustomizations.cpp
@@ -34,35 +34,38 @@ void FDataflowSEditorCustomization::CustomizeDetails(IDetailLayoutBuilder& Detai
 		{
 			if (TSharedPtr<Dataflow::FNode> Node = DataflowSObj->Node)
 			{
-				for (Dataflow::FProperty* Property : Node->GetProperties())
+				if (UDataflow* Graph = DataflowSObj->Graph)
 				{
-					IDetailPropertyRow& PropertyRow = Category.AddProperty(FName("CustomProperty"));
-					TSharedPtr<IPropertyHandle> PropertyHandle = PropertyRow.GetPropertyHandle();
-
-					switch (Property->GetType())
+					for (Dataflow::FProperty* Property : Node->GetProperties())
 					{
-					case Dataflow::FProperty::EType::BOOL:
-						Dataflow::PropertyWidgetFactory<bool>(DetailBuilder, PropertyHandle, Node, (Dataflow::TProperty<bool>*)Property);
-						break;
-					case Dataflow::FProperty::EType::INT:
-						Dataflow::PropertyWidgetFactory<int>(DetailBuilder, PropertyHandle, Node, (Dataflow::TProperty<int>*)Property);
-						break;
-					case Dataflow::FProperty::EType::FLOAT:
-						Dataflow::PropertyWidgetFactory<float>(DetailBuilder, PropertyHandle, Node, (Dataflow::TProperty<float>*)Property);
-						break;
-					case Dataflow::FProperty::EType::DOUBLE:
-						Dataflow::PropertyWidgetFactory<double>(DetailBuilder, PropertyHandle, Node, (Dataflow::TProperty<double>*)Property);
-						break;
-					case Dataflow::FProperty::EType::STRING:
-						Dataflow::PropertyWidgetFactory<FString>(DetailBuilder, PropertyHandle, Node, (Dataflow::TProperty<FString>*)Property);
-						break;
-					case Dataflow::FProperty::EType::NAME:
-						Dataflow::PropertyWidgetFactory<FName>(DetailBuilder, PropertyHandle, Node, (Dataflow::TProperty<FName>*)Property);
-						break;
-					default:
-						ensureMsgf(true,TEXT("Missing slate property convert."));
+						IDetailPropertyRow& PropertyRow = Category.AddProperty(FName("CustomProperty"));
+						TSharedPtr<IPropertyHandle> PropertyHandle = PropertyRow.GetPropertyHandle();
+
+						switch (Property->GetType())
+						{
+						case Dataflow::FProperty::EType::BOOL:
+							Dataflow::PropertyWidgetFactory<bool>(DetailBuilder, PropertyHandle, Graph, Node, (Dataflow::TProperty<bool>*)Property);
+							break;
+						case Dataflow::FProperty::EType::INT:
+							Dataflow::PropertyWidgetFactory<int>(DetailBuilder, PropertyHandle, Graph, Node, (Dataflow::TProperty<int>*)Property);
+							break;
+						case Dataflow::FProperty::EType::FLOAT:
+							Dataflow::PropertyWidgetFactory<float>(DetailBuilder, PropertyHandle, Graph, Node, (Dataflow::TProperty<float>*)Property);
+							break;
+						case Dataflow::FProperty::EType::DOUBLE:
+							Dataflow::PropertyWidgetFactory<double>(DetailBuilder, PropertyHandle, Graph, Node, (Dataflow::TProperty<double>*)Property);
+							break;
+						case Dataflow::FProperty::EType::STRING:
+							Dataflow::PropertyWidgetFactory<FString>(DetailBuilder, PropertyHandle, Graph, Node, (Dataflow::TProperty<FString>*)Property);
+							break;
+						case Dataflow::FProperty::EType::NAME:
+							Dataflow::PropertyWidgetFactory<FName>(DetailBuilder, PropertyHandle, Graph, Node, (Dataflow::TProperty<FName>*)Property);
+							break;
+						default:
+							ensureMsgf(true, TEXT("Missing slate property convert."));
+						}
+						UE_LOG(LogCDataflowEProprtyCustomizations, Verbose, TEXT("FFloatObjectDetails::CustomizeDetails"));
 					}
-					UE_LOG(LogCDataflowEProprtyCustomizations, Verbose, TEXT("FFloatObjectDetails::CustomizeDetails"));
 				}
 			}
 		}
diff --git a/Engine/Plugins/Experimental/Dataflow/Source/DataflowEditor/Private/Dataflow/DataflowEditorActions.cpp b/Engine/Plugins/Experimental/Dataflow/Source/DataflowEditor/Private/Dataflow/DataflowEditorActions.cpp
--- a/Engine/Plugins/Experimental/Dataflow/Source/DataflowEditor/Private/Dataflow/DataflowEditorActions.cpp
+++ b/Engine/Plugins/Experimental/Dataflow/Source/DataflowEditor/Private/Dataflow/DataflowEditorActions.cpp
@@ -115,6 +115,7 @@ void FDataflowEditorCommands::OnSelectedNodesChanged(TSharedPtr<IDetailsView> Pr
 							const FName NodeName = MakeUniqueObjectName(Graph, UDataflow::StaticClass(), DataflowNode->GetName());
 							UDataflowSEditorObject* Object = NewObject<UDataflowSEditorObject>(Asset, NodeName);
 							Object->Node = DataflowNode;
+							Object->Graph = Graph;
 							Objects.Add(Object);
 						}
 					}
diff --git a/Engine/Plugins/Experimental/Dataflow/Source/DataflowEditor/Public/Dataflow/DataflowEPropertyCustomizations.h b/Engine/Plugins/Experimental/Dataflow/Source/DataflowEditor/Public/Dataflow/DataflowEPropertyCustomizations.h
--- a/Engine/Plugins/Experimental/Dataflow/Source/DataflowEditor/Public/Dataflow/DataflowEPropertyCustomizations.h
+++ b/Engine/Plugins/Experimental/Dataflow/Source/DataflowEditor/Public/Dataflow/DataflowEPropertyCustomizations.h
@@ -12,14 +12,15 @@
 namespace Dataflow {
 	class FNode;
 }
+class UDataflow;
 
-
 UCLASS()
 class DATAFLOWEDITOR_API UDataflowSEditorObject : public UObject
 {
 	GENERATED_BODY()
 public:
 	TSharedPtr<Dataflow::FNode> Node = nullptr;
+	UDataflow* Graph = nullptr;
 
 	virtual void PostEditChangeProperty(struct FPropertyChangedEvent& PropertyChangedEvent) override;
 };
diff --git a/Engine/Plugins/Experimental/Dataflow/Source/DataflowEditor/Public/Dataflow/DataflowEPropertyFactories.h b/Engine/Plugins/Experimental/Dataflow/Source/DataflowEditor/Public/Dataflow/DataflowEPropertyFactories.h
--- a/Engine/Plugins/Experimental/Dataflow/Source/DataflowEditor/Public/Dataflow/DataflowEPropertyFactories.h
+++ b/Engine/Plugins/Experimental/Dataflow/Source/DataflowEditor/Public/Dataflow/DataflowEPropertyFactories.h
@@ -3,6 +3,7 @@
 #pragma once
 
 #include "CoreMinimal.h"
+#include "Dataflow/DataflowObject.h"
 #include "Dataflow/DataflowProperty.h"
 #include "DetailCategoryBuilder.h"
 #include "DetailLayoutBuilder.h"
@@ -24,6 +25,7 @@ namespace Dataflow
 	template<class T> DATAFLOWEDITOR_API void PropertyWidgetFactory(
 		IDetailLayoutBuilder& InDetailBuilder,
 		TSharedPtr<IPropertyHandle> PropertyHandle,
+		UDataflow* Graph,
 		TSharedPtr<Dataflow::FNode> InNode,
 		TProperty<T>* InProperty) 
 	{
@@ -42,8 +44,9 @@ namespace Dataflow
 						{
 							return InProperty->GetValue();
 						})
-					.OnValueCommitted_Lambda([InNode, InProperty](T NewValue, ETextCommit::Type CommitType)
+					.OnValueCommitted_Lambda([InNode, InProperty, Graph](T NewValue, ETextCommit::Type CommitType)
 						{
+							Graph->Modify();
 							InProperty->SetValue(NewValue);
 						})
 			];
@@ -57,6 +60,7 @@ namespace Dataflow
 		void PropertyWidgetFactory<bool>(
 			IDetailLayoutBuilder& InDetailBuilder,
 			TSharedPtr<IPropertyHandle> PropertyHandle,
+			UDataflow* Graph,
 			TSharedPtr<Dataflow::FNode> InNode,
 			TProperty<bool>* InProperty)
 	{
@@ -70,8 +74,9 @@ namespace Dataflow
 						{
 							return InProperty->GetValue()?ECheckBoxState::Checked:ECheckBoxState::Unchecked;
 						})
-					.OnCheckStateChanged_Lambda([InNode, InProperty](ECheckBoxState NewState)
+					.OnCheckStateChanged_Lambda([InNode, InProperty, Graph](ECheckBoxState NewState)
 						{
+							Graph->Modify();
 							InProperty->SetValue(NewState == ECheckBoxState::Checked);
 						})
 
@@ -85,6 +90,7 @@ namespace Dataflow
 	void PropertyWidgetFactory<FString>(
 		IDetailLayoutBuilder& InDetailBuilder,
 		TSharedPtr<IPropertyHandle> PropertyHandle,
+		UDataflow* Graph,
 		TSharedPtr<Dataflow::FNode> InNode,
 		TProperty<FString>* InProperty)
 	{
@@ -99,8 +105,9 @@ namespace Dataflow
 						{
 							return FText::FromString(InProperty->GetValue());
 						})
-					.OnTextCommitted_Lambda([InNode, InProperty](const FText& InText, ETextCommit::Type TextCommitType)
+					.OnTextCommitted_Lambda([InNode, InProperty, Graph](const FText& InText, ETextCommit::Type TextCommitType)
 						{
+							Graph->Modify();
 							InProperty->SetValue(InText.ToString());
 						})
 
@@ -114,6 +121,7 @@ namespace Dataflow
 		void PropertyWidgetFactory<FName>(
 			IDetailLayoutBuilder& InDetailBuilder,
 			TSharedPtr<IPropertyHandle> PropertyHandle,
+			UDataflow* Graph,
 			TSharedPtr<Dataflow::FNode> InNode,
 			TProperty<FName>* InProperty)
 	{
@@ -128,8 +136,9 @@ namespace Dataflow
 						{
 							return FText::FromString(InProperty->GetValue().ToString());
 						})
-					.OnTextCommitted_Lambda([InNode, InProperty](const FText& InText, ETextCommit::Type TextCommitType)
+					.OnTextCommitted_Lambda([InNode, InProperty, Graph](const FText& InText, ETextCommit::Type TextCommitType)
 						{
+							Graph->Modify();
 							InProperty->SetValue(FName(InText.ToString()));
 						})
 
