Dataflow : fix for crash on delete selected

diff --git a/Engine/Plugins/Experimental/Dataflow/Source/DataflowEditor/Private/Dataflow/DataflowGraphEditor.cpp b/Engine/Plugins/Experimental/Dataflow/Source/DataflowEditor/Private/Dataflow/DataflowGraphEditor.cpp
--- a/Engine/Plugins/Experimental/Dataflow/Source/DataflowEditor/Private/Dataflow/DataflowGraphEditor.cpp
+++ b/Engine/Plugins/Experimental/Dataflow/Source/DataflowEditor/Private/Dataflow/DataflowGraphEditor.cpp
@@ -6,6 +6,7 @@
 #include "Dataflow/DataflowEngine.h"
 #include "Dataflow/DataflowSNodeFactories.h"
 #include "BoneDragDropOp.h"
+#include "IStructureDetailsView.h"
 
 DEFINE_LOG_CATEGORY_STATIC(SDataflowGraphEditorLog, Log, All);
 
@@ -77,6 +78,11 @@ void SDataflowGraphEditor::DeleteNode()
 {
 	if (DataflowAsset.Get())
 	{
+		if (DetailsView)
+		{
+			DetailsView->SetStructureData(nullptr);
+		}
+
 		FDataflowEditorCommands::DeleteNodes(DataflowAsset.Get(), GetSelectedNodes());
 	}
 }
