Dataflow : File cleanup

diff --git a/Engine/Plugins/Experimental/Dataflow/Source/DataflowEditor/Private/Dataflow/DataflowEPropertyCustomizations.cpp b/Engine/Plugins/Experimental/Dataflow/Source/DataflowEditor/Private/Dataflow/DataflowEPropertyCustomizations.cpp
deleted file mode 100644
diff --git a/Engine/Plugins/Experimental/Dataflow/Source/DataflowEditor/Private/Dataflow/DataflowEditorActions.cpp b/Engine/Plugins/Experimental/Dataflow/Source/DataflowEditor/Private/Dataflow/DataflowEditorActions.cpp
--- a/Engine/Plugins/Experimental/Dataflow/Source/DataflowEditor/Private/Dataflow/DataflowEditorActions.cpp
+++ b/Engine/Plugins/Experimental/Dataflow/Source/DataflowEditor/Private/Dataflow/DataflowEditorActions.cpp
@@ -2,7 +2,6 @@
 
 #include "Dataflow/DataflowEditorActions.h"
 #include "Dataflow/DataflowEdNode.h"
-#include "Dataflow/DataflowEPropertyCustomizations.h"
 #include "Dataflow/DataflowNodeFactory.h"
 #include "Dataflow/DataflowObject.h"
 #include "IStructureDetailsView.h"
diff --git a/Engine/Plugins/Experimental/Dataflow/Source/DataflowEditor/Private/Dataflow/DataflowEditorPlugin.cpp b/Engine/Plugins/Experimental/Dataflow/Source/DataflowEditor/Private/Dataflow/DataflowEditorPlugin.cpp
--- a/Engine/Plugins/Experimental/Dataflow/Source/DataflowEditor/Private/Dataflow/DataflowEditorPlugin.cpp
+++ b/Engine/Plugins/Experimental/Dataflow/Source/DataflowEditor/Private/Dataflow/DataflowEditorPlugin.cpp
@@ -8,7 +8,6 @@
 #include "DataflowEditorToolkit.h"
 #include "Dataflow/DataflowNodeFactory.h"
 #include "Dataflow/DataflowAssetActions.h"
-#include "Dataflow/DataflowEPropertyCustomizations.h"
 #include "Dataflow/DataflowSNodeFactories.h"
 
 #define LOCTEXT_NAMESPACE "DataflowEditor"
@@ -25,7 +24,6 @@
 	AssetTools.RegisterAssetTypeActions(MakeShareable(DataflowAssetActions));
 
 	FPropertyEditorModule& PropertyEditorModule = FModuleManager::LoadModuleChecked<FPropertyEditorModule>(TEXT("PropertyEditor"));
-	PropertyEditorModule.RegisterCustomClassLayout(UDataflowSEditorObject::StaticClass()->GetFName(), FOnGetDetailCustomizationInstance::CreateStatic(&FDataflowSEditorCustomization::MakeInstance));
 
 	DataflowSNodeFactory = MakeShareable(new FDataflowSNodeFactory());
 	FEdGraphUtilities::RegisterVisualNodeFactory(DataflowSNodeFactory);
diff --git a/Engine/Plugins/Experimental/Dataflow/Source/DataflowEditor/Public/Dataflow/DataflowEPropertyCustomizations.h b/Engine/Plugins/Experimental/Dataflow/Source/DataflowEditor/Public/Dataflow/DataflowEPropertyCustomizations.h
deleted file mode 100644
diff --git a/Engine/Plugins/Experimental/Dataflow/Source/DataflowNodes/Private/Dataflow/DataflowNodesConnectionTypes.cpp b/Engine/Plugins/Experimental/Dataflow/Source/DataflowNodes/Private/Dataflow/DataflowNodesConnectionTypes.cpp
deleted file mode 100644
diff --git a/Engine/Plugins/Experimental/Dataflow/Source/DataflowNodes/Public/Dataflow/DataflowNodesConnectionTypes.h b/Engine/Plugins/Experimental/Dataflow/Source/DataflowNodes/Public/Dataflow/DataflowNodesConnectionTypes.h
deleted file mode 100644
diff --git a/Engine/Plugins/Experimental/Dataflow/Source/DataflowNodes/Public/Dataflow/DataflowSkeletalMeshNodes.h b/Engine/Plugins/Experimental/Dataflow/Source/DataflowNodes/Public/Dataflow/DataflowSkeletalMeshNodes.h
--- a/Engine/Plugins/Experimental/Dataflow/Source/DataflowNodes/Public/Dataflow/DataflowSkeletalMeshNodes.h
+++ b/Engine/Plugins/Experimental/Dataflow/Source/DataflowNodes/Public/Dataflow/DataflowSkeletalMeshNodes.h
@@ -3,13 +3,9 @@
 #pragma once 
 
 #include "CoreMinimal.h"
-
-#include "Dataflow/DataflowNodesConnectionTypes.h"
 #include "Dataflow/DataflowCore.h"
 #include "Dataflow/DataflowEngine.h"
 
-
-
 #include "DataflowSkeletalMeshNodes.generated.h"
 
 DEFINE_LOG_CATEGORY_STATIC(LogDataflowSkeletalMeshNodes, Log, All);
diff --git a/Engine/Plugins/Experimental/Dataflow/Source/DataflowNodes/Public/Dataflow/DataflowStaticMeshNodes.h b/Engine/Plugins/Experimental/Dataflow/Source/DataflowNodes/Public/Dataflow/DataflowStaticMeshNodes.h
--- a/Engine/Plugins/Experimental/Dataflow/Source/DataflowNodes/Public/Dataflow/DataflowStaticMeshNodes.h
+++ b/Engine/Plugins/Experimental/Dataflow/Source/DataflowNodes/Public/Dataflow/DataflowStaticMeshNodes.h
@@ -3,13 +3,9 @@
 #pragma once 
 
 #include "CoreMinimal.h"
-
-
-#include "Dataflow/DataflowNodesConnectionTypes.h"
 #include "Dataflow/DataflowCore.h"
 #include "Dataflow/DataflowEngine.h"
 
-
 #include "DataflowStaticMeshNodes.generated.h"
 
 DEFINE_LOG_CATEGORY_STATIC(LogDataflowStaticMeshNodes, Log, All);
diff --git a/Engine/Plugins/Experimental/GeometryCollectionPlugin/Source/GeometryCollectionNodes/Public/Dataflow/GeometryCollectionNodeConnectionTypes.h b/Engine/Plugins/Experimental/GeometryCollectionPlugin/Source/GeometryCollectionNodes/Public/Dataflow/GeometryCollectionNodeConnectionTypes.h
deleted file mode 100644
diff --git a/Engine/Plugins/Experimental/GeometryCollectionPlugin/Source/GeometryCollectionNodes/Public/Dataflow/GeometryCollectionNodes.h b/Engine/Plugins/Experimental/GeometryCollectionPlugin/Source/GeometryCollectionNodes/Public/Dataflow/GeometryCollectionNodes.h
--- a/Engine/Plugins/Experimental/GeometryCollectionPlugin/Source/GeometryCollectionNodes/Public/Dataflow/GeometryCollectionNodes.h
+++ b/Engine/Plugins/Experimental/GeometryCollectionPlugin/Source/GeometryCollectionNodes/Public/Dataflow/GeometryCollectionNodes.h
@@ -2,8 +2,6 @@
 
 #pragma once 
 
-#include "GeometryCollectionNodeConnectionTypes.h"
-
 #include "CoreMinimal.h"
 #include "Dataflow/DataflowEngine.h"
 #include "GeometryCollection/ManagedArrayCollection.h"
diff --git a/Engine/Source/Runtime/Experimental/Dataflow/Core/Private/Dataflow/DataflowNodeParameters.cpp b/Engine/Source/Runtime/Experimental/Dataflow/Core/Private/Dataflow/DataflowNodeParameters.cpp
deleted file mode 100644
