Datflow : Re-enable asset editor

diff --git a/Engine/Plugins/Experimental/Dataflow/Source/DataflowEditor/Private/Dataflow/DataflowAssetActions.cpp b/Engine/Plugins/Experimental/Dataflow/Source/DataflowEditor/Private/Dataflow/DataflowAssetActions.cpp
--- a/Engine/Plugins/Experimental/Dataflow/Source/DataflowEditor/Private/Dataflow/DataflowAssetActions.cpp
+++ b/Engine/Plugins/Experimental/Dataflow/Source/DataflowEditor/Private/Dataflow/DataflowAssetActions.cpp
@@ -30,9 +30,8 @@ void FDataflowAssetActions::GetActions(const TArray<UObject*>& InObjects,
 void FDataflowAssetActions::OpenAssetEditor(
 	const TArray<UObject*>& InObjects, TSharedPtr<IToolkitHost> EditWithinLevelEditor)
 {
-	FAssetTypeActions_Base::OpenAssetEditor(InObjects, EditWithinLevelEditor);
+	//FAssetTypeActions_Base::OpenAssetEditor(InObjects, EditWithinLevelEditor);
 
-	/*
 	EToolkitMode::Type Mode = EditWithinLevelEditor.IsValid() ? EToolkitMode::WorldCentric : EToolkitMode::Standalone;
 	for (auto ObjIt = InObjects.CreateConstIterator(); ObjIt; ++ObjIt)
 	{
@@ -42,7 +41,6 @@ void FDataflowAssetActions::OpenAssetEditor(
 			DataflowEditorPlugin->CreateDataflowAssetEditor(Mode, EditWithinLevelEditor, Object);
 		}
 	}
-	*/
 }
 
 
diff --git a/Engine/Plugins/Experimental/Dataflow/Source/DataflowEditor/Private/Dataflow/DataflowEditorToolkit.cpp b/Engine/Plugins/Experimental/Dataflow/Source/DataflowEditor/Private/Dataflow/DataflowEditorToolkit.cpp
--- a/Engine/Plugins/Experimental/Dataflow/Source/DataflowEditor/Private/Dataflow/DataflowEditorToolkit.cpp
+++ b/Engine/Plugins/Experimental/Dataflow/Source/DataflowEditor/Private/Dataflow/DataflowEditorToolkit.cpp
@@ -7,6 +7,7 @@
 #include "Dataflow/DataflowEdNode.h"
 #include "Dataflow/DataflowNodeFactory.h"
 #include "Dataflow/DataflowObject.h"
+#include "Dataflow/DataflowObjectInterface.h"
 #include "Dataflow/DataflowSchema.h"
 #include "Dataflow/DataflowCore.h"
 #include "EditorStyleSet.h"
@@ -23,7 +24,9 @@
 //DEFINE_LOG_CATEGORY_STATIC(FDataflowEditorToolkitLog, Log, All);
 
 const FName FDataflowEditorToolkit::GraphCanvasTabId(TEXT("DataflowEditor_GraphCanvas"));
-const FName FDataflowEditorToolkit::PropertiesTabId(TEXT("DataflowEditor_Properties"));
+const FName FDataflowEditorToolkit::AssetDetailsTabId(TEXT("DataflowEditor_AssetDetails"));
+const FName FDataflowEditorToolkit::NodeDetailsTabId(TEXT("DataflowEditor_NodeDetails"));
+
 
 void FDataflowEditorToolkit::InitDataflowEditor(const EToolkitMode::Type Mode, const TSharedPtr<IToolkitHost>& InitToolkitHost, UObject* ObjectToEdit)
 {
@@ -32,10 +35,11 @@ void FDataflowEditorToolkit::InitDataflowEditor(const EToolkitMode::Type Mode, c
 	{
 		Dataflow->Schema = UDataflowSchema::StaticClass();
 
-		GraphEditor = CreateGraphEditorWidget(Dataflow);
-		PropertiesEditor = CreatePropertiesEditorWidget(ObjectToEdit);
+		NodeDetailsEditor = CreateNodeDetailsEditorWidget(ObjectToEdit);
+		AssetDetailsEditor = CreateAssetDetailsEditorWidget(ObjectToEdit);
+		GraphEditor = CreateGraphEditorWidget(Dataflow, NodeDetailsEditor);
 
-		const TSharedRef<FTabManager::FLayout> StandaloneDefaultLayout = FTabManager::NewLayout("Dataflow_Layout")
+		const TSharedRef<FTabManager::FLayout> StandaloneDefaultLayout = FTabManager::NewLayout("Dataflow_Layout.V1")
 			->AddArea
 			(
 				FTabManager::NewPrimaryArea()->SetOrientation(Orient_Vertical)
@@ -60,7 +64,7 @@ void FDataflowEditorToolkit::InitDataflowEditor(const EToolkitMode::Type Mode, c
 							(
 								FTabManager::NewStack()
 								->SetSizeCoefficient(0.7f)
-								->AddTab(PropertiesTabId, ETabState::OpenedTab)
+								->AddTab(AssetDetailsTabId, ETabState::OpenedTab)
 							)
 						)
 					)
@@ -75,15 +79,40 @@ void FDataflowEditorToolkit::InitDataflowEditor(const EToolkitMode::Type Mode, c
 
 
 
-TSharedRef<SGraphEditor> FDataflowEditorToolkit::CreateGraphEditorWidget(UDataflow* DataflowToEdit)
+TSharedRef<SGraphEditor> FDataflowEditorToolkit::CreateGraphEditorWidget(UDataflow* DataflowToEdit, TSharedPtr<IStructureDetailsView> InNodeDetailsEditor)
 {
 	ensure(DataflowToEdit);
+	using namespace Dataflow;
+
+	FDataflowEditorCommands::FGraphEvaluationCallback Evaluate = [&](FDataflowNode* Node, FDataflowOutput* Out)
+	{
+		if (DataflowToEdit->bActive)
+		{
+			float EvalTime = FGameTime::GetTimeSinceAppStart().GetRealTimeSeconds();
+			if (DataflowToEdit->Targets.Num())
+			{
+				for (UObject* Object : DataflowToEdit->Targets)
+				{
+					FEngineContext Context(Object, Dataflow, EvalTime, FString("TargetObject"));
+					Node->Evaluate(Context, Out);
+				}
+			}
+			else
+			{
+				FEngineContext Context(DataflowToEdit, Dataflow, EvalTime, FString("UDataflow"));
+				Node->Evaluate(Context, Out);
+			}
+		}
+	};
+
 	return SNew(SDataflowGraphEditor, DataflowToEdit)
 		.GraphToEdit(DataflowToEdit)
-		.DetailsView(GetPropertiesEditor());
+		.DetailsView(InNodeDetailsEditor)
+		.EvaluateGraph(Evaluate);
+
 }
 
-TSharedPtr<IStructureDetailsView> FDataflowEditorToolkit::CreatePropertiesEditorWidget(UObject* ObjectToEdit)
+TSharedPtr<IStructureDetailsView> FDataflowEditorToolkit::CreateNodeDetailsEditorWidget(UObject* ObjectToEdit)
 {
 	ensure(ObjectToEdit);
 	FPropertyEditorModule& PropertyEditorModule = FModuleManager::LoadModuleChecked<FPropertyEditorModule>(TEXT("PropertyEditor"));
@@ -108,7 +137,6 @@ TSharedPtr<IStructureDetailsView> FDataflowEditorToolkit::CreatePropertiesEditor
 		StructureViewArgs.bShowClasses = true;
 		StructureViewArgs.bShowInterfaces = true;
 	}
-
 	TSharedPtr<IStructureDetailsView> DetailsView = PropertyEditorModule.CreateStructureDetailView(DetailsViewArgs, StructureViewArgs, nullptr);
 	DetailsView->GetDetailsView()->SetObject(ObjectToEdit);
 
@@ -116,6 +144,26 @@ TSharedPtr<IStructureDetailsView> FDataflowEditorToolkit::CreatePropertiesEditor
 
 }
 
+TSharedPtr<IDetailsView> FDataflowEditorToolkit::CreateAssetDetailsEditorWidget(UObject* ObjectToEdit)
+{
+	ensure(ObjectToEdit);
+	FPropertyEditorModule& PropertyEditorModule = FModuleManager::LoadModuleChecked<FPropertyEditorModule>(TEXT("PropertyEditor"));
+
+	FDetailsViewArgs DetailsViewArgs;
+	{
+		DetailsViewArgs.bAllowSearch = true;
+		DetailsViewArgs.bLockable = false;
+		DetailsViewArgs.bUpdatesFromSelection = false;
+		DetailsViewArgs.NameAreaSettings = FDetailsViewArgs::HideNameArea;
+		DetailsViewArgs.NotifyHook = this;
+	}
+
+	TSharedPtr<IDetailsView> DetailsView = PropertyEditorModule.CreateDetailView(DetailsViewArgs);
+	DetailsView->SetObject(ObjectToEdit);
+	return DetailsView;
+
+}
+
 
 TSharedRef<SDockTab> FDataflowEditorToolkit::SpawnTab_GraphCanvas(const FSpawnTabArgs& Args)
 {
@@ -130,17 +178,27 @@ TSharedRef<SDockTab> FDataflowEditorToolkit::SpawnTab_GraphCanvas(const FSpawnTa
 	return SpawnedTab;
 }
 
-TSharedRef<SDockTab> FDataflowEditorToolkit::SpawnTab_Properties(const FSpawnTabArgs& Args)
+TSharedRef<SDockTab> FDataflowEditorToolkit::SpawnTab_AssetDetails(const FSpawnTabArgs& Args)
 {
-	check(Args.GetTabId() == PropertiesTabId);
+	check(Args.GetTabId() == AssetDetailsTabId);
 
 	return SNew(SDockTab)
-		.Label(LOCTEXT("DataflowEditor_Properties_TabTitle", "Details"))
+		.Label(LOCTEXT("DataflowEditor_AssetDetails_TabTitle", "Details"))
 		[
-			PropertiesEditor->GetWidget()->AsShared()
+			AssetDetailsEditor->AsShared()
 		];
 }
 
+TSharedRef<SDockTab> FDataflowEditorToolkit::SpawnTab_NodeDetails(const FSpawnTabArgs& Args)
+{
+	check(Args.GetTabId() == NodeDetailsTabId);
+
+	return SNew(SDockTab)
+		.Label(LOCTEXT("DataflowEditor_NodeDetails_TabTitle", "Node Details"))
+		[
+			NodeDetailsEditor->GetWidget()->AsShared()
+		];
+}
 
 void FDataflowEditorToolkit::RegisterTabSpawners(const TSharedRef<FTabManager>& InTabManager)
 {
@@ -151,8 +209,13 @@ void FDataflowEditorToolkit::RegisterTabSpawners(const TSharedRef<FTabManager>&
 		.SetGroup(WorkspaceMenuCategoryRef)
 		.SetIcon(FSlateIcon(FAppStyle::GetAppStyleSetName(), "GraphEditor.EventGraph_16x"));
 
-	InTabManager->RegisterTabSpawner(PropertiesTabId, FOnSpawnTab::CreateSP(this, &FDataflowEditorToolkit::SpawnTab_Properties))
-		.SetDisplayName(LOCTEXT("PropertiesTab", "Details"))
+	InTabManager->RegisterTabSpawner(AssetDetailsTabId, FOnSpawnTab::CreateSP(this, &FDataflowEditorToolkit::SpawnTab_AssetDetails))
+		.SetDisplayName(LOCTEXT("AssetDetailsTab", "Details"))
+		.SetGroup(WorkspaceMenuCategoryRef)
+		.SetIcon(FSlateIcon(FAppStyle::GetAppStyleSetName(), "LevelEditor.Tabs.Details"));
+
+	InTabManager->RegisterTabSpawner(NodeDetailsTabId, FOnSpawnTab::CreateSP(this, &FDataflowEditorToolkit::SpawnTab_NodeDetails))
+		.SetDisplayName(LOCTEXT("NodeDetailsTab", "Node Details"))
 		.SetGroup(WorkspaceMenuCategoryRef)
 		.SetIcon(FSlateIcon(FAppStyle::GetAppStyleSetName(), "LevelEditor.Tabs.Details"));
 
diff --git a/Engine/Plugins/Experimental/Dataflow/Source/DataflowEditor/Private/Dataflow/DataflowEditorToolkit.h b/Engine/Plugins/Experimental/Dataflow/Source/DataflowEditor/Private/Dataflow/DataflowEditorToolkit.h
--- a/Engine/Plugins/Experimental/Dataflow/Source/DataflowEditor/Private/Dataflow/DataflowEditorToolkit.h
+++ b/Engine/Plugins/Experimental/Dataflow/Source/DataflowEditor/Private/Dataflow/DataflowEditorToolkit.h
@@ -30,14 +30,18 @@ public:
 	// Tab spawners 
 	virtual void RegisterTabSpawners(const TSharedRef<FTabManager>& TabManager) override;
 	TSharedRef<SDockTab> SpawnTab_GraphCanvas(const FSpawnTabArgs& Args);
-	TSharedRef<SDockTab> SpawnTab_Properties(const FSpawnTabArgs& Args);
+	TSharedRef<SDockTab> SpawnTab_AssetDetails(const FSpawnTabArgs& Args);
+	TSharedRef<SDockTab> SpawnTab_NodeDetails(const FSpawnTabArgs& Args);
 
 	// Member Access
 	UDataflow* GetDataflow() { return Dataflow; }
 	const UDataflow* GetDataflow() const { return Dataflow; }
 
-	TSharedPtr<IStructureDetailsView> GetPropertiesEditor() {return PropertiesEditor;}
-	const TSharedPtr<IStructureDetailsView> GetPropertiesEditor() const { return PropertiesEditor; }
+	TSharedPtr<IDetailsView> GetAssetDetailsEditor() {return AssetDetailsEditor;}
+	const TSharedPtr<IDetailsView> GetAssetDetailsEditor() const { return AssetDetailsEditor; }
+
+	TSharedPtr<IStructureDetailsView> GetNodeDetailsEditor() { return NodeDetailsEditor; }
+	const TSharedPtr<IStructureDetailsView> GetNodeDetailsEditor() const { return NodeDetailsEditor; }
 
 	TSharedPtr<SGraphEditor> GetGraphEditor() { return GraphEditor; }
 	const TSharedPtr<SGraphEditor> GetGraphEditor() const { return GraphEditor; }
@@ -49,9 +53,13 @@ private:
 	static const FName GraphCanvasTabId;
 	TSharedPtr<SGraphEditor> GraphEditor;
 	TSharedPtr<FUICommandList> GraphEditorCommands;
-	TSharedRef<SGraphEditor> CreateGraphEditorWidget(UDataflow* ObjectToEdit);
+	TSharedRef<SGraphEditor> CreateGraphEditorWidget(UDataflow* ObjectToEdit, TSharedPtr<IStructureDetailsView> PropertiesEditor);
+
+	static const FName AssetDetailsTabId;
+	TSharedPtr<IDetailsView> AssetDetailsEditor;
+	TSharedPtr<IDetailsView> CreateAssetDetailsEditorWidget(UObject* ObjectToEdit);
 
-	static const FName PropertiesTabId;
-	TSharedPtr<IStructureDetailsView> PropertiesEditor;
-	TSharedPtr<IStructureDetailsView> CreatePropertiesEditorWidget(UObject* ObjectToEdit);
+	static const FName NodeDetailsTabId;
+	TSharedPtr<IStructureDetailsView> NodeDetailsEditor;
+	TSharedPtr<IStructureDetailsView> CreateNodeDetailsEditorWidget(UObject* ObjectToEdit);
 };
diff --git a/Engine/Source/Runtime/Experimental/Dataflow/Engine/Public/Dataflow/DataflowObject.h b/Engine/Source/Runtime/Experimental/Dataflow/Engine/Public/Dataflow/DataflowObject.h
--- a/Engine/Source/Runtime/Experimental/Dataflow/Engine/Public/Dataflow/DataflowObject.h
+++ b/Engine/Source/Runtime/Experimental/Dataflow/Engine/Public/Dataflow/DataflowObject.h
@@ -52,6 +52,12 @@ class DATAFLOWENGINE_API UDataflow : public UEdGraph
 public:
 	UDataflow(const FObjectInitializer& ObjectInitializer = FObjectInitializer::Get());
 
+	UPROPERTY(EditAnywhere, Category = "Evaluation")
+	bool bActive = true;
+
+	UPROPERTY(EditAnywhere, Category = "Evaluation", AdvancedDisplay )
+	TArray<TObjectPtr<UObject>> Targets;
+
 	/** UObject Interface */
 #if WITH_EDITOR
 	virtual void PostEditChangeProperty(struct FPropertyChangedEvent& PropertyChangedEvent) override;
