Dataflow Node :
- added pin color based on data type, it uses the predefined colors from the Editor Preferences
- the type colors which are not defined in the Editor Preferences can be defined in Project Settings under the Dataflow plugin
- added connection between nodes color based on data type of the pins it's connecting
- added connection between nodes focus display when user hovers over it
- updated Context menu with standard editor/node functions, all the standard editor alignment functions
- added ability to create comment box
- updated delete to delete comment box

diff --git a/Engine/Plugins/Experimental/Dataflow/Source/DataflowEditor/DataflowEditor.Build.cs b/Engine/Plugins/Experimental/Dataflow/Source/DataflowEditor/DataflowEditor.Build.cs
--- a/Engine/Plugins/Experimental/Dataflow/Source/DataflowEditor/DataflowEditor.Build.cs
+++ b/Engine/Plugins/Experimental/Dataflow/Source/DataflowEditor/DataflowEditor.Build.cs
@@ -35,6 +35,7 @@ namespace UnrealBuildTool.Rules
 					"DataflowCore",
 					"DataflowEngine",
 					"Slate",
+					"DeveloperSettings",
 				}
 			);
 		}
diff --git a/Engine/Plugins/Experimental/Dataflow/Source/DataflowEditor/Private/Dataflow/DataflowCommentNode.cpp b/Engine/Plugins/Experimental/Dataflow/Source/DataflowEditor/Private/Dataflow/DataflowCommentNode.cpp
new file mode 100644
--- /dev/null
+++ b/Engine/Plugins/Experimental/Dataflow/Source/DataflowEditor/Private/Dataflow/DataflowCommentNode.cpp
@@ -0,0 +1,75 @@
+// Copyright Epic Games, Inc. All Rights Reserved.
+
+#include "Dataflow/DataflowCommentNode.h"
+
+#include "EdGraphNode_Comment.h"
+#include "Dataflow/DataflowEdNode.h"
+#include "Dataflow/DataflowNodeFactory.h"
+#include "Dataflow/DataflowObject.h"
+#include "Dataflow/DataflowCore.h"
+#include "Logging/LogMacros.h"
+#include "Widgets/Layout/SBorder.h"
+#include "Settings/EditorStyleSettings.h"
+#include "ScopedTransaction.h"
+
+#define LOCTEXT_NAMESPACE "SDataflowEdNodeComment"
+
+//
+// Add a menu option to create a graph node.
+//
+TSharedPtr<FAssetSchemaAction_Dataflow_CreateCommentNode_DataflowEdNode> FAssetSchemaAction_Dataflow_CreateCommentNode_DataflowEdNode::CreateAction(UEdGraph* ParentGraph, const TSharedPtr<SGraphEditor>& GraphEditor)
+{
+	return MakeShared<FAssetSchemaAction_Dataflow_CreateCommentNode_DataflowEdNode>(GraphEditor);
+}
+
+//
+//  Created comment node
+//
+UEdGraphNode* FAssetSchemaAction_Dataflow_CreateCommentNode_DataflowEdNode::PerformAction(class UEdGraph* ParentGraph, UEdGraphPin* FromPin, const FVector2D Location, bool bSelectNewNode)
+{
+	UEdGraphNode_Comment* CommentTemplate = NewObject<UEdGraphNode_Comment>();
+
+	FVector2D SpawnLocation = Location;
+	FSlateRect Bounds;
+
+	if (GraphEditor->GetBoundsForSelectedNodes(Bounds, 50.0f))
+	{
+		CommentTemplate->SetBounds(Bounds);
+		SpawnLocation.X = CommentTemplate->NodePosX;
+		SpawnLocation.Y = CommentTemplate->NodePosY;
+	}
+
+	CommentTemplate->bCommentBubbleVisible_InDetailsPanel = false;
+	CommentTemplate->bCommentBubbleVisible = false;
+	CommentTemplate->bCommentBubblePinned = false;
+
+	const FScopedTransaction Transaction(NSLOCTEXT("UnrealEd", "DataflowEditorNewNode", "Dataflow Editor: New Comment Node"));
+	ParentGraph->Modify();
+
+	CommentTemplate->SetFlags(RF_Transactional);
+
+	// set outer to be the graph so it doesn't go away
+	CommentTemplate->Rename(NULL, ParentGraph, REN_NonTransactional);
+	ParentGraph->AddNode(CommentTemplate, true, bSelectNewNode);
+
+	CommentTemplate->CreateNewGuid();
+	CommentTemplate->PostPlacedNewNode();
+	CommentTemplate->AllocateDefaultPins();
+	CommentTemplate->AutowireNewNode(FromPin);
+
+	CommentTemplate->NodePosX = SpawnLocation.X;
+	CommentTemplate->NodePosY = SpawnLocation.Y;
+	CommentTemplate->SnapToGrid(GetDefault<UEditorStyleSettings>()->GridSnapSize);
+
+	ParentGraph->NotifyGraphChanged();
+
+	return CommentTemplate;
+}
+
+//void FAssetSchemaAction_Dataflow_CreateCommentNode_DataflowEdNode::AddReferencedObjects(FReferenceCollector& Collector)
+//{
+//	FEdGraphSchemaAction::AddReferencedObjects(Collector);
+//	Collector.AddReferencedObject(NodeTemplate);
+//}
+
+#undef LOCTEXT_NAMESPACE
diff --git a/Engine/Plugins/Experimental/Dataflow/Source/DataflowEditor/Private/Dataflow/DataflowEditorActions.cpp b/Engine/Plugins/Experimental/Dataflow/Source/DataflowEditor/Private/Dataflow/DataflowEditorActions.cpp
--- a/Engine/Plugins/Experimental/Dataflow/Source/DataflowEditor/Private/Dataflow/DataflowEditorActions.cpp
+++ b/Engine/Plugins/Experimental/Dataflow/Source/DataflowEditor/Private/Dataflow/DataflowEditorActions.cpp
@@ -5,13 +5,15 @@
 #include "Dataflow/DataflowNodeFactory.h"
 #include "Dataflow/DataflowObject.h"
 #include "IStructureDetailsView.h"
-
+#include "EdGraphNode_Comment.h"
 
 #define LOCTEXT_NAMESPACE "DataflowEditorCommands"
 
 void FDataflowEditorCommandsImpl::RegisterCommands()
 {
 	UI_COMMAND(EvaluateNode, "Evaluate", "Trigger an evaluation of the selected node.", EUserInterfaceActionType::Button, FInputChord());
+	UI_COMMAND(CreateComment, "CreateComment", "Create a Comment node.", EUserInterfaceActionType::None, FInputChord());
+	UI_COMMAND(ToggleEnabledState, "ToggleEnabledState", "Toggle node between Enabled/Disabled state.", EUserInterfaceActionType::Button, FInputChord());
 
 	if (Dataflow::FNodeFactory* Factory = Dataflow::FNodeFactory::GetInstance())
 	{
@@ -95,6 +97,10 @@ void FDataflowEditorCommands::DeleteNodes(UDataflow* Graph, const FGraphPanelSel
 				}
 			}
 		}
+		else if (UEdGraphNode_Comment* CommentNode = dynamic_cast<UEdGraphNode_Comment*>(Ode))
+		{
+			Graph->RemoveNode(CommentNode);
+		}
 	}
 }
 
@@ -127,4 +133,8 @@ void FDataflowEditorCommands::OnSelectedNodesChanged(TSharedPtr<IStructureDetail
 	}
 }
 
+void FDataflowEditorCommands::ToggleEnabledState(UDataflow* Graph)
+{
+}
+
 #undef LOCTEXT_NAMESPACE
diff --git a/Engine/Plugins/Experimental/Dataflow/Source/DataflowEditor/Private/Dataflow/DataflowGraphEditor.cpp b/Engine/Plugins/Experimental/Dataflow/Source/DataflowEditor/Private/Dataflow/DataflowGraphEditor.cpp
--- a/Engine/Plugins/Experimental/Dataflow/Source/DataflowEditor/Private/Dataflow/DataflowGraphEditor.cpp
+++ b/Engine/Plugins/Experimental/Dataflow/Source/DataflowEditor/Private/Dataflow/DataflowGraphEditor.cpp
@@ -7,6 +7,9 @@
 #include "Dataflow/DataflowSNodeFactories.h"
 #include "BoneDragDropOp.h"
 #include "IStructureDetailsView.h"
+#include "Dataflow/DataflowCommentNode.h"
+
+#include "Dataflow/DataflowSchema.h"
 
 DEFINE_LOG_CATEGORY_STATIC(SDataflowGraphEditorLog, Log, All);
 
@@ -31,12 +34,58 @@ void SDataflowGraphEditor::Construct(const FArguments& InArgs, UObject* InAssetO
 	{
 		GraphEditorCommands = MakeShareable(new FUICommandList);
 		{
-			GraphEditorCommands->MapAction(FGenericCommands::Get().Delete,
+			GraphEditorCommands->MapAction(
+				FGenericCommands::Get().Delete,
 				FExecuteAction::CreateSP(this, &SDataflowGraphEditor::DeleteNode)
 			);
-			GraphEditorCommands->MapAction(FDataflowEditorCommands::Get().EvaluateNode,
+			GraphEditorCommands->MapAction(
+				FDataflowEditorCommands::Get().EvaluateNode,
 				FExecuteAction::CreateSP(this, &SDataflowGraphEditor::EvaluateNode)
 			);
+			GraphEditorCommands->MapAction(
+				FGraphEditorCommands::Get().CreateComment,
+				FExecuteAction::CreateRaw(this, &SDataflowGraphEditor::CreateComment)
+			);
+			GraphEditorCommands->MapAction(
+				FGraphEditorCommands::Get().AlignNodesTop,
+				FExecuteAction::CreateSP(this, &SDataflowGraphEditor::AlignTop)
+			);
+			GraphEditorCommands->MapAction(
+				FGraphEditorCommands::Get().AlignNodesMiddle,
+				FExecuteAction::CreateSP(this, &SDataflowGraphEditor::AlignMiddle)
+			);
+			GraphEditorCommands->MapAction(
+				FGraphEditorCommands::Get().AlignNodesBottom,
+				FExecuteAction::CreateSP(this, &SDataflowGraphEditor::AlignBottom)
+			);
+			GraphEditorCommands->MapAction(
+				FGraphEditorCommands::Get().AlignNodesLeft,
+				FExecuteAction::CreateSP(this, &SDataflowGraphEditor::AlignLeft)
+			);
+			GraphEditorCommands->MapAction(
+				FGraphEditorCommands::Get().AlignNodesCenter,
+				FExecuteAction::CreateSP(this, &SDataflowGraphEditor::AlignCenter)
+			);
+			GraphEditorCommands->MapAction(
+				FGraphEditorCommands::Get().AlignNodesRight,
+				FExecuteAction::CreateSP(this, &SDataflowGraphEditor::AlignRight)
+			);
+			GraphEditorCommands->MapAction(
+				FGraphEditorCommands::Get().StraightenConnections,
+				FExecuteAction::CreateSP(this, &SDataflowGraphEditor::StraightenConnections)
+			);
+			GraphEditorCommands->MapAction(
+				FGraphEditorCommands::Get().DistributeNodesHorizontally,
+				FExecuteAction::CreateSP(this, &SDataflowGraphEditor::DistributeHorizontally)
+			);
+			GraphEditorCommands->MapAction(
+				FGraphEditorCommands::Get().DistributeNodesVertically,
+				FExecuteAction::CreateSP(this, &SDataflowGraphEditor::DistributeVertically)
+			);
+			GraphEditorCommands->MapAction(
+				FDataflowEditorCommands::Get().ToggleEnabledState,
+				FExecuteAction::CreateRaw(this, &SDataflowGraphEditor::ToggleEnabledState)
+			);
 		}
 	}
 
@@ -144,4 +193,63 @@ void SDataflowGraphEditor::OnDragLeave(const FDragDropEvent& DragDropEvent)
 }
 */
 
+void SDataflowGraphEditor::CreateComment()
+{
+	UDataflow* Graph = DataflowAsset.Get();
+	const TSharedPtr<SGraphEditor>& InGraphEditor = SharedThis(GetGraphEditor());
+
+	TSharedPtr<FAssetSchemaAction_Dataflow_CreateCommentNode_DataflowEdNode> CommentAction = FAssetSchemaAction_Dataflow_CreateCommentNode_DataflowEdNode::CreateAction(Graph, InGraphEditor);
+	CommentAction->PerformAction(Graph, nullptr, GetGraphEditor()->GetPasteLocation(), false);
+}
+
+void SDataflowGraphEditor::AlignTop()
+{
+	GetGraphEditor()->OnAlignTop();
+}
+
+void SDataflowGraphEditor::AlignMiddle()
+{
+	GetGraphEditor()->OnAlignMiddle();
+}
+
+void SDataflowGraphEditor::AlignBottom()
+{
+	GetGraphEditor()->OnAlignBottom();
+}
+
+void SDataflowGraphEditor::AlignLeft()
+{
+	GetGraphEditor()->OnAlignLeft();
+}
+
+void SDataflowGraphEditor::AlignCenter()
+{
+	GetGraphEditor()->OnAlignCenter();
+}
+
+void SDataflowGraphEditor::AlignRight()
+{
+	GetGraphEditor()->OnAlignRight();
+}
+
+void SDataflowGraphEditor::StraightenConnections()
+{
+	GetGraphEditor()->OnStraightenConnections();
+}
+
+void SDataflowGraphEditor::DistributeHorizontally()
+{
+	GetGraphEditor()->OnDistributeNodesH();
+}
+
+void SDataflowGraphEditor::DistributeVertically()
+{
+	GetGraphEditor()->OnDistributeNodesV();
+}
+
+void SDataflowGraphEditor::ToggleEnabledState()
+{
+	FDataflowEditorCommands::ToggleEnabledState(DataflowAsset.Get());
+}
+
 #undef LOCTEXT_NAMESPACE
diff --git a/Engine/Plugins/Experimental/Dataflow/Source/DataflowEditor/Private/Dataflow/DataflowSNodeFactories.cpp b/Engine/Plugins/Experimental/Dataflow/Source/DataflowEditor/Private/Dataflow/DataflowSNodeFactories.cpp
--- a/Engine/Plugins/Experimental/Dataflow/Source/DataflowEditor/Private/Dataflow/DataflowSNodeFactories.cpp
+++ b/Engine/Plugins/Experimental/Dataflow/Source/DataflowEditor/Private/Dataflow/DataflowSNodeFactories.cpp
@@ -3,12 +3,18 @@
 #include "Dataflow/DataflowSNodeFactories.h"
 #include "Dataflow/DataflowSNode.h"
 #include "Dataflow/DataflowEdNode.h"
+#include "EdGraphNode_Comment.h"
+#include "Dataflow/DataflowCommentNode.h"
 
 TSharedPtr<class SGraphNode> FDataflowSNodeFactory::CreateNode(UEdGraphNode* InNode) const
 {
 	if (UDataflowEdNode* Node = Cast<UDataflowEdNode>(InNode))
 	{
 		return SNew(SDataflowEdNode, Node);
+	}	
+	else if (UEdGraphNode_Comment* CommentNode = Cast<UEdGraphNode_Comment>(InNode))
+	{
+		return SNew(SDataflowEdNodeComment, CommentNode);
 	}
 	return NULL;
 }
diff --git a/Engine/Plugins/Experimental/Dataflow/Source/DataflowEditor/Private/Dataflow/DataflowSchema.cpp b/Engine/Plugins/Experimental/Dataflow/Source/DataflowEditor/Private/Dataflow/DataflowSchema.cpp
--- a/Engine/Plugins/Experimental/Dataflow/Source/DataflowEditor/Private/Dataflow/DataflowSchema.cpp
+++ b/Engine/Plugins/Experimental/Dataflow/Source/DataflowEditor/Private/Dataflow/DataflowSchema.cpp
@@ -14,6 +14,7 @@
 #include "ToolMenuSection.h"
 #include "ToolMenu.h"
 #include "GraphEditorActions.h"
+#include "Dataflow/DataflowSettings.h"
 
 #define LOCTEXT_NAMESPACE "DataflowNode"
 
@@ -25,11 +26,44 @@ void UDataflowSchema::GetContextMenuActions(class UToolMenu* Menu, class UGraphN
 {
 	if (Context->Node)
 	{
-		FToolMenuSection& Section = Menu->AddSection("TestGraphSchemaNodeActions", LOCTEXT("ClassActionsMenuHeader", "Node Actions"));
 		{
-			Section.AddMenuEntry(FGenericCommands::Get().Delete);
-			Section.AddMenuEntry(FGraphEditorCommands::Get().BreakNodeLinks);
-			Section.AddMenuEntry(FDataflowEditorCommands::Get().EvaluateNode);
+			FToolMenuSection& Section = Menu->AddSection("TestGraphSchemaNodeActions", LOCTEXT("ClassActionsMenuHeader", "Node Actions"));
+			{
+				Section.AddMenuEntry(FGenericCommands::Get().Rename);
+				Section.AddMenuEntry(FGenericCommands::Get().Delete);
+				Section.AddMenuEntry(FGenericCommands::Get().Cut);
+				Section.AddMenuEntry(FGenericCommands::Get().Copy);
+				Section.AddMenuEntry(FGenericCommands::Get().Duplicate);
+				Section.AddMenuEntry(FDataflowEditorCommands::Get().ToggleEnabledState, FText::FromString("Toggle Enabled State"));
+				Section.AddMenuEntry(FGraphEditorCommands::Get().BreakNodeLinks);
+				Section.AddMenuEntry(FDataflowEditorCommands::Get().EvaluateNode);
+			}
+		}
+
+		{
+			FToolMenuSection& Section = Menu->AddSection("TestGraphSchemaOrganization", LOCTEXT("ClassActionsMenuHeader", "Organization"));
+			{
+				Section.AddSubMenu("Alignment", LOCTEXT("AlignmentHeader", "Alignment"), FText(), FNewToolMenuDelegate::CreateLambda([](UToolMenu* AlignmentMenu)
+				{
+					{
+						FToolMenuSection& InSection = AlignmentMenu->AddSection("TestGraphSchemaAlignment", LOCTEXT("AlignHeader", "Align"));
+
+						InSection.AddMenuEntry(FGraphEditorCommands::Get().AlignNodesTop);
+						InSection.AddMenuEntry(FGraphEditorCommands::Get().AlignNodesMiddle);
+						InSection.AddMenuEntry(FGraphEditorCommands::Get().AlignNodesBottom);
+						InSection.AddMenuEntry(FGraphEditorCommands::Get().AlignNodesLeft);
+						InSection.AddMenuEntry(FGraphEditorCommands::Get().AlignNodesCenter);
+						InSection.AddMenuEntry(FGraphEditorCommands::Get().AlignNodesRight);
+						InSection.AddMenuEntry(FGraphEditorCommands::Get().StraightenConnections);
+					}
+
+					{
+						FToolMenuSection& InSection = AlignmentMenu->AddSection("TestGraphSchemaDistribution", LOCTEXT("DistributionHeader", "Distribution"));
+						InSection.AddMenuEntry(FGraphEditorCommands::Get().DistributeNodesHorizontally);
+						InSection.AddMenuEntry(FGraphEditorCommands::Get().DistributeNodesVertically);
+					}
+				}));
+			}
 		}
 	}
 
@@ -100,5 +134,90 @@ const FPinConnectionResponse UDataflowSchema::CanCreateConnection(const UEdGraph
 	return FPinConnectionResponse(CONNECT_RESPONSE_DISALLOW, NoConnectionResponse[FMath::RandRange(0, NoConnectionResponse.Num()-1)]);
 }
 
+FLinearColor UDataflowSchema::GetPinTypeColor(const FEdGraphPinType& PinType) const
+{
+	return GetTypeColor(PinType.PinCategory);
+}
+
+FLinearColor UDataflowSchema::GetTypeColor(const FName& Type)
+{
+	const UGraphEditorSettings* Settings = GetDefault<UGraphEditorSettings>();
+	const UDataflowSettings* DataflowSettings = GetDefault<UDataflowSettings>();
+
+	if (Type == FName("FManagedArrayCollection"))
+	{
+		return DataflowSettings->ManagedArrayCollectionPinTypeColor;
+	}
+	else if (Type == FName("float"))
+	{
+		return Settings->FloatPinTypeColor;
+	}
+	else if (Type == FName("int32"))
+	{
+		return Settings->IntPinTypeColor;
+	}
+	else if (Type == FName("bool"))
+	{
+		return Settings->BooleanPinTypeColor;
+	}
+	else if (Type == FName("FString"))
+	{
+		return Settings->StringPinTypeColor;
+	}
+	else if (Type == FName("FVector"))
+	{
+		return Settings->VectorPinTypeColor;
+	}
+	else if (Type == FName("TArray"))
+	{
+		return DataflowSettings->ArrayPinTypeColor;
+	}
+	else if (Type == FName("FBox"))
+	{
+		return DataflowSettings->BoxPinTypeColor;
+	}
+
+	return Settings->DefaultPinTypeColor;
+}
+
+//void UDataflowSchema::OnPinConnectionDoubleCicked(UEdGraphPin* PinA, UEdGraphPin* PinB, const FVector2D& GraphPosition) const
+//{
+//}
+
+FConnectionDrawingPolicy* UDataflowSchema::CreateConnectionDrawingPolicy(int32 InBackLayerID, int32 InFrontLayerID, float InZoomFactor, const FSlateRect& InClippingRect, class FSlateWindowElementList& InDrawElements, class UEdGraph* InGraphObj) const
+{
+	return new FDataflowConnectionDrawingPolicy(InBackLayerID, InFrontLayerID, InZoomFactor, InClippingRect, InDrawElements, InGraphObj);
+}
+
+FDataflowConnectionDrawingPolicy::FDataflowConnectionDrawingPolicy(int32 InBackLayerID, int32 InFrontLayerID, float InZoomFactor, const FSlateRect& InClippingRect, FSlateWindowElementList& InDrawElements, UEdGraph* InGraph)
+	: FConnectionDrawingPolicy(InBackLayerID, InFrontLayerID, InZoomFactor, InClippingRect, InDrawElements)
+	, Schema((UDataflowSchema*)(InGraph->GetSchema()))
+{
+	ArrowImage = nullptr;
+	ArrowRadius = FVector2D::ZeroVector;
+}
+
+void FDataflowConnectionDrawingPolicy::DetermineWiringStyle(UEdGraphPin* OutputPin, UEdGraphPin* InputPin, /*inout*/ FConnectionParams& Params)
+{
+	FConnectionDrawingPolicy::DetermineWiringStyle(OutputPin, InputPin, Params);
+	if (HoveredPins.Contains(InputPin) && HoveredPins.Contains(OutputPin))
+	{
+		Params.WireThickness = Params.WireThickness * 5;
+	}
+
+	const UDataflowSchema* DataflowSchema = GetSchema();
+	if (DataflowSchema && OutputPin)
+	{
+		Params.WireColor = DataflowSchema->GetPinTypeColor(OutputPin->PinType);
+	}
+
+	if (OutputPin && InputPin)
+	{
+		if (OutputPin->bOrphanedPin || InputPin->bOrphanedPin)
+		{
+			Params.WireColor = FLinearColor::Red;
+		}
+	}
+}
 
 #undef LOCTEXT_NAMESPACE
diff --git a/Engine/Plugins/Experimental/Dataflow/Source/DataflowEditor/Private/Dataflow/DataflowSettings.cpp b/Engine/Plugins/Experimental/Dataflow/Source/DataflowEditor/Private/Dataflow/DataflowSettings.cpp
new file mode 100644
--- /dev/null
+++ b/Engine/Plugins/Experimental/Dataflow/Source/DataflowEditor/Private/Dataflow/DataflowSettings.cpp
@@ -0,0 +1,40 @@
+// Copyright Epic Games, Inc. All Rights Reserved.
+
+#include "Dataflow/DataflowSettings.h"
+
+UDataflowSettings::UDataflowSettings(const FObjectInitializer& ObjectInitlaizer)
+	: Super(ObjectInitlaizer)
+{
+	ManagedArrayCollectionPinTypeColor = FLinearColor(0.353393f, 0.454175f, 1.0f, 1.0f);
+	ArrayPinTypeColor = FLinearColor(1.0f, 0.172585f, 0.0f, 1.0f);
+	BoxPinTypeColor = FLinearColor(0.013575f, 0.770000f, 0.429609f, 1.0f);
+}
+
+FName UDataflowSettings::GetCategoryName() const
+{
+	return TEXT("Plugins");
+}
+
+FText UDataflowSettings::GetSectionText() const
+{
+	return NSLOCTEXT("DataflowPlugin", "DataflowSettingsSection", "Dataflow");
+}
+
+void UDataflowSettings::PostEditChangeProperty(FPropertyChangedEvent& PropertyChangedEvent)
+{
+	if (PropertyChangedEvent.Property != nullptr)
+	{
+		SettingsChangedDelegate.Broadcast(PropertyChangedEvent.Property->GetFName(), this);
+	}
+}
+
+UDataflowSettings::FOnDataflowSettingsChanged& UDataflowSettings::OnSettingsChanged()
+{
+	return SettingsChangedDelegate;
+}
+
+UDataflowSettings::FOnDataflowSettingsChanged UDataflowSettings::SettingsChangedDelegate;
+
+
+
+
diff --git a/Engine/Plugins/Experimental/Dataflow/Source/DataflowEditor/Public/Dataflow/DataflowCommentNode.h b/Engine/Plugins/Experimental/Dataflow/Source/DataflowEditor/Public/Dataflow/DataflowCommentNode.h
new file mode 100644
--- /dev/null
+++ b/Engine/Plugins/Experimental/Dataflow/Source/DataflowEditor/Public/Dataflow/DataflowCommentNode.h
@@ -0,0 +1,47 @@
+// Copyright Epic Games, Inc. All Rights Reserved.
+
+#pragma once
+
+#include "EdGraphUtilities.h"
+#include "SGraphNode.h"
+#include "SGraphNodeComment.h"
+
+#include "DataflowCommentNode.generated.h"
+
+//
+// SDataflowEdNodeComment
+//
+
+class DATAFLOWEDITOR_API SDataflowEdNodeComment : public SGraphNodeComment
+{
+	typedef SGraphNodeComment Super;
+
+public:
+	SLATE_BEGIN_ARGS(SDataflowEdNodeComment) {}
+	SLATE_END_ARGS()
+
+	void Construct(const FArguments& InArgs, UEdGraphNode_Comment* InNode) { Super::Construct(SGraphNodeComment::FArguments(), InNode); }
+};
+
+//
+// Action to add a node to the graph
+//
+USTRUCT()
+struct DATAFLOWEDITOR_API FAssetSchemaAction_Dataflow_CreateCommentNode_DataflowEdNode : public FEdGraphSchemaAction
+{
+	GENERATED_USTRUCT_BODY();
+
+public:
+	FAssetSchemaAction_Dataflow_CreateCommentNode_DataflowEdNode(const TSharedPtr<SGraphEditor>& InGraphEditor) : FEdGraphSchemaAction()
+		, GraphEditor(InGraphEditor)
+		{}
+
+	FAssetSchemaAction_Dataflow_CreateCommentNode_DataflowEdNode() {}
+
+	static TSharedPtr<FAssetSchemaAction_Dataflow_CreateCommentNode_DataflowEdNode> CreateAction(UEdGraph* ParentGraph, const TSharedPtr<SGraphEditor>& GraphEditor);
+
+	virtual UEdGraphNode* PerformAction(class UEdGraph* ParentGraph, UEdGraphPin* FromPin, const FVector2D Location, bool bSelectNewNode = true) override;
+	//virtual void AddReferencedObjects(FReferenceCollector& Collector) override;
+
+	TSharedPtr<SGraphEditor> GraphEditor;
+};
diff --git a/Engine/Plugins/Experimental/Dataflow/Source/DataflowEditor/Public/Dataflow/DataflowEditorActions.h b/Engine/Plugins/Experimental/Dataflow/Source/DataflowEditor/Public/Dataflow/DataflowEditorActions.h
--- a/Engine/Plugins/Experimental/Dataflow/Source/DataflowEditor/Public/Dataflow/DataflowEditorActions.h
+++ b/Engine/Plugins/Experimental/Dataflow/Source/DataflowEditor/Public/Dataflow/DataflowEditorActions.h
@@ -34,6 +34,9 @@ public:
 	DATAFLOWEDITOR_API virtual void RegisterCommands() override;
 
 	TSharedPtr< FUICommandInfo > EvaluateNode;
+	TSharedPtr< FUICommandInfo > CreateComment;
+	TSharedPtr< FUICommandInfo > ToggleEnabledState;
+
 	TMap< FName, TSharedPtr<FUICommandInfo> > CreateNodesMap;
 };
 
@@ -63,4 +66,9 @@ public:
 	*/
 	static void OnSelectedNodesChanged(TSharedPtr<IStructureDetailsView> PropertiesEditor, UObject* Asset, UDataflow* Graph, const TSet<UObject*>& NewSelection);
 
+	/*
+	*  ToggleEnabledState
+	*/
+	static void ToggleEnabledState(UDataflow* Graph);
+
 };
diff --git a/Engine/Plugins/Experimental/Dataflow/Source/DataflowEditor/Public/Dataflow/DataflowGraphEditor.h b/Engine/Plugins/Experimental/Dataflow/Source/DataflowEditor/Public/Dataflow/DataflowGraphEditor.h
--- a/Engine/Plugins/Experimental/Dataflow/Source/DataflowEditor/Public/Dataflow/DataflowGraphEditor.h
+++ b/Engine/Plugins/Experimental/Dataflow/Source/DataflowEditor/Public/Dataflow/DataflowGraphEditor.h
@@ -68,8 +68,40 @@ public:
 	/** */
 	void OnSelectedNodesChanged(const TSet<UObject*>& NewSelection);
 
+	/** */
+	void CreateComment();
+
+	/** */
+	void AlignTop();
+
+	/** */
+	void AlignMiddle();
+
+	/** */
+	void AlignBottom();
+
+	/** */
+	void AlignLeft();
+
+	/** */
+	void AlignCenter();
+
+	/** */
+	void AlignRight();
+
+	/** */
+	void StraightenConnections();
 
+	/** */
+	void DistributeHorizontally();
+
+	/** */
+	void DistributeVertically();
+
+	/** */
+	void ToggleEnabledState();
 
+	SGraphEditor* GetGraphEditor() { return (SGraphEditor*)this; }
 
 private:
 	FDataflowEditorCommands::FOnDragDropEventCallback OnDragDropEventCallback;
diff --git a/Engine/Plugins/Experimental/Dataflow/Source/DataflowEditor/Public/Dataflow/DataflowSchema.h b/Engine/Plugins/Experimental/Dataflow/Source/DataflowEditor/Public/Dataflow/DataflowSchema.h
--- a/Engine/Plugins/Experimental/Dataflow/Source/DataflowEditor/Public/Dataflow/DataflowSchema.h
+++ b/Engine/Plugins/Experimental/Dataflow/Source/DataflowEditor/Public/Dataflow/DataflowSchema.h
@@ -1,7 +1,10 @@
 // Copyright Epic Games, Inc. All Rights Reserved.
 
 #pragma once
+
 #include "EdGraph/EdGraphSchema.h"
+#include "ConnectionDrawingPolicy.h"
+#include "Dataflow/DataflowGraph.h"
 
 #include "DataflowSchema.generated.h"
 
@@ -18,5 +21,25 @@ public:
 	virtual const FPinConnectionResponse CanCreateConnection(const UEdGraphPin* PinA, const UEdGraphPin* PinB) const override;
 	virtual void GetContextMenuActions(class UToolMenu* Menu, class UGraphNodeContextMenuContext* Context) const override;
 	virtual void GetGraphContextActions(FGraphContextMenuBuilder& ContextMenuBuilder) const override;
+
+	virtual FLinearColor GetPinTypeColor(const FEdGraphPinType& PinType) const override;
+	virtual FConnectionDrawingPolicy* CreateConnectionDrawingPolicy(int32 InBackLayerID, int32 InFrontLayerID, float InZoomFactor, const FSlateRect& InClippingRect, class FSlateWindowElementList& InDrawElements, class UEdGraph* InGraphObj) const override;
+	//virtual void OnPinConnectionDoubleCicked(UEdGraphPin* PinA, UEdGraphPin* PinB, const FVector2D& GraphPosition) const override;
+
 	//~ End EdGraphSchema Interface
+
+	static FLinearColor GetTypeColor(const FName& Type);
 };
+
+class FDataflowConnectionDrawingPolicy : public FConnectionDrawingPolicy
+{
+public:
+	FDataflowConnectionDrawingPolicy(int32 InBackLayerID, int32 InFrontLayerID, float InZoomFactor, const FSlateRect& InClippingRect, FSlateWindowElementList& InDrawElements, UEdGraph* InGraph);
+	virtual void DetermineWiringStyle(UEdGraphPin* OutputPin, UEdGraphPin* InputPin, /*inout*/ FConnectionParams& Params) override;
+
+	const UDataflowSchema* GetSchema() { return Schema; }
+
+private:
+	class UDataflowSchema* Schema;
+};
+
diff --git a/Engine/Plugins/Experimental/Dataflow/Source/DataflowEditor/Public/Dataflow/DataflowSettings.h b/Engine/Plugins/Experimental/Dataflow/Source/DataflowEditor/Public/Dataflow/DataflowSettings.h
new file mode 100644
--- /dev/null
+++ b/Engine/Plugins/Experimental/Dataflow/Source/DataflowEditor/Public/Dataflow/DataflowSettings.h
@@ -0,0 +1,46 @@
+// Copyright Epic Games, Inc. All Rights Reserved.
+
+#pragma once
+
+#include "Engine/DeveloperSettings.h"
+#include "InputCoreTypes.h"
+#include "DataflowSettings.generated.h"
+
+UCLASS(config = Dataflow, defaultconfig, meta = (DisplayName = "Dataflow"))
+class DATAFLOWEDITOR_API UDataflowSettings : public UDeveloperSettings
+{
+	GENERATED_UCLASS_BODY()
+
+	/** TArray<> pin type color. The other pin colors are defined in the general editor settings. */
+	UPROPERTY(config, EditAnywhere, Category = PinColors)
+	FLinearColor ArrayPinTypeColor;
+
+	/** FManagedArrayCollection pin type color. The other pin colors are defined in the general editor settings. */
+	UPROPERTY(config, EditAnywhere, Category = PinColors)
+	FLinearColor ManagedArrayCollectionPinTypeColor;
+
+	/** FBox pin type color. The other pin colors are defined in the general editor settings. */
+	UPROPERTY(config, EditAnywhere, Category = PinColors)
+	FLinearColor BoxPinTypeColor;
+
+	// Begin UDeveloperSettings Interface
+	virtual FName GetCategoryName() const override;
+#if WITH_EDITOR
+	virtual FText GetSectionText() const override;
+#endif
+	// END UDeveloperSettings Interface
+
+#if WITH_EDITOR
+	virtual void PostEditChangeProperty(FPropertyChangedEvent& PropertyChangedEvent) override;
+
+public:
+	DECLARE_MULTICAST_DELEGATE_TwoParams(FOnDataflowSettingsChanged, const FName&, const UDataflowSettings*);
+
+	/** Gets a multicast delegate which is called whenever one of the parameters in this settings object changes. */
+	static FOnDataflowSettingsChanged& OnSettingsChanged();
+
+protected:
+	static FOnDataflowSettingsChanged SettingsChangedDelegate;
+
+#endif
+};
diff --git a/Engine/Source/Runtime/Experimental/Dataflow/Engine/Private/Dataflow/DataflowObject.cpp b/Engine/Source/Runtime/Experimental/Dataflow/Engine/Private/Dataflow/DataflowObject.cpp
--- a/Engine/Source/Runtime/Experimental/Dataflow/Engine/Private/Dataflow/DataflowObject.cpp
+++ b/Engine/Source/Runtime/Experimental/Dataflow/Engine/Private/Dataflow/DataflowObject.cpp
@@ -54,8 +54,8 @@ void UDataflow::PostLoad()
 
 	for (UEdGraphNode* EdNode : Nodes)
 	{
-		UDataflowEdNode* DataflowEdNode = Cast<UDataflowEdNode>(EdNode);
-		if(ensure(DataflowEdNode))
+		// Not all nodes are UDataflowEdNode (There is now UDataflowEdNodeComment)
+		if (UDataflowEdNode* DataflowEdNode = Cast<UDataflowEdNode>(EdNode))
 		{
 			DataflowEdNode->SetDataflowGraph(Dataflow);
 		}
