Chaos : Dataflow graph  : fix loctext conflicts

diff --git a/Engine/Plugins/Experimental/Dataflow/Source/DataflowEditor/Private/Dataflow/DataflowSchema.cpp b/Engine/Plugins/Experimental/Dataflow/Source/DataflowEditor/Private/Dataflow/DataflowSchema.cpp
--- a/Engine/Plugins/Experimental/Dataflow/Source/DataflowEditor/Private/Dataflow/DataflowSchema.cpp
+++ b/Engine/Plugins/Experimental/Dataflow/Source/DataflowEditor/Private/Dataflow/DataflowSchema.cpp
@@ -27,7 +27,7 @@ void UDataflowSchema::GetContextMenuActions(class UToolMenu* Menu, class UGraphN
 	if (Context->Node)
 	{
 		{
-			FToolMenuSection& Section = Menu->AddSection("TestGraphSchemaNodeActions", LOCTEXT("ClassActionsMenuHeader", "Node Actions"));
+			FToolMenuSection& Section = Menu->AddSection("TestGraphSchemaNodeActions", LOCTEXT("GraphSchemaNodeActions_MenuHeader", "Node Actions"));
 			{
 				Section.AddMenuEntry(FGenericCommands::Get().Rename);
 				Section.AddMenuEntry(FGenericCommands::Get().Delete);
@@ -41,12 +41,12 @@ void UDataflowSchema::GetContextMenuActions(class UToolMenu* Menu, class UGraphN
 		}
 
 		{
-			FToolMenuSection& Section = Menu->AddSection("TestGraphSchemaOrganization", LOCTEXT("ClassActionsMenuHeader", "Organization"));
+			FToolMenuSection& Section = Menu->AddSection("TestGraphSchemaOrganization", LOCTEXT("GraphSchemaOrganization_MenuHeader", "Organization"));
 			{
 				Section.AddSubMenu("Alignment", LOCTEXT("AlignmentHeader", "Alignment"), FText(), FNewToolMenuDelegate::CreateLambda([](UToolMenu* AlignmentMenu)
 				{
 					{
-						FToolMenuSection& InSection = AlignmentMenu->AddSection("TestGraphSchemaAlignment", LOCTEXT("AlignHeader", "Align"));
+						FToolMenuSection& InSection = AlignmentMenu->AddSection("TestGraphSchemaAlignment", LOCTEXT("GraphSchemaAlignment_MenuHeader", "Align"));
 
 						InSection.AddMenuEntry(FGraphEditorCommands::Get().AlignNodesTop);
 						InSection.AddMenuEntry(FGraphEditorCommands::Get().AlignNodesMiddle);
@@ -58,7 +58,7 @@ void UDataflowSchema::GetContextMenuActions(class UToolMenu* Menu, class UGraphN
 					}
 
 					{
-						FToolMenuSection& InSection = AlignmentMenu->AddSection("TestGraphSchemaDistribution", LOCTEXT("DistributionHeader", "Distribution"));
+						FToolMenuSection& InSection = AlignmentMenu->AddSection("TestGraphSchemaDistribution", LOCTEXT("GraphSchemaDistribution_MenuHeader", "Distribution"));
 						InSection.AddMenuEntry(FGraphEditorCommands::Get().DistributeNodesHorizontally);
 						InSection.AddMenuEntry(FGraphEditorCommands::Get().DistributeNodesVertically);
 					}
