Dataflow : Editor Style

diff --git a/Engine/Plugins/Experimental/Dataflow/Source/DataflowEditor/DataflowEditor.Build.cs b/Engine/Plugins/Experimental/Dataflow/Source/DataflowEditor/DataflowEditor.Build.cs
--- a/Engine/Plugins/Experimental/Dataflow/Source/DataflowEditor/DataflowEditor.Build.cs
+++ b/Engine/Plugins/Experimental/Dataflow/Source/DataflowEditor/DataflowEditor.Build.cs
@@ -19,7 +19,8 @@ namespace UnrealBuildTool.Rules
 				    "Engine",
 					"EditorFramework",
 					"UnrealEd",
-				    "PropertyEditor",
+					"Projects",
+					"PropertyEditor",
 				    "RenderCore",
 				    "RHI",
 				    "AssetTools",
diff --git a/Engine/Plugins/Experimental/Dataflow/Source/DataflowEditor/Private/Dataflow/DataflowEditorPlugin.cpp b/Engine/Plugins/Experimental/Dataflow/Source/DataflowEditor/Private/Dataflow/DataflowEditorPlugin.cpp
--- a/Engine/Plugins/Experimental/Dataflow/Source/DataflowEditor/Private/Dataflow/DataflowEditorPlugin.cpp
+++ b/Engine/Plugins/Experimental/Dataflow/Source/DataflowEditor/Private/Dataflow/DataflowEditorPlugin.cpp
@@ -1,6 +1,7 @@
 // Copyright Epic Games, Inc. All Rights Reserved.
 
 #include "Dataflow/DataflowEditorPlugin.h"
+#include "Dataflow/DataflowEditorStyle.h"
 
 #include "AssetToolsModule.h"
 #include "CoreMinimal.h"
@@ -17,6 +18,8 @@
 
 void IDataflowEditorPlugin::StartupModule()
 {
+	FDataflowEditorStyle::Get();
+
 	DataflowAssetActions = new FDataflowAssetActions();
 
 	FAssetToolsModule& AssetToolsModule = FAssetToolsModule::GetModule();
diff --git a/Engine/Plugins/Experimental/Dataflow/Source/DataflowEditor/Private/Dataflow/DataflowEditorStyle.h b/Engine/Plugins/Experimental/Dataflow/Source/DataflowEditor/Private/Dataflow/DataflowEditorStyle.h
new file mode 100644
--- /dev/null
+++ b/Engine/Plugins/Experimental/Dataflow/Source/DataflowEditor/Private/Dataflow/DataflowEditorStyle.h
@@ -0,0 +1,47 @@
+// Copyright Epic Games, Inc. All Rights Reserved.
+
+#pragma once
+#include "Styling/SlateStyle.h"
+#include "Styling/SlateStyleRegistry.h"
+#include "Interfaces/IPluginManager.h"
+
+class FDataflowEditorStyle final : public FSlateStyleSet
+{
+public:
+	FDataflowEditorStyle() : FSlateStyleSet("DataflowEditorStyle")
+	{
+		const FVector2D Icon16x16(16.f, 16.f);
+		const FVector2D Icon64x64(64.f, 64.f);
+
+		SetContentRoot(IPluginManager::Get().FindPlugin("Dataflow")->GetBaseDir() / TEXT("Resources"));
+		Set("ClassIcon.Dataflow", new FSlateVectorImageBrush(RootToContentDir(TEXT("DataflowAsset_16.svg")), Icon16x16));
+		Set("ClassThumbnail.Dataflow", new FSlateVectorImageBrush(RootToContentDir(TEXT("DataflowAsset_64.svg")), Icon64x64));
+		FSlateStyleRegistry::RegisterSlateStyle(*this);
+	}
+
+	~FDataflowEditorStyle()
+	{
+		FSlateStyleRegistry::UnRegisterSlateStyle(*this);
+	}
+
+public:
+
+	static FDataflowEditorStyle& Get()
+	{
+		if (!Singleton.IsSet())
+		{
+			Singleton.Emplace();
+		}
+		return Singleton.GetValue();
+	}
+
+	static void Destroy()
+	{
+		Singleton.Reset();
+	}
+
+private:
+	static TOptional<FDataflowEditorStyle> Singleton;
+};
+
+TOptional<FDataflowEditorStyle> FDataflowEditorStyle::Singleton;
\ No newline at end of file
