Dataflow : Delete
- Fixed inability to delete disabled nodes.

diff --git a/Engine/Plugins/Experimental/Dataflow/Source/DataflowEditor/Private/Dataflow/DataflowEditorActions.cpp b/Engine/Plugins/Experimental/Dataflow/Source/DataflowEditor/Private/Dataflow/DataflowEditorActions.cpp
--- a/Engine/Plugins/Experimental/Dataflow/Source/DataflowEditor/Private/Dataflow/DataflowEditorActions.cpp
+++ b/Engine/Plugins/Experimental/Dataflow/Source/DataflowEditor/Private/Dataflow/DataflowEditorActions.cpp
@@ -84,22 +84,25 @@ void FDataflowEditorCommands::EvaluateNodes(const FGraphPanelSelectionSet& Selec
 
 void FDataflowEditorCommands::DeleteNodes(UDataflow* Graph, const FGraphPanelSelectionSet& SelectedNodes)
 {
-	for (UObject* Ode : SelectedNodes)
+	if (ensureMsgf(Graph != nullptr, TEXT("Warning : Failed to find valid graph.")))
 	{
-		if (UDataflowEdNode* EdNode = dynamic_cast<UDataflowEdNode*>(Ode))
+		for (UObject* Ode : SelectedNodes)
 		{
-			if (const TSharedPtr<Dataflow::FGraph> DataflowGraph = EdNode->GetDataflowGraph())
+			if (UDataflowEdNode* EdNode = dynamic_cast<UDataflowEdNode*>(Ode))
 			{
-				if (TSharedPtr<FDataflowNode> DataflowNode = DataflowGraph->FindBaseNode(EdNode->GetDataflowNodeGuid()))
+				if (const TSharedPtr<Dataflow::FGraph> DataflowGraph = EdNode->GetDataflowGraph())
 				{
 					Graph->RemoveNode(EdNode);
-					DataflowGraph->RemoveNode(DataflowNode);
+					if (TSharedPtr<FDataflowNode> DataflowNode = DataflowGraph->FindBaseNode(EdNode->GetDataflowNodeGuid()))
+					{
+						DataflowGraph->RemoveNode(DataflowNode);
+					}
 				}
 			}
-		}
-		else if (UEdGraphNode_Comment* CommentNode = dynamic_cast<UEdGraphNode_Comment*>(Ode))
-		{
-			Graph->RemoveNode(CommentNode);
+			else if (UEdGraphNode_Comment* CommentNode = dynamic_cast<UEdGraphNode_Comment*>(Ode))
+			{
+				Graph->RemoveNode(CommentNode);
+			}
 		}
 	}
 }
