GeometryCollection : Append Skeleton
- Added nodes to append a skeleton to the geometry collection.

diff --git a/Engine/Plugins/Experimental/Dataflow/Source/DataflowNodes/Private/Dataflow/DataflowSkeletalMeshNodes.cpp b/Engine/Plugins/Experimental/Dataflow/Source/DataflowNodes/Private/Dataflow/DataflowSkeletalMeshNodes.cpp
--- a/Engine/Plugins/Experimental/Dataflow/Source/DataflowNodes/Private/Dataflow/DataflowSkeletalMeshNodes.cpp
+++ b/Engine/Plugins/Experimental/Dataflow/Source/DataflowNodes/Private/Dataflow/DataflowSkeletalMeshNodes.cpp
@@ -14,6 +14,7 @@ namespace Dataflow
 	void RegisterSkeletalMeshNodes()
 	{
 		DATAFLOW_NODE_REGISTER_CREATION_FACTORY(FGetSkeletalMeshDataflowNode);
+		DATAFLOW_NODE_REGISTER_CREATION_FACTORY(FGetSkeletonDataflowNode);
 		DATAFLOW_NODE_REGISTER_CREATION_FACTORY(FSkeletalMeshBoneDataflowNode);
 	}
 }
@@ -25,16 +26,37 @@ void FGetSkeletalMeshDataflowNode::Evaluate(Dataflow::FContext& Context, const F
 	{
 		SetValue<DataType>(Context, SkeletalMesh, &SkeletalMesh);
 
-		if (SkeletalMesh)
+		if (!SkeletalMesh)
 		{
-			SetValue<DataType>(Context, SkeletalMesh, &SkeletalMesh);
+			if (const Dataflow::FEngineContext* EngineContext = Context.AsType<Dataflow::FEngineContext>())
+			{
+				if (const USkeletalMesh* SkeletalMeshFromOwner = Dataflow::Reflection::FindObjectPtrProperty<USkeletalMesh>(
+					EngineContext->Owner, PropertyName))
+				{
+					SetValue<DataType>(Context, DataType(SkeletalMeshFromOwner), &SkeletalMesh);
+				}
+			}
 		}
-		else if (const Dataflow::FEngineContext* EngineContext = Context.AsType<Dataflow::FEngineContext>())
+	}
+}
+
+
+void FGetSkeletonDataflowNode::Evaluate(Dataflow::FContext& Context, const FDataflowOutput* Out) const
+{
+	typedef TObjectPtr<const USkeleton> DataType;
+	if (Out->IsA<DataType>(&Skeleton))
+	{
+		SetValue<DataType>(Context, Skeleton, &Skeleton);
+
+		if (!Skeleton)
 		{
-			if (const USkeletalMesh* SkeletalMeshFromOwner = Dataflow::Reflection::FindObjectPtrProperty<USkeletalMesh>(
-				EngineContext->Owner, PropertyName))
+			if (const Dataflow::FEngineContext* EngineContext = Context.AsType<Dataflow::FEngineContext>())
 			{
-				SetValue<DataType>(Context, DataType(SkeletalMeshFromOwner), &SkeletalMesh);
+				if (const USkeleton* SkeletonFromOwner = Dataflow::Reflection::FindObjectPtrProperty<USkeleton>(
+					EngineContext->Owner, PropertyName))
+				{
+					SetValue<DataType>(Context, DataType(SkeletonFromOwner), &Skeleton);
+				}
 			}
 		}
 	}
diff --git a/Engine/Plugins/Experimental/Dataflow/Source/DataflowNodes/Public/Dataflow/DataflowSkeletalMeshNodes.h b/Engine/Plugins/Experimental/Dataflow/Source/DataflowNodes/Public/Dataflow/DataflowSkeletalMeshNodes.h
--- a/Engine/Plugins/Experimental/Dataflow/Source/DataflowNodes/Public/Dataflow/DataflowSkeletalMeshNodes.h
+++ b/Engine/Plugins/Experimental/Dataflow/Source/DataflowNodes/Public/Dataflow/DataflowSkeletalMeshNodes.h
@@ -16,7 +16,7 @@ USTRUCT()
 struct DATAFLOWNODES_API FGetSkeletalMeshDataflowNode: public FDataflowNode
 {
 	GENERATED_USTRUCT_BODY()
-	DATAFLOW_NODE_DEFINE_INTERNAL(FGetSkeletalMeshDataflowNode, "SkeletalMesh", "Dataflow", "Skeletal Mesh")
+	DATAFLOW_NODE_DEFINE_INTERNAL(FGetSkeletalMeshDataflowNode, "SkeletalMesh", "General", "Skeletal Mesh")
 
 public:
 	
@@ -36,12 +36,36 @@ public:
 	virtual void Evaluate(Dataflow::FContext& Context, const FDataflowOutput* Out) const override;
 };
 
+USTRUCT()
+struct DATAFLOWNODES_API FGetSkeletonDataflowNode : public FDataflowNode
+{
+	GENERATED_USTRUCT_BODY()
+	DATAFLOW_NODE_DEFINE_INTERNAL(FGetSkeletonDataflowNode, "Skeleton", "General", "Skeletal Mesh")
+
+public:
+
+	UPROPERTY(EditAnywhere, Category = "Dataflow", meta = (DataflowOutput, DisplayName = "Skeleton"))
+	TObjectPtr<const USkeleton> Skeleton = nullptr;
+
+	UPROPERTY(EditAnywhere, Category = "Dataflow")
+	FName PropertyName = "Skeleton";
+
+	FGetSkeletonDataflowNode(const Dataflow::FNodeParameters& InParam, FGuid InGuid = FGuid::NewGuid())
+		: FDataflowNode(InParam, InGuid)
+	{
+		RegisterOutputConnection(&Skeleton);
+	}
+
+
+	virtual void Evaluate(Dataflow::FContext& Context, const FDataflowOutput* Out) const override;
+};
+
 
 USTRUCT()
 struct DATAFLOWNODES_API FSkeletalMeshBoneDataflowNode : public FDataflowNode
 {
 	GENERATED_USTRUCT_BODY()
-	DATAFLOW_NODE_DEFINE_INTERNAL(FSkeletalMeshBoneDataflowNode, "SkeletalMeshBone", "Dataflow", "Skeletal Mesh")
+	DATAFLOW_NODE_DEFINE_INTERNAL(FSkeletalMeshBoneDataflowNode, "SkeletalMeshBone", "General", "Skeletal Mesh")
 
 public:
 	UPROPERTY(EditAnywhere, Category = "Dataflow")
diff --git a/Engine/Plugins/Experimental/Dataflow/Source/DataflowNodes/Public/Dataflow/DataflowStaticMeshNodes.h b/Engine/Plugins/Experimental/Dataflow/Source/DataflowNodes/Public/Dataflow/DataflowStaticMeshNodes.h
--- a/Engine/Plugins/Experimental/Dataflow/Source/DataflowNodes/Public/Dataflow/DataflowStaticMeshNodes.h
+++ b/Engine/Plugins/Experimental/Dataflow/Source/DataflowNodes/Public/Dataflow/DataflowStaticMeshNodes.h
@@ -16,7 +16,7 @@ USTRUCT()
 struct DATAFLOWNODES_API FGetStaticMeshDataflowNode : public FDataflowNode
 {
 	GENERATED_USTRUCT_BODY()
-	DATAFLOW_NODE_DEFINE_INTERNAL(FGetStaticMeshDataflowNode, "StaticMesh", "Dataflow", "Static Mesh")
+	DATAFLOW_NODE_DEFINE_INTERNAL(FGetStaticMeshDataflowNode, "StaticMesh", "General", "Static Mesh")
 
 public:
 
diff --git a/Engine/Plugins/Experimental/GeometryCollectionPlugin/Source/GeometryCollectionNodes/Private/Dataflow/GeometryCollectionNodes.cpp b/Engine/Plugins/Experimental/GeometryCollectionPlugin/Source/GeometryCollectionNodes/Private/Dataflow/GeometryCollectionNodes.cpp
--- a/Engine/Plugins/Experimental/GeometryCollectionPlugin/Source/GeometryCollectionNodes/Private/Dataflow/GeometryCollectionNodes.cpp
+++ b/Engine/Plugins/Experimental/GeometryCollectionPlugin/Source/GeometryCollectionNodes/Private/Dataflow/GeometryCollectionNodes.cpp
@@ -30,6 +30,7 @@ namespace Dataflow
 		DATAFLOW_NODE_REGISTER_CREATION_FACTORY(FGetCollectionAssetDataflowNode);
 		DATAFLOW_NODE_REGISTER_CREATION_FACTORY(FExampleCollectionEditDataflowNode);
 		DATAFLOW_NODE_REGISTER_CREATION_FACTORY(FSetCollectionAssetDataflowNode);
+		DATAFLOW_NODE_REGISTER_CREATION_FACTORY(FAppendCollectionAssetsDataflowNode);
 		DATAFLOW_NODE_REGISTER_CREATION_FACTORY(FResetGeometryCollectionDataflowNode);
 
 		DATAFLOW_NODE_REGISTER_CREATION_FACTORY(FPrintStringDataflowNode);
@@ -135,6 +136,17 @@ void FSetCollectionAssetDataflowNode::Evaluate(Dataflow::FContext& Context, cons
 	}
 }
 
+void FAppendCollectionAssetsDataflowNode::Evaluate(Dataflow::FContext& Context, const FDataflowOutput* Out) const
+{
+	if (Out->IsA<DataType>(&Collection1))
+	{
+		DataType InCollection1 = GetValue<DataType>(Context, &Collection1);
+		DataType InCollection2 = GetValue<DataType>(Context, &Collection2);
+		InCollection1.Append(InCollection2);
+		SetValue<DataType>(Context, Collection1, &Collection1); // prime to avoid ensure
+	}
+}
+
 void FResetGeometryCollectionDataflowNode::Evaluate(Dataflow::FContext& Context, const FDataflowOutput* Out) const
 {
 	if (Out->IsA<DataType>(&Collection))
diff --git a/Engine/Plugins/Experimental/GeometryCollectionPlugin/Source/GeometryCollectionNodes/Private/Dataflow/GeometryCollectionSkeletalMeshNodes.cpp b/Engine/Plugins/Experimental/GeometryCollectionPlugin/Source/GeometryCollectionNodes/Private/Dataflow/GeometryCollectionSkeletalMeshNodes.cpp
--- a/Engine/Plugins/Experimental/GeometryCollectionPlugin/Source/GeometryCollectionNodes/Private/Dataflow/GeometryCollectionSkeletalMeshNodes.cpp
+++ b/Engine/Plugins/Experimental/GeometryCollectionPlugin/Source/GeometryCollectionNodes/Private/Dataflow/GeometryCollectionSkeletalMeshNodes.cpp
@@ -6,29 +6,29 @@
 #include "GeometryCollection/GeometryCollectionObject.h"
 #include "GeometryCollection/ManagedArrayCollection.h"
 #include "GeometryCollection/GeometryCollection.h"
+#include "GeometryCollection/GeometryCollectionEngineConversion.h"
 
 
 namespace Dataflow
 {
 	void GeometryCollectionSkeletalMeshNodes()
 	{
-		DATAFLOW_NODE_REGISTER_CREATION_FACTORY(FSkeletalMeshToCollectionDataflowNode);
-
+		DATAFLOW_NODE_REGISTER_CREATION_FACTORY(FSkeletonToCollectionDataflowNode);
 	}
 }
 
-void FSkeletalMeshToCollectionDataflowNode::Evaluate(Dataflow::FContext& Context, const FDataflowOutput* Out) const
+void FSkeletonToCollectionDataflowNode::Evaluate(Dataflow::FContext& Context, const FDataflowOutput* Out) const
 {
 	if (Out->IsA<DataType>(&Collection))
 	{
-		DataType InCollection = GetValue<DataType>(Context, &Collection);
-
-		//
-		// @todo(dataflow) : Implemention conversion from skeletal mesh to TManagedArrayCollection
-		//
-
-		SetValue<DataType>(Context, InCollection, &Collection);
+		DataType OutCollection;
+		TObjectPtr<const USkeleton> SkeletonValue = GetValue<TObjectPtr<const USkeleton>>(Context, &Skeleton);
+		if (SkeletonValue)
+		{
+			FGeometryCollectionEngineConversion::AppendSkeleton(SkeletonValue.Get(), FTransform::Identity, &OutCollection);
+		}
+		SetValue<DataType>(Context, OutCollection, &Collection);
 	}
 }
 
-
+									
diff --git a/Engine/Plugins/Experimental/GeometryCollectionPlugin/Source/GeometryCollectionNodes/Public/Dataflow/GeometryCollectionNodes.h b/Engine/Plugins/Experimental/GeometryCollectionPlugin/Source/GeometryCollectionNodes/Public/Dataflow/GeometryCollectionNodes.h
--- a/Engine/Plugins/Experimental/GeometryCollectionPlugin/Source/GeometryCollectionNodes/Public/Dataflow/GeometryCollectionNodes.h
+++ b/Engine/Plugins/Experimental/GeometryCollectionPlugin/Source/GeometryCollectionNodes/Public/Dataflow/GeometryCollectionNodes.h
@@ -85,6 +85,35 @@ public:
 
 };
 
+
+USTRUCT()
+struct FAppendCollectionAssetsDataflowNode : public FDataflowNode
+{
+	GENERATED_USTRUCT_BODY()
+		DATAFLOW_NODE_DEFINE_INTERNAL(FAppendCollectionAssetsDataflowNode, "AppendCollections", "GeometryCollection", "")
+
+public:
+	typedef FManagedArrayCollection DataType;
+
+	UPROPERTY(meta = (DataflowInput, DisplayOutput, DisplayName = "Collection", Passthrough = "Collection"))
+	FManagedArrayCollection Collection1;
+
+	UPROPERTY(meta = (DataflowInput, DisplayName = "Collection"))
+	FManagedArrayCollection Collection2;
+
+
+	FAppendCollectionAssetsDataflowNode(const Dataflow::FNodeParameters& InParam, FGuid InGuid = FGuid::NewGuid())
+		: FDataflowNode(InParam, InGuid)
+	{
+		RegisterOutputConnection(&Collection1, &Collection1);
+		RegisterInputConnection(&Collection1);
+		RegisterInputConnection(&Collection2);
+	}
+
+	virtual void Evaluate(Dataflow::FContext& Context, const FDataflowOutput* Out) const override;
+
+};
+
 USTRUCT()
 struct FResetGeometryCollectionDataflowNode : public FDataflowNode
 {
diff --git a/Engine/Plugins/Experimental/GeometryCollectionPlugin/Source/GeometryCollectionNodes/Public/Dataflow/GeometryCollectionSkeletalMeshNodes.h b/Engine/Plugins/Experimental/GeometryCollectionPlugin/Source/GeometryCollectionNodes/Public/Dataflow/GeometryCollectionSkeletalMeshNodes.h
--- a/Engine/Plugins/Experimental/GeometryCollectionPlugin/Source/GeometryCollectionNodes/Public/Dataflow/GeometryCollectionSkeletalMeshNodes.h
+++ b/Engine/Plugins/Experimental/GeometryCollectionPlugin/Source/GeometryCollectionNodes/Public/Dataflow/GeometryCollectionSkeletalMeshNodes.h
@@ -11,26 +11,24 @@
 class USkeletalMesh;
 
 USTRUCT()
-struct FSkeletalMeshToCollectionDataflowNode : public FDataflowNode
+struct FSkeletonToCollectionDataflowNode : public FDataflowNode
 {
 	GENERATED_USTRUCT_BODY()
-	DATAFLOW_NODE_DEFINE_INTERNAL(FSkeletalMeshToCollectionDataflowNode, "SkeletalMeshToCollection", "GeometryCollection", "")
+	DATAFLOW_NODE_DEFINE_INTERNAL(FSkeletonToCollectionDataflowNode, "SkeletonToCollection", "GeometryCollection", "")
 
 public:
 	typedef FManagedArrayCollection DataType;
 
+	UPROPERTY(EditAnywhere, Category = "Dataflow", meta = (DataflowInput, DisplayName = "Skeleton"))
+	TObjectPtr<const USkeleton> Skeleton = nullptr;
 
-	UPROPERTY(meta = (DataflowInput, DisplayName = "SkeletalMesh"))
-	TObjectPtr<const USkeletalMesh> SkeletalMesh = nullptr;
-
-	UPROPERTY(meta = (DataflowInput, DataflowOutput, DisplayName = "Collection"))
+	UPROPERTY(meta = (DataflowOutput, DisplayName = "Collection"))
 	FManagedArrayCollection Collection;
 
-	FSkeletalMeshToCollectionDataflowNode(const Dataflow::FNodeParameters& InParam, FGuid InGuid = FGuid::NewGuid())
+	FSkeletonToCollectionDataflowNode(const Dataflow::FNodeParameters& InParam, FGuid InGuid = FGuid::NewGuid())
 		: FDataflowNode(InParam, InGuid)
 	{
-		RegisterInputConnection(&SkeletalMesh);
-		RegisterInputConnection(&Collection);
+		RegisterInputConnection(&Skeleton);
 		RegisterOutputConnection(&Collection);
 	}
 
