Dataflow : Double click to source.

diff --git a/Engine/Plugins/Experimental/Dataflow/Source/DataflowEditor/Private/Dataflow/DataflowSNode.cpp b/Engine/Plugins/Experimental/Dataflow/Source/DataflowEditor/Private/Dataflow/DataflowSNode.cpp
--- a/Engine/Plugins/Experimental/Dataflow/Source/DataflowEditor/Private/Dataflow/DataflowSNode.cpp
+++ b/Engine/Plugins/Experimental/Dataflow/Source/DataflowEditor/Private/Dataflow/DataflowSNode.cpp
@@ -7,6 +7,7 @@
 #include "Dataflow/DataflowObject.h"
 #include "Dataflow/DataflowCore.h"
 #include "Logging/LogMacros.h"
+#include "SourceCodeNavigation.h"
 #include "Widgets/Layout/SBorder.h"
 
 #define LOCTEXT_NAMESPACE "SDataflowEdNode"
@@ -22,6 +23,22 @@ void SDataflowEdNode::Construct(const FArguments& InArgs, UDataflowEdNode* InNod
 
 FReply SDataflowEdNode::OnMouseButtonDoubleClick(const FGeometry& InMyGeometry, const FPointerEvent& InMouseEvent)
 {
+	if (GraphNode)
+	{
+		if (UDataflowEdNode* DataflowNode = Cast<UDataflowEdNode>(GraphNode))
+		{
+			if (TSharedPtr<Dataflow::FGraph> Graph = DataflowNode->GetDataflowGraph())
+			{
+				if (TSharedPtr<FDataflowNode> Node = Graph->FindBaseNode(DataflowNode->GetDataflowNodeGuid()))
+				{
+					if (FSourceCodeNavigation::CanNavigateToStruct(Node->TypedScriptStruct()))
+					{
+						FSourceCodeNavigation::NavigateToStruct(Node->TypedScriptStruct());
+					}
+				}
+			}
+		}
+	}
 	return Super::OnMouseButtonDoubleClick(InMyGeometry, InMouseEvent);
 }
 
diff --git a/Engine/Source/Runtime/Experimental/Dataflow/Core/Public/Dataflow/DataflowNode.h b/Engine/Source/Runtime/Experimental/Dataflow/Core/Public/Dataflow/DataflowNode.h
--- a/Engine/Source/Runtime/Experimental/Dataflow/Core/Public/Dataflow/DataflowNode.h
+++ b/Engine/Source/Runtime/Experimental/Dataflow/Core/Public/Dataflow/DataflowNode.h
@@ -13,6 +13,7 @@
 
 struct FDataflowInput;
 struct FDataflowOutput;
+class UScriptStruct;
 
 namespace Dataflow {
 	struct DATAFLOWCORE_API FNodeParameters {
@@ -106,6 +107,7 @@ struct DATAFLOWCORE_API FDataflowNode
 
 	virtual void SerializeInternal(FArchive& Ar) { check(false); }
 	virtual FStructOnScope* NewStructOnScope() { return nullptr; }
+	virtual const UScriptStruct* TypedScriptStruct() const { return nullptr; }
 
 	/** Register the Input and Outputs after the creation in the factory */
 	void RegisterInputConnection(const void* Property);
@@ -216,6 +218,8 @@ public:																				\
 	virtual FName GetCategory() const override { return TYPE::StaticCategory(); }	\
 	virtual FString GetTags() const override { return TYPE::StaticTags(); }			\
 	virtual FString GetToolTip() const override { return TYPE::StaticToolTip(); }	\
+	virtual const UScriptStruct* TypedScriptStruct() const override					\
+		{return TYPE::StaticStruct();}												\
 	TYPE() {}																		\
 private:
 
