Dataflow : Editor
- Prevent evaluation when view dataflow asset through the generic editor.

diff --git a/Engine/Plugins/Experimental/Dataflow/Source/DataflowEditor/Private/Dataflow/DataflowEditorToolkit.cpp b/Engine/Plugins/Experimental/Dataflow/Source/DataflowEditor/Private/Dataflow/DataflowEditorToolkit.cpp
--- a/Engine/Plugins/Experimental/Dataflow/Source/DataflowEditor/Private/Dataflow/DataflowEditorToolkit.cpp
+++ b/Engine/Plugins/Experimental/Dataflow/Source/DataflowEditor/Private/Dataflow/DataflowEditorToolkit.cpp
@@ -18,6 +18,7 @@
 #include "Framework/Commands/GenericCommands.h"
 #include "Widgets/Docking/SDockTab.h"
 #include "IStructureDetailsView.h"
+#include "Misc/MessageDialog.h"
 
 #define LOCTEXT_NAMESPACE "DataflowEditorToolkit"
 
@@ -86,22 +87,10 @@ TSharedRef<SGraphEditor> FDataflowEditorToolkit::CreateGraphEditorWidget(UDatafl
 
 	FDataflowEditorCommands::FGraphEvaluationCallback Evaluate = [&](FDataflowNode* Node, FDataflowOutput* Out)
 	{
-		if (DataflowToEdit->bActive)
+		if (DataflowToEdit)
 		{
-			float EvalTime = FGameTime::GetTimeSinceAppStart().GetRealTimeSeconds();
-			if (DataflowToEdit->Targets.Num())
-			{
-				for (UObject* Object : DataflowToEdit->Targets)
-				{
-					FEngineContext Context(Object, Dataflow, EvalTime, FString("TargetObject"));
-					Node->Evaluate(Context, Out);
-				}
-			}
-			else
-			{
-				FEngineContext Context(DataflowToEdit, Dataflow, EvalTime, FString("UDataflow"));
-				Node->Evaluate(Context, Out);
-			}
+			FMessageDialog::Open(EAppMsgType::Ok, NSLOCTEXT("Dataflow", "UnknownEvaluationContext_Dataflow", 
+				"Unsupported dataflow evaluation.\nEvaluate the dataflow graph through its associated asset editor."));
 		}
 	};
 
@@ -139,6 +128,15 @@ TSharedPtr<IStructureDetailsView> FDataflowEditorToolkit::CreateNodeDetailsEdito
 	}
 	TSharedPtr<IStructureDetailsView> DetailsView = PropertyEditorModule.CreateStructureDetailView(DetailsViewArgs, StructureViewArgs, nullptr);
 	DetailsView->GetDetailsView()->SetObject(ObjectToEdit);
+	DetailsView->GetOnFinishedChangingPropertiesDelegate().AddLambda([this](const FPropertyChangedEvent& PropertyChangedEvent) 
+	{
+			FMessageDialog::Open(EAppMsgType::Ok, NSLOCTEXT("Dataflow", "UnknownEvaluationContext_Dataflow","Property Changed"));
+			if (PropertyChangedEvent.ChangeType == EPropertyChangeType::ValueSet)
+			{
+				FMessageDialog::Open(EAppMsgType::Ok, NSLOCTEXT("Dataflow", "UnknownEvaluationContext_Dataflow", "Property Changed really"));
+			}
+	});
+
 
 	return DetailsView;
 
