Dataflow: Evaluation
- Moved evaluation logic to base library.

diff --git a/Engine/Plugins/Experimental/Dataflow/Source/DataflowEditor/Private/Dataflow/DataflowEditorActions.cpp b/Engine/Plugins/Experimental/Dataflow/Source/DataflowEditor/Private/Dataflow/DataflowEditorActions.cpp
--- a/Engine/Plugins/Experimental/Dataflow/Source/DataflowEditor/Private/Dataflow/DataflowEditorActions.cpp
+++ b/Engine/Plugins/Experimental/Dataflow/Source/DataflowEditor/Private/Dataflow/DataflowEditorActions.cpp
@@ -52,7 +52,7 @@ void FDataflowEditorCommands::Unregister()
 
 
 
-void FDataflowEditorCommands::EvaluateNodes(const FGraphPanelSelectionSet& SelectedNodes, FDataflowEditorCommands::FGraphEvaluationCallback Evaluate)
+void FDataflowEditorCommands::EvaluateSelectedNodes(const FGraphPanelSelectionSet& SelectedNodes, FDataflowEditorCommands::FGraphEvaluationCallback Evaluate)
 {
 	for (UObject* Ode : SelectedNodes)
 	{
@@ -82,6 +82,79 @@ void FDataflowEditorCommands::EvaluateNodes(const FGraphPanelSelectionSet& Selec
 	}
 }
 
+void FDataflowEditorCommands::EvaluateNode(Dataflow::FContext& Context, Dataflow::FTimestamp& OutLastNodeTimestamp,
+	const UDataflow* Dataflow, const FDataflowNode* InNode, const FDataflowOutput* Output, FString NodeName)
+{
+	if (Dataflow)
+	{
+		const FDataflowNode* Node = InNode;
+		if (Node == nullptr)
+		{
+			if (const TSharedPtr<Dataflow::FGraph> Graph = Dataflow->GetDataflow())
+			{
+				if (TSharedPtr<const FDataflowNode> GraphNode = Graph->FindBaseNode(FName(NodeName)))
+				{
+					Node = GraphNode.Get();
+				}
+			}
+		}
+
+		if (Node != nullptr)
+		{
+			if (Output == nullptr)
+			{
+				if (Node->GetTimestamp() >= OutLastNodeTimestamp)
+				{
+					Context.Evaluate(Node, nullptr);
+					OutLastNodeTimestamp = Context.GetTimestamp();
+				}
+			}
+			else // Output != nullptr
+			{
+				if (!Context.HasData(Output->CacheKey(), Context.GetTimestamp()))
+				{
+					Context.Evaluate(Node, Output);
+				}
+			}
+		}
+	}
+}
+
+void FDataflowEditorCommands::OnPropertyValueChanged(UDataflow* OutDataflow, TSharedPtr<Dataflow::FEngineContext>& Context, Dataflow::FTimestamp& OutLastNodeTimestamp, const FPropertyChangedEvent& InPropertyChangedEvent)
+{
+	if (InPropertyChangedEvent.ChangeType == EPropertyChangeType::ValueSet)
+	{
+		TSharedPtr<const FDataflowNode> UpdatedNode = nullptr;
+		if (OutDataflow && InPropertyChangedEvent.Property && InPropertyChangedEvent.Property->GetOwnerUObject())
+		{
+			FString Name = InPropertyChangedEvent.Property->GetOwnerUObject()->GetName();
+			Name = Name.Left(Name.Len() - FString("DataflowNode").Len());
+			{
+				FDataflowAssetEdit Edit = OutDataflow->EditDataflow();
+
+				Dataflow::FGraph* Graph = Edit.GetGraph();
+				for (auto Node : Graph->GetNodes())
+				{
+					FString NodeName = Node->GetName().ToString();
+					if (NodeName.StartsWith(Name))
+					{
+						UpdatedNode = Node;
+						Node->Invalidate();
+					}
+				}
+			}
+		}
+
+		if (!UpdatedNode && Context)
+		{
+			// Some base properties dont link back to the parent, so just clobber the cache for now. 
+			Context.Reset();
+		}
+		OutLastNodeTimestamp = Dataflow::FTimestamp::Invalid;
+	}
+}
+
+
 void FDataflowEditorCommands::DeleteNodes(UDataflow* Graph, const FGraphPanelSelectionSet& SelectedNodes)
 {
 	if (ensureMsgf(Graph != nullptr, TEXT("Warning : Failed to find valid graph.")))
diff --git a/Engine/Plugins/Experimental/Dataflow/Source/DataflowEditor/Private/Dataflow/DataflowGraphEditor.cpp b/Engine/Plugins/Experimental/Dataflow/Source/DataflowEditor/Private/Dataflow/DataflowGraphEditor.cpp
--- a/Engine/Plugins/Experimental/Dataflow/Source/DataflowEditor/Private/Dataflow/DataflowGraphEditor.cpp
+++ b/Engine/Plugins/Experimental/Dataflow/Source/DataflowEditor/Private/Dataflow/DataflowGraphEditor.cpp
@@ -109,7 +109,7 @@ void SDataflowGraphEditor::EvaluateNode()
 {
 	if (EvaluateGraphCallback)
 	{
-		FDataflowEditorCommands::EvaluateNodes(GetSelectedNodes(), EvaluateGraphCallback);
+		FDataflowEditorCommands::EvaluateSelectedNodes(GetSelectedNodes(), EvaluateGraphCallback);
 	}
 	else
 	{
@@ -119,7 +119,7 @@ void SDataflowGraphEditor::EvaluateNode()
 			FContextThreaded(FPlatformTime::Cycles64()).Evaluate(Node, Out);
 		};
 
-		FDataflowEditorCommands::EvaluateNodes(GetSelectedNodes(), LocalEvaluateCallback);
+		FDataflowEditorCommands::EvaluateSelectedNodes(GetSelectedNodes(), LocalEvaluateCallback);
 	}
 }
 
diff --git a/Engine/Plugins/Experimental/Dataflow/Source/DataflowEditor/Public/Dataflow/DataflowEditorActions.h b/Engine/Plugins/Experimental/Dataflow/Source/DataflowEditor/Public/Dataflow/DataflowEditorActions.h
--- a/Engine/Plugins/Experimental/Dataflow/Source/DataflowEditor/Public/Dataflow/DataflowEditorActions.h
+++ b/Engine/Plugins/Experimental/Dataflow/Source/DataflowEditor/Public/Dataflow/DataflowEditorActions.h
@@ -3,6 +3,7 @@
 #pragma once
 
 #include "CoreMinimal.h"
+#include "Dataflow/DataflowObjectInterface.h"
 #include "Framework/Commands/Commands.h"
 #include "Styling/AppStyle.h"
 
@@ -52,15 +53,27 @@ public:
 	static const FDataflowEditorCommandsImpl& Get();
 
 	/*
-	*  EvaluateNodes
+	*  EvaluateSelectedNodes
 	*/
-	static void EvaluateNodes(const FGraphPanelSelectionSet& SelectedNodes, FGraphEvaluationCallback);
-	
+	static void EvaluateSelectedNodes(const FGraphPanelSelectionSet& SelectedNodes, FGraphEvaluationCallback);
+
+	/*
+	* EvaluateGraph
+	*/
+	static void EvaluateNode(Dataflow::FContext& Context, Dataflow::FTimestamp& OutLastNodeTimestamp,
+		const UDataflow* Dataflow, const FDataflowNode* Node = nullptr, const FDataflowOutput* Out = nullptr, 
+		FString NodeName = FString());
+
 	/*
 	*  DeleteNodes
 	*/
 	static void DeleteNodes(UDataflow* Graph, const FGraphPanelSelectionSet& SelectedNodes);
 	
+	/*
+	*  OnPropertyValueChanged
+	*/
+	static void OnPropertyValueChanged(UDataflow* Graph, TSharedPtr<Dataflow::FEngineContext>& Context, Dataflow::FTimestamp& OutLastNodeTimestamp, const FPropertyChangedEvent& PropertyChangedEvent);
+
 	/*
 	*  OnSelectedNodesChanged
 	*/
