Dataflow : Node renaming

diff --git a/Engine/Plugins/Experimental/Dataflow/Source/DataflowEditor/Private/Dataflow/DataflowCommentNode.cpp b/Engine/Plugins/Experimental/Dataflow/Source/DataflowEditor/Private/Dataflow/DataflowCommentNode.cpp
deleted file mode 100644
diff --git a/Engine/Plugins/Experimental/Dataflow/Source/DataflowEditor/Private/Dataflow/DataflowEditorActions.cpp b/Engine/Plugins/Experimental/Dataflow/Source/DataflowEditor/Private/Dataflow/DataflowEditorActions.cpp
--- a/Engine/Plugins/Experimental/Dataflow/Source/DataflowEditor/Private/Dataflow/DataflowEditorActions.cpp
+++ b/Engine/Plugins/Experimental/Dataflow/Source/DataflowEditor/Private/Dataflow/DataflowEditorActions.cpp
@@ -4,6 +4,7 @@
 #include "Dataflow/DataflowEdNode.h"
 #include "Dataflow/DataflowNodeFactory.h"
 #include "Dataflow/DataflowObject.h"
+#include "EdGraph/EdGraphNode.h"
 #include "IStructureDetailsView.h"
 #include "EdGraphNode_Comment.h"
 
@@ -120,6 +121,59 @@ void FDataflowEditorCommands::EvaluateNode(Dataflow::FContext& Context, Dataflow
 	}
 }
 
+
+bool FDataflowEditorCommands::OnNodeVerifyTitleCommit(const FText& NewText, UEdGraphNode* GraphNode, FText& OutErrorMessage)
+{
+	if (GraphNode)
+	{
+		if (UDataflowEdNode* DataflowNode = Cast<UDataflowEdNode>(GraphNode))
+		{
+			if (TSharedPtr<Dataflow::FGraph> Graph = DataflowNode->GetDataflowGraph())
+			{
+				if( Graph->FindBaseNode(FName(NewText.ToString())).Get()==nullptr )
+				{
+					return true;
+				}
+			}
+		}
+		else if( Cast<UEdGraphNode_Comment>(GraphNode))
+		{
+			return true;
+		}
+	}
+	OutErrorMessage = FText::FromString(FString::Printf(TEXT("Non-unique name for graph node (%s)"), *NewText.ToString()));
+	return false;
+}
+
+
+void FDataflowEditorCommands::OnNodeTitleCommitted(const FText& InNewText, ETextCommit::Type InCommitType, UEdGraphNode* GraphNode)
+{
+	if (InCommitType == ETextCommit::OnCleared)
+	{
+		return;
+	}
+
+	if (GraphNode)
+	{
+		if (UDataflowEdNode* DataflowNode = Cast<UDataflowEdNode>(GraphNode))
+		{
+			if (TSharedPtr<Dataflow::FGraph> Graph = DataflowNode->GetDataflowGraph())
+			{
+				if (TSharedPtr<FDataflowNode> Node = Graph->FindBaseNode(DataflowNode->GetDataflowNodeGuid()))
+				{
+					GraphNode->Rename(*InNewText.ToString());
+					Node->SetName(FName(InNewText.ToString()));
+				}
+			}
+		}
+		else if (UEdGraphNode_Comment* CommentNode = Cast<UEdGraphNode_Comment>(GraphNode))
+		{
+			GraphNode->NodeComment = InNewText.ToString();
+		}
+	}
+}
+
+
 void FDataflowEditorCommands::OnPropertyValueChanged(UDataflow* OutDataflow, TSharedPtr<Dataflow::FEngineContext>& Context, Dataflow::FTimestamp& OutLastNodeTimestamp, const FPropertyChangedEvent& InPropertyChangedEvent)
 {
 	if (InPropertyChangedEvent.ChangeType == EPropertyChangeType::ValueSet)
diff --git a/Engine/Plugins/Experimental/Dataflow/Source/DataflowEditor/Private/Dataflow/DataflowGraphEditor.cpp b/Engine/Plugins/Experimental/Dataflow/Source/DataflowEditor/Private/Dataflow/DataflowGraphEditor.cpp
--- a/Engine/Plugins/Experimental/Dataflow/Source/DataflowEditor/Private/Dataflow/DataflowGraphEditor.cpp
+++ b/Engine/Plugins/Experimental/Dataflow/Source/DataflowEditor/Private/Dataflow/DataflowGraphEditor.cpp
@@ -7,7 +7,7 @@
 #include "Dataflow/DataflowSNodeFactories.h"
 #include "BoneDragDropOp.h"
 #include "IStructureDetailsView.h"
-#include "Dataflow/DataflowCommentNode.h"
+#include "Dataflow/DataflowSCommentNode.h"
 #include "Dataflow/DataflowNodeParameters.h"
 
 #include "Dataflow/DataflowSchema.h"
diff --git a/Engine/Plugins/Experimental/Dataflow/Source/DataflowEditor/Private/Dataflow/DataflowSCommentNode.cpp b/Engine/Plugins/Experimental/Dataflow/Source/DataflowEditor/Private/Dataflow/DataflowSCommentNode.cpp
new file mode 100644
--- /dev/null
+++ b/Engine/Plugins/Experimental/Dataflow/Source/DataflowEditor/Private/Dataflow/DataflowSCommentNode.cpp
@@ -0,0 +1,75 @@
+// Copyright Epic Games, Inc. All Rights Reserved.
+
+#include "Dataflow/DataflowSCommentNode.h"
+
+#include "EdGraphNode_Comment.h"
+#include "Dataflow/DataflowEdNode.h"
+#include "Dataflow/DataflowNodeFactory.h"
+#include "Dataflow/DataflowObject.h"
+#include "Dataflow/DataflowCore.h"
+#include "Logging/LogMacros.h"
+#include "Widgets/Layout/SBorder.h"
+#include "Settings/EditorStyleSettings.h"
+#include "ScopedTransaction.h"
+
+#define LOCTEXT_NAMESPACE "SDataflowEdNodeComment"
+
+//
+// Add a menu option to create a graph node.
+//
+TSharedPtr<FAssetSchemaAction_Dataflow_CreateCommentNode_DataflowEdNode> FAssetSchemaAction_Dataflow_CreateCommentNode_DataflowEdNode::CreateAction(UEdGraph* ParentGraph, const TSharedPtr<SGraphEditor>& GraphEditor)
+{
+	return MakeShared<FAssetSchemaAction_Dataflow_CreateCommentNode_DataflowEdNode>(GraphEditor);
+}
+
+//
+//  Created comment node
+//
+UEdGraphNode* FAssetSchemaAction_Dataflow_CreateCommentNode_DataflowEdNode::PerformAction(class UEdGraph* ParentGraph, UEdGraphPin* FromPin, const FVector2D Location, bool bSelectNewNode)
+{
+	UEdGraphNode_Comment* CommentTemplate = NewObject<UEdGraphNode_Comment>(ParentGraph, FName("Comment"));
+
+	FVector2D SpawnLocation = Location;
+	FSlateRect Bounds;
+
+	if (GraphEditor->GetBoundsForSelectedNodes(Bounds, 50.0f))
+	{
+		CommentTemplate->SetBounds(Bounds);
+		SpawnLocation.X = CommentTemplate->NodePosX;
+		SpawnLocation.Y = CommentTemplate->NodePosY;
+	}
+
+	CommentTemplate->bCommentBubbleVisible_InDetailsPanel = false;
+	CommentTemplate->bCommentBubbleVisible = false;
+	CommentTemplate->bCommentBubblePinned = false;
+
+	const FScopedTransaction Transaction(NSLOCTEXT("UnrealEd", "DataflowEditorNewNode", "Dataflow Editor: New Comment Node"));
+	ParentGraph->Modify();
+
+	CommentTemplate->SetFlags(RF_Transactional);
+
+	// set outer to be the graph so it doesn't go away
+	CommentTemplate->Rename(NULL, ParentGraph, REN_NonTransactional);
+	ParentGraph->AddNode(CommentTemplate, true, bSelectNewNode);
+
+	CommentTemplate->CreateNewGuid();
+	CommentTemplate->PostPlacedNewNode();
+	CommentTemplate->AllocateDefaultPins();
+	CommentTemplate->AutowireNewNode(FromPin);
+
+	CommentTemplate->NodePosX = SpawnLocation.X;
+	CommentTemplate->NodePosY = SpawnLocation.Y;
+	CommentTemplate->SnapToGrid(GetDefault<UEditorStyleSettings>()->GridSnapSize);
+
+	ParentGraph->NotifyGraphChanged();
+
+	return CommentTemplate;
+}
+
+//void FAssetSchemaAction_Dataflow_CreateCommentNode_DataflowEdNode::AddReferencedObjects(FReferenceCollector& Collector)
+//{
+//	FEdGraphSchemaAction::AddReferencedObjects(Collector);
+//	Collector.AddReferencedObject(NodeTemplate);
+//}
+
+#undef LOCTEXT_NAMESPACE
diff --git a/Engine/Plugins/Experimental/Dataflow/Source/DataflowEditor/Private/Dataflow/DataflowSNodeFactories.cpp b/Engine/Plugins/Experimental/Dataflow/Source/DataflowEditor/Private/Dataflow/DataflowSNodeFactories.cpp
--- a/Engine/Plugins/Experimental/Dataflow/Source/DataflowEditor/Private/Dataflow/DataflowSNodeFactories.cpp
+++ b/Engine/Plugins/Experimental/Dataflow/Source/DataflowEditor/Private/Dataflow/DataflowSNodeFactories.cpp
@@ -4,7 +4,7 @@
 #include "Dataflow/DataflowSNode.h"
 #include "Dataflow/DataflowEdNode.h"
 #include "EdGraphNode_Comment.h"
-#include "Dataflow/DataflowCommentNode.h"
+#include "Dataflow/DataflowSCommentNode.h"
 #include "Dataflow/DataflowSchema.h"
 
 TSharedPtr<class SGraphNode> FDataflowSNodeFactory::CreateNode(UEdGraphNode* InNode) const
diff --git a/Engine/Plugins/Experimental/Dataflow/Source/DataflowEditor/Public/Dataflow/DataflowCommentNode.h b/Engine/Plugins/Experimental/Dataflow/Source/DataflowEditor/Public/Dataflow/DataflowCommentNode.h
deleted file mode 100644
diff --git a/Engine/Plugins/Experimental/Dataflow/Source/DataflowEditor/Public/Dataflow/DataflowEditorActions.h b/Engine/Plugins/Experimental/Dataflow/Source/DataflowEditor/Public/Dataflow/DataflowEditorActions.h
--- a/Engine/Plugins/Experimental/Dataflow/Source/DataflowEditor/Public/Dataflow/DataflowEditorActions.h
+++ b/Engine/Plugins/Experimental/Dataflow/Source/DataflowEditor/Public/Dataflow/DataflowEditorActions.h
@@ -12,6 +12,7 @@ class IStructureDetailsView;
 class UDataflow;
 class UDataflowEdNode;
 struct FDataflowNode;
+class UEdGraphNode;
 
 typedef TSet<class UObject*> FGraphPanelSelectionSet;
 
@@ -68,7 +69,17 @@ public:
 	*  DeleteNodes
 	*/
 	static void DeleteNodes(UDataflow* Graph, const FGraphPanelSelectionSet& SelectedNodes);
-	
+
+	/*
+	* OnNodeVerifyTitleCommit
+	*/
+	static bool OnNodeVerifyTitleCommit(const FText& NewText, UEdGraphNode* GraphNode, FText& OutErrorMessage);
+
+	/*
+	* OnNodeTitleCommitted
+	*/
+	static void OnNodeTitleCommitted(const FText& InNewText, ETextCommit::Type InCommitType, UEdGraphNode* GraphNode);
+
 	/*
 	*  OnPropertyValueChanged
 	*/
diff --git a/Engine/Plugins/Experimental/Dataflow/Source/DataflowEditor/Public/Dataflow/DataflowSCommentNode.h b/Engine/Plugins/Experimental/Dataflow/Source/DataflowEditor/Public/Dataflow/DataflowSCommentNode.h
new file mode 100644
--- /dev/null
+++ b/Engine/Plugins/Experimental/Dataflow/Source/DataflowEditor/Public/Dataflow/DataflowSCommentNode.h
@@ -0,0 +1,47 @@
+// Copyright Epic Games, Inc. All Rights Reserved.
+
+#pragma once
+
+#include "EdGraphUtilities.h"
+#include "SGraphNode.h"
+#include "SGraphNodeComment.h"
+
+#include "DataflowSCommentNode.generated.h"
+
+//
+// SDataflowEdNodeComment
+//
+
+class DATAFLOWEDITOR_API SDataflowEdNodeComment : public SGraphNodeComment
+{
+	typedef SGraphNodeComment Super;
+
+public:
+	SLATE_BEGIN_ARGS(SDataflowEdNodeComment) {}
+	SLATE_END_ARGS()
+
+	void Construct(const FArguments& InArgs, UEdGraphNode_Comment* InNode) { Super::Construct(SGraphNodeComment::FArguments(), InNode); }
+};
+
+//
+// Action to add a node to the graph
+//
+USTRUCT()
+struct DATAFLOWEDITOR_API FAssetSchemaAction_Dataflow_CreateCommentNode_DataflowEdNode : public FEdGraphSchemaAction
+{
+	GENERATED_USTRUCT_BODY();
+
+public:
+	FAssetSchemaAction_Dataflow_CreateCommentNode_DataflowEdNode(const TSharedPtr<SGraphEditor>& InGraphEditor) : FEdGraphSchemaAction()
+		, GraphEditor(InGraphEditor)
+		{}
+
+	FAssetSchemaAction_Dataflow_CreateCommentNode_DataflowEdNode() {}
+
+	static TSharedPtr<FAssetSchemaAction_Dataflow_CreateCommentNode_DataflowEdNode> CreateAction(UEdGraph* ParentGraph, const TSharedPtr<SGraphEditor>& GraphEditor);
+
+	virtual UEdGraphNode* PerformAction(class UEdGraph* ParentGraph, UEdGraphPin* FromPin, const FVector2D Location, bool bSelectNewNode = true) override;
+	//virtual void AddReferencedObjects(FReferenceCollector& Collector) override;
+
+	TSharedPtr<SGraphEditor> GraphEditor;
+};
diff --git a/Engine/Source/Runtime/Experimental/Dataflow/Engine/Private/Dataflow/DataflowEdNode.cpp b/Engine/Source/Runtime/Experimental/Dataflow/Engine/Private/Dataflow/DataflowEdNode.cpp
--- a/Engine/Source/Runtime/Experimental/Dataflow/Engine/Private/Dataflow/DataflowEdNode.cpp
+++ b/Engine/Source/Runtime/Experimental/Dataflow/Engine/Private/Dataflow/DataflowEdNode.cpp
@@ -12,11 +12,19 @@
 #include "EdGraph/EdGraphPin.h"
 #endif
 
-
 #define LOCTEXT_NAMESPACE "DataflowEdNode"
 
 DEFINE_LOG_CATEGORY_STATIC(DATAFLOWNODE_LOG, Error, All);
 
+
+UDataflowEdNode::UDataflowEdNode(const FObjectInitializer& ObjectInitializer)
+	: Super(ObjectInitializer)
+{
+#if WITH_EDITOR && !(UE_BUILD_SHIPPING || UE_BUILD_TEST)
+	bCanRenameNode = true;
+#endif // WITH_EDITOR && !UE_BUILD_SHIPPING
+}
+
 void UDataflowEdNode::AllocateDefaultPins()
 {
 	UE_LOG(DATAFLOWNODE_LOG, Verbose, TEXT("UDataflowEdNode::AllocateDefaultPins()"));
diff --git a/Engine/Source/Runtime/Experimental/Dataflow/Engine/Public/Dataflow/DataflowEdNode.h b/Engine/Source/Runtime/Experimental/Dataflow/Engine/Public/Dataflow/DataflowEdNode.h
--- a/Engine/Source/Runtime/Experimental/Dataflow/Engine/Public/Dataflow/DataflowEdNode.h
+++ b/Engine/Source/Runtime/Experimental/Dataflow/Engine/Public/Dataflow/DataflowEdNode.h
@@ -24,11 +24,13 @@ namespace Dataflow { class FGraph; }
 UCLASS()
 class DATAFLOWENGINE_API UDataflowEdNode : public UEdGraphNode
 {
-	GENERATED_BODY()
+	GENERATED_UCLASS_BODY()
 
 	FGuid DataflowNodeGuid;
 	TSharedPtr<Dataflow::FGraph> DataflowGraph;
+
 public:
+
 	// UEdGraphNode interface
 	virtual void AllocateDefaultPins();
 	virtual FText GetNodeTitle(ENodeTitleType::Type TitleType) const;
