Dataflow : Node Rendering

diff --git a/Engine/Plugins/Experimental/Dataflow/Dataflow.uplugin b/Engine/Plugins/Experimental/Dataflow/Dataflow.uplugin
--- a/Engine/Plugins/Experimental/Dataflow/Dataflow.uplugin
+++ b/Engine/Plugins/Experimental/Dataflow/Dataflow.uplugin
@@ -21,6 +21,11 @@
 			"Type": "Editor",
 			"LoadingPhase": "Default"
 		},
+		{
+			"Name": "DataflowEnginePlugin",
+			"Type": "Runtime",
+			"LoadingPhase": "Default"
+		},
 		{
 			"Name": "DataflowNodes",
 			"Type": "Runtime",
@@ -28,5 +33,9 @@
 		}
 	],
 	"Plugins": [
+		{
+			"Name": "ProceduralMeshComponent",
+			"Enabled": true
+		}
 	]
 }
diff --git a/Engine/Plugins/Experimental/Dataflow/Source/DataflowEditor/DataflowEditor.Build.cs b/Engine/Plugins/Experimental/Dataflow/Source/DataflowEditor/DataflowEditor.Build.cs
--- a/Engine/Plugins/Experimental/Dataflow/Source/DataflowEditor/DataflowEditor.Build.cs
+++ b/Engine/Plugins/Experimental/Dataflow/Source/DataflowEditor/DataflowEditor.Build.cs
@@ -32,6 +32,7 @@ namespace UnrealBuildTool.Rules
 					"GraphEditor",
 					"DataflowCore",
 					"DataflowEngine",
+					"DataflowEnginePlugin",
 					"Slate",
 				}
 			);
diff --git a/Engine/Plugins/Experimental/Dataflow/Source/DataflowEditor/Private/Dataflow/DataflowEditorStyle.h b/Engine/Plugins/Experimental/Dataflow/Source/DataflowEditor/Private/Dataflow/DataflowEditorStyle.h
--- a/Engine/Plugins/Experimental/Dataflow/Source/DataflowEditor/Private/Dataflow/DataflowEditorStyle.h
+++ b/Engine/Plugins/Experimental/Dataflow/Source/DataflowEditor/Private/Dataflow/DataflowEditorStyle.h
@@ -2,6 +2,7 @@
 
 #pragma once
 #include "Styling/SlateStyle.h"
+#include "Styling/SlateStyleMacros.h"
 #include "Styling/SlateStyleRegistry.h"
 #include "Interfaces/IPluginManager.h"
 
@@ -12,10 +13,17 @@ public:
 	{
 		const FVector2D Icon16x16(16.f, 16.f);
 		const FVector2D Icon64x64(64.f, 64.f);
+		const FVector2D Icon56x28(56.f, 28.f);
+		const FVector2D Icon28x14(28.f, 14.f);
 
 		SetContentRoot(IPluginManager::Get().FindPlugin("Dataflow")->GetBaseDir() / TEXT("Resources"));
 		Set("ClassIcon.Dataflow", new FSlateVectorImageBrush(RootToContentDir(TEXT("DataflowAsset_16.svg")), Icon16x16));
 		Set("ClassThumbnail.Dataflow", new FSlateVectorImageBrush(RootToContentDir(TEXT("DataflowAsset_64.svg")), Icon64x64));
+
+		Set("Dataflow.Render.Unknown", new FSlateImageBrush(RootToContentDir(TEXT("Slate/Switch_Undetermined_56x_28x")), Icon56x28));
+		Set("Dataflow.Render.Disabled", new FSlateImageBrush(RootToContentDir(TEXT("Slate/Switch_OFF_56x_28x")), Icon56x28));
+		Set("Dataflow.Render.Enabled", new FSlateImageBrush(RootToContentDir(TEXT("Slate/Switch_ON_56x_28x")), Icon56x28));
+
 		FSlateStyleRegistry::RegisterSlateStyle(*this);
 	}
 
@@ -28,20 +36,9 @@ public:
 
 	static FDataflowEditorStyle& Get()
 	{
-		if (!Singleton.IsSet())
-		{
-			Singleton.Emplace();
-		}
-		return Singleton.GetValue();
-	}
-
-	static void Destroy()
-	{
-		Singleton.Reset();
+		static FDataflowEditorStyle Inst;
+		return Inst;
 	}
 
-private:
-	static TOptional<FDataflowEditorStyle> Singleton;
 };
 
-TOptional<FDataflowEditorStyle> FDataflowEditorStyle::Singleton;
\ No newline at end of file
diff --git a/Engine/Plugins/Experimental/Dataflow/Source/DataflowEditor/Private/Dataflow/DataflowSNode.cpp b/Engine/Plugins/Experimental/Dataflow/Source/DataflowEditor/Private/Dataflow/DataflowSNode.cpp
--- a/Engine/Plugins/Experimental/Dataflow/Source/DataflowEditor/Private/Dataflow/DataflowSNode.cpp
+++ b/Engine/Plugins/Experimental/Dataflow/Source/DataflowEditor/Private/Dataflow/DataflowSNode.cpp
@@ -3,12 +3,18 @@
 #include "Dataflow/DataflowSNode.h"
 
 #include "Dataflow/DataflowEdNode.h"
+#include "Dataflow/DataflowEditorStyle.h"
+#include "Dataflow/DataflowGraphEditor.h"
 #include "Dataflow/DataflowNodeFactory.h"
 #include "Dataflow/DataflowObject.h"
 #include "Dataflow/DataflowCore.h"
 #include "Logging/LogMacros.h"
 #include "SourceCodeNavigation.h"
+#include "Styling/SlateTypes.h"
+#include "Styling/CoreStyle.h"
+#include "Styling/AppStyle.h"
 #include "Widgets/Layout/SBorder.h"
+#include "Widgets/Input/SCheckBox.h"
 
 #include UE_INLINE_GENERATED_CPP_BY_NAME(DataflowSNode)
 
@@ -20,7 +26,64 @@
 void SDataflowEdNode::Construct(const FArguments& InArgs, UDataflowEdNode* InNode)
 {
 	GraphNode = InNode;
+	DataflowGraphNode = Cast<UDataflowEdNode>(InNode);
 	UpdateGraphNode();
+
+	/*
+	const FSlateBrush* DisabledSwitchBrush = FDataflowEditorStyle::Get().GetBrush(TEXT("Dataflow.Render.Disabled"));
+	const FSlateBrush* EnabledSwitchBrush = FDataflowEditorStyle::Get().GetBrush(TEXT("Dataflow.Render.Enabled"));
+
+	CheckBoxStyle = FCheckBoxStyle()
+		.SetCheckBoxType(ESlateCheckBoxType::CheckBox)
+		.SetUncheckedImage(*DisabledSwitchBrush)
+		.SetUncheckedHoveredImage(*DisabledSwitchBrush)
+		.SetUncheckedPressedImage(*DisabledSwitchBrush)
+		.SetCheckedImage(*EnabledSwitchBrush)
+		.SetCheckedHoveredImage(*EnabledSwitchBrush)
+		.SetCheckedPressedImage(*EnabledSwitchBrush)
+		.SetPadding(FMargin(0, 0, 0, 1));
+	*/
+	//TSharedPtr<SNodePanel> ParentPanel = GetParentPanel();
+
+	RenderCheckBoxWidget = SNew(SCheckBox)
+		.IsChecked_Lambda([this]()-> ECheckBoxState
+			{
+				if (DataflowGraphNode && DataflowGraphNode->DoAssetRender())
+				{
+					return ECheckBoxState::Checked;
+				}
+				return ECheckBoxState::Unchecked;
+			})
+		.OnCheckStateChanged_Lambda([&](const ECheckBoxState NewState)
+			{
+				if (DataflowGraphNode)
+				{
+					if (NewState == ECheckBoxState::Checked)
+						DataflowGraphNode->SetAssetRender(true);
+					else
+						DataflowGraphNode->SetAssetRender(false);
+				}
+			});
+		//.Style(&CheckBoxStyle);
+
+}
+
+TArray<FOverlayWidgetInfo> SDataflowEdNode::GetOverlayWidgets(bool bSelected, const FVector2D& WidgetSize) const
+{
+	TArray<FOverlayWidgetInfo> Widgets = SGraphNode::GetOverlayWidgets(bSelected, WidgetSize);
+
+	if (DataflowGraphNode && DataflowGraphNode->GetDataflowNode()->GetRenderParameters().Num())
+	{
+		const FVector2D ImageSize = RenderCheckBoxWidget->GetDesiredSize();
+
+		FOverlayWidgetInfo Info;
+		Info.OverlayOffset = FVector2D(WidgetSize.X - ImageSize.X - 6.f, 6.f);
+		Info.Widget = RenderCheckBoxWidget;
+
+		Widgets.Add(Info);
+	}
+
+	return Widgets;
 }
 
 FReply SDataflowEdNode::OnMouseButtonDoubleClick(const FGeometry& InMyGeometry, const FPointerEvent& InMouseEvent)
diff --git a/Engine/Plugins/Experimental/Dataflow/Source/DataflowEditor/Public/Dataflow/DataflowSNode.h b/Engine/Plugins/Experimental/Dataflow/Source/DataflowEditor/Public/Dataflow/DataflowSNode.h
--- a/Engine/Plugins/Experimental/Dataflow/Source/DataflowEditor/Public/Dataflow/DataflowSNode.h
+++ b/Engine/Plugins/Experimental/Dataflow/Source/DataflowEditor/Public/Dataflow/DataflowSNode.h
@@ -8,6 +8,7 @@
 #include "DataflowSNode.generated.h"
 
 class UDataflowEdNode;
+class SCheckBox;
 
 //
 // SDataflowEdNode
@@ -18,19 +19,26 @@ class DATAFLOWEDITOR_API SDataflowEdNode : public SGraphNode
 	typedef SGraphNode Super;
 
 public:
+	typedef TFunction<void(UEdGraphNode* InNode, bool InEnabled)> FToggleRenderCallback;
+
+
 	SLATE_BEGIN_ARGS(SDataflowEdNode)
 		: _GraphNodeObj(nullptr)
 	{}
-
 	SLATE_ARGUMENT(UDataflowEdNode*, GraphNodeObj)
-
 	SLATE_END_ARGS()
 	
 	void Construct(const FArguments& InArgs, UDataflowEdNode* InNode);
 
 	// SGraphNode interface
 	virtual FReply OnMouseButtonDoubleClick(const FGeometry& InMyGeometry, const FPointerEvent& InMouseEvent) override;
+	TArray<FOverlayWidgetInfo> GetOverlayWidgets(bool bSelected, const FVector2D& WidgetSize) const;
+
+private:
+	UDataflowEdNode* DataflowGraphNode = nullptr;	
 
+	FCheckBoxStyle CheckBoxStyle;
+	TSharedPtr<SCheckBox> RenderCheckBoxWidget;
 };
 
 //
diff --git a/Engine/Plugins/Experimental/Dataflow/Source/DataflowEnginePlugin/DataflowEnginePlugin.Build.cs b/Engine/Plugins/Experimental/Dataflow/Source/DataflowEnginePlugin/DataflowEnginePlugin.Build.cs
new file mode 100644
--- /dev/null
+++ b/Engine/Plugins/Experimental/Dataflow/Source/DataflowEnginePlugin/DataflowEnginePlugin.Build.cs
@@ -0,0 +1,24 @@
+// Copyright Epic Games, Inc. All Rights Reserved.
+
+namespace UnrealBuildTool.Rules
+{
+	public class DataflowEnginePlugin : ModuleRules
+	{
+        public DataflowEnginePlugin(ReadOnlyTargetRules Target) : base(Target)
+		{
+			bTreatAsEngineModule = true;
+			PublicDependencyModuleNames.AddRange(
+				new string[]
+				{
+					"Core",
+					"CoreUObject",
+					"Engine",
+					"DataflowCore",
+					"DataflowEngine",
+					"Chaos",
+					"ProceduralMeshComponent"
+				}
+			);
+		}
+	}
+}
diff --git a/Engine/Plugins/Experimental/Dataflow/Source/DataflowEnginePlugin/Private/Dataflow/DataflowEnginePlugin.cpp b/Engine/Plugins/Experimental/Dataflow/Source/DataflowEnginePlugin/Private/Dataflow/DataflowEnginePlugin.cpp
new file mode 100644
--- /dev/null
+++ b/Engine/Plugins/Experimental/Dataflow/Source/DataflowEnginePlugin/Private/Dataflow/DataflowEnginePlugin.cpp
@@ -0,0 +1,82 @@
+// Copyright Epic Games, Inc. All Rights Reserved.
+
+#include "Dataflow/DataflowEnginePlugin.h"
+
+#include "CoreMinimal.h"
+#include "Modules/ModuleManager.h"
+#include "Dataflow/DataflowRenderingFactory.h"
+#include "GeometryCollection/GeometryCollectionUtility.h"
+#include "GeometryCollection/Facades/CollectionRenderingFacade.h"
+#include "GeometryCollection/ManagedArrayCollection.h"
+#include "GeometryCollection/GeometryCollectionAlgo.h"
+
+
+class FDataflowEnginePlugin : public IDataflowEnginePlugin
+{
+	/** IModuleInterface implementation */
+	virtual void StartupModule() override;
+	virtual void ShutdownModule() override;
+};
+
+IMPLEMENT_MODULE( FDataflowEnginePlugin, FleshEngine )
+
+
+
+void FDataflowEnginePlugin::StartupModule()
+{
+	using namespace Dataflow;
+
+	/**
+	* DatflowNode (FGeometryCollection) Rendering
+	*
+	*		@param Type : FGeometryCollection::StaticType()
+
+	*		@param Outputs : {FManagedArrayCollection : "Collection"}
+	*/
+	FRenderingFactory::GetInstance()->RegisterOutput(FGeometryCollection::StaticType(),
+		[](GeometryCollection::Facades::FRenderingFacade& RenderCollection, const Dataflow::FGraphRenderingState& State)
+		{
+			if (State.GetRenderOutputs().Num())
+			{
+				FName PrimaryOutput = State.GetRenderOutputs()[0]; // "Collection"
+				
+				FManagedArrayCollection Default;
+				const FManagedArrayCollection& Collection = State.GetValue<FManagedArrayCollection>(PrimaryOutput, Default);
+
+				if (Collection.FindAttributeTyped<FIntVector>("Indices", FGeometryCollection::FacesGroup)
+					&& Collection.FindAttributeTyped<FVector3f>("Vertex", FGeometryCollection::VerticesGroup)
+					&& Collection.FindAttributeTyped<FTransform>(FTransformCollection::TransformAttribute, FTransformCollection::TransformGroup)
+					&& Collection.FindAttributeTyped<int32>("BoneMap", FGeometryCollection::VerticesGroup)
+					&& Collection.FindAttributeTyped<int32>(FTransformCollection::ParentAttribute, FTransformCollection::TransformGroup))
+				{
+					const TManagedArray<int32>& BoneIndex = Collection.GetAttribute<int32>("BoneMap", FGeometryCollection::VerticesGroup);
+					const TManagedArray<int32>& Parents = Collection.GetAttribute<int32>(FTransformCollection::ParentAttribute, FTransformCollection::TransformGroup);
+					const TManagedArray<FTransform>& Transforms = Collection.GetAttribute<FTransform>(FTransformCollection::TransformAttribute, FTransformCollection::TransformGroup);
+
+					TArray<FTransform> M;
+					GeometryCollectionAlgo::GlobalMatrices(Transforms, Parents, M);
+
+					auto ToD = [](FVector3f V) { return FVector3d(V.X, V.Y, V.Z); };
+
+					const TManagedArray<FVector3f>& Vertex = Collection.GetAttribute<FVector3f>("Vertex", FGeometryCollection::VerticesGroup);
+					const TManagedArray<FIntVector>& Faces = Collection.GetAttribute<FIntVector>("Indices", FGeometryCollection::FacesGroup);
+					for (const FIntVector& Face : Faces)
+					{
+						FTransform Ms[3] = { M[BoneIndex[Face[0]]], M[BoneIndex[Face[1]]], M[BoneIndex[Face[2]]] };
+						RenderCollection.AddTriangle(Chaos::FTriangle(Ms[0].TransformPosition(ToD(Vertex[Face[0]])),
+							Ms[1].TransformPosition(ToD(Vertex[Face[1]])), Ms[2].TransformPosition(ToD(Vertex[Face[2]]))));
+					}
+				}
+			}
+		});
+
+}
+
+
+void FDataflowEnginePlugin::ShutdownModule()
+{
+	
+}
+
+
+
diff --git a/Engine/Plugins/Experimental/Dataflow/Source/DataflowEnginePlugin/Private/Dataflow/DataflowEnginePluginPrivate.h b/Engine/Plugins/Experimental/Dataflow/Source/DataflowEnginePlugin/Private/Dataflow/DataflowEnginePluginPrivate.h
new file mode 100644
--- /dev/null
+++ b/Engine/Plugins/Experimental/Dataflow/Source/DataflowEnginePlugin/Private/Dataflow/DataflowEnginePluginPrivate.h
@@ -0,0 +1,12 @@
+// Copyright Epic Games, Inc. All Rights Reserved.
+
+#pragma once
+
+#include "CoreMinimal.h"
+#include "Stats/Stats.h"
+
+// You should place include statements to your module's private header files here.  You only need to
+// add includes for headers that are used in most of your module's source files though.
+
+
+DECLARE_STATS_GROUP(TEXT("DataflowEnginePlugin"), STATGROUP_DataflowEnginePlugin, STATCAT_Advanced);
diff --git a/Engine/Plugins/Experimental/Dataflow/Source/DataflowEnginePlugin/Private/Dataflow/DataflowRenderingActor.cpp b/Engine/Plugins/Experimental/Dataflow/Source/DataflowEnginePlugin/Private/Dataflow/DataflowRenderingActor.cpp
new file mode 100644
--- /dev/null
+++ b/Engine/Plugins/Experimental/Dataflow/Source/DataflowEnginePlugin/Private/Dataflow/DataflowRenderingActor.cpp
@@ -0,0 +1,28 @@
+// Copyright Epic Games, Inc. All Rights Reserved.
+
+#include "Dataflow/DataflowRenderingActor.h"
+
+#include "CoreMinimal.h"
+#include "Dataflow/DataflowRenderingComponent.h"
+
+
+#include UE_INLINE_GENERATED_CPP_BY_NAME(DataflowRenderingActor)
+
+DEFINE_LOG_CATEGORY_STATIC(ADataflowLogging, Log, All);
+
+ADataflowRenderingActor::ADataflowRenderingActor(const FObjectInitializer& ObjectInitializer)
+	: Super(ObjectInitializer)
+{
+	//UE_LOG(ADataflowLogging, Verbose, TEXT("ADataflowRenderingActor::ADataflowRenderingActor()"));
+	DataflowRenderingComponent = CreateDefaultSubobject<UDataflowRenderingComponent>(TEXT("DataflowRenderingComponent0"));
+	RootComponent = DataflowRenderingComponent;
+	PrimaryActorTick.bCanEverTick = true;
+}
+
+#if WITH_EDITOR
+bool ADataflowRenderingActor::GetReferencedContentObjects(TArray<UObject*>& Objects) const
+{
+	return Super::GetReferencedContentObjects(Objects);
+}
+#endif
+
diff --git a/Engine/Plugins/Experimental/Dataflow/Source/DataflowEnginePlugin/Private/Dataflow/DataflowRenderingComponent.cpp b/Engine/Plugins/Experimental/Dataflow/Source/DataflowEnginePlugin/Private/Dataflow/DataflowRenderingComponent.cpp
new file mode 100644
--- /dev/null
+++ b/Engine/Plugins/Experimental/Dataflow/Source/DataflowEnginePlugin/Private/Dataflow/DataflowRenderingComponent.cpp
@@ -0,0 +1,182 @@
+// Copyright Epic Games, Inc. All Rights Reserved.
+
+#include "Dataflow/DataflowRenderingComponent.h"
+
+#include "Dataflow/DataflowEdNode.h"
+#include "Dataflow/DataflowRenderingFactory.h"
+#include "GeometryCollection/Facades/CollectionRenderingFacade.h"
+#include "GeometryCollection/Facades/CollectionBoundsFacade.h"
+#include "GeometryCollection/GeometryCollection.h"
+
+#include UE_INLINE_GENERATED_CPP_BY_NAME(DataflowRenderingComponent)
+
+DEFINE_LOG_CATEGORY_STATIC(LogDataflowRenderComponentInternal, Log, All);
+
+UDataflowRenderingComponent::UDataflowRenderingComponent(const FObjectInitializer& ObjectInitializer)
+	: Super(ObjectInitializer)
+{
+	PrimaryComponentTick.bCanEverTick = true;
+	bTickInEditor = true;
+	Mesh = ObjectInitializer.CreateDefaultSubobject<UProceduralMeshComponent>(this, TEXT("Dataflow Visualization Component"));
+}
+
+UDataflowRenderingComponent::~UDataflowRenderingComponent()
+{
+	if (RenderMesh) delete RenderMesh;
+}
+
+void UDataflowRenderingComponent::Invalidate()
+{
+	bUpdateRender = true;
+	bBoundsNeedsUpdate = true;
+}
+
+
+void UDataflowRenderingComponent::ResetRenderTargets() 
+{ 
+	RenderTargets.Reset();
+	ResetRenderingCollection();
+}
+
+void UDataflowRenderingComponent::AddRenderTarget(const UDataflowEdNode* InTarget) 
+{ 
+	RenderTargets.AddUnique(InTarget); 
+	ResetRenderingCollection();
+}
+
+void UDataflowRenderingComponent::TickComponent(float DeltaTime, enum ELevelTick TickType, FActorComponentTickFunction *ThisTickFunction)
+{
+	Super::TickComponent(DeltaTime, TickType, ThisTickFunction);
+	RenderProceduralMesh();
+}
+
+void UDataflowRenderingComponent::UpdateLocalBounds()
+{
+	if (bBoundsNeedsUpdate)
+	{
+		GeometryCollection::Facades::FBoundsFacade::UpdateBoundingBox(RenderCollection);
+		bBoundsNeedsUpdate = false;
+	}
+}
+
+FBoxSphereBounds UDataflowRenderingComponent::CalcBounds(const FTransform& LocalToWorldIn) const
+{
+	return BoundingBox.TransformBy(GetComponentTransform());
+}
+
+void UDataflowRenderingComponent::ResetRenderingCollection()
+{
+	Invalidate();
+	RenderCollection = FManagedArrayCollection();
+	ResetProceduralMesh();
+}
+
+void UDataflowRenderingComponent::SetRenderingCollection(FManagedArrayCollection&& InCollection)
+{
+	Invalidate();
+	RenderCollection = InCollection;
+	UpdateLocalBounds();
+	ResetProceduralMesh();
+}
+
+
+void UDataflowRenderingComponent::ResetProceduralMesh()
+{
+	if (Mesh && RenderMesh)
+	{
+		Mesh->ClearAllMeshSections();
+		delete RenderMesh;
+		RenderMesh = nullptr;
+	}
+}
+
+void UDataflowRenderingComponent::RenderProceduralMesh()
+{
+	if (bUpdateRender)
+	{
+		RenderCollection = FManagedArrayCollection();
+		GeometryCollection::Facades::FRenderingFacade Facade(&RenderCollection);
+
+		bool bNeedsRefresh = false;
+		if (Context && Dataflow)
+		{
+			for (const UDataflowEdNode* Target : RenderTargets)
+			{
+				bNeedsRefresh |= Target->Render(Facade, Context);
+			}
+		}
+		bUpdateRender = false;
+	}
+		
+	bool bCanRender = false;	
+	if (GeometryCollection::Facades::FRenderingFacade::IsValid(&RenderCollection))
+	{
+		int32 NumVertices = RenderCollection.NumElements(FGeometryCollection::VerticesGroup);
+		int32 NumFaces = RenderCollection.NumElements(FGeometryCollection::FacesGroup);
+		if (NumFaces && NumVertices)
+		{
+
+			const TManagedArray<FIntVector>& Indices = RenderCollection.GetAttribute<FIntVector>("Indices", FGeometryCollection::FacesGroup);
+			const TManagedArray<FVector3f>& Vertex = RenderCollection.GetAttribute<FVector3f>("Vertex", FGeometryCollection::VerticesGroup);
+
+			if (RenderMesh && RenderMesh->Vertices.Num() != NumFaces * 3)
+			{
+				ResetProceduralMesh();
+			}
+
+			if (!RenderMesh)
+			{
+				RenderMesh = new FRenderMesh;
+
+				for (int i = 0; i < NumFaces; ++i)
+				{
+					const auto& P1 = Vertex[Indices[i][0]];
+					const auto& P2 = Vertex[Indices[i][1]];
+					const auto& P3 = Vertex[Indices[i][2]];
+
+					RenderMesh->Vertices.Add(FVector(P1));
+					RenderMesh->Vertices.Add(FVector(P2));
+					RenderMesh->Vertices.Add(FVector(P3));
+
+					RenderMesh->Colors.Add(FLinearColor::White);
+					RenderMesh->Colors.Add(FLinearColor::White);
+					RenderMesh->Colors.Add(FLinearColor::White);
+
+					RenderMesh->UVs.Add(FVector2D(0, 0));
+					RenderMesh->UVs.Add(FVector2D(0, 0));
+					RenderMesh->UVs.Add(FVector2D(0, 0));
+
+					RenderMesh->Triangles.Add(3 * i);
+					RenderMesh->Triangles.Add(3 * i + 1);
+					RenderMesh->Triangles.Add(3 * i + 2);
+
+					auto Normal = Chaos::FVec3::CrossProduct(P3 - P1, P2 - P1);
+					RenderMesh->Normals.Add(Normal);
+					RenderMesh->Normals.Add(Normal);
+					RenderMesh->Normals.Add(Normal);
+
+					auto Tangent = (P2 - P1).GetSafeNormal();
+					RenderMesh->Tangents.Add(FProcMeshTangent(Tangent[0], Tangent[1], Tangent[2]));
+					Tangent = (P3 - P2).GetSafeNormal();
+					RenderMesh->Tangents.Add(FProcMeshTangent(Tangent[0], Tangent[1], Tangent[2]));
+					Tangent = (P1 - P3).GetSafeNormal();
+					RenderMesh->Tangents.Add(FProcMeshTangent(Tangent[0], Tangent[1], Tangent[2]));
+				}
+				Mesh->CreateMeshSection_LinearColor(0, RenderMesh->Vertices, RenderMesh->Triangles, RenderMesh->Normals, RenderMesh->UVs, RenderMesh->Colors, RenderMesh->Tangents, false);
+			}
+			else
+			{
+				Mesh->UpdateMeshSection_LinearColor(0, RenderMesh->Vertices, RenderMesh->Normals, RenderMesh->UVs, RenderMesh->Colors, RenderMesh->Tangents);
+			}
+			bCanRender = true;
+		}
+	}
+
+	if (!bCanRender)
+	{
+		ResetProceduralMesh();
+	}
+}
+
+
+
diff --git a/Engine/Plugins/Experimental/Dataflow/Source/DataflowEnginePlugin/Public/Dataflow/DataflowEnginePlugin.h b/Engine/Plugins/Experimental/Dataflow/Source/DataflowEnginePlugin/Public/Dataflow/DataflowEnginePlugin.h
new file mode 100644
--- /dev/null
+++ b/Engine/Plugins/Experimental/Dataflow/Source/DataflowEnginePlugin/Public/Dataflow/DataflowEnginePlugin.h
@@ -0,0 +1,39 @@
+// Copyright Epic Games, Inc. All Rights Reserved.
+
+#pragma once
+
+#include "CoreMinimal.h"
+#include "Modules/ModuleInterface.h"
+#include "Modules/ModuleManager.h"
+
+
+/**
+ * The public interface to this module
+ */
+class IDataflowEnginePlugin : public IModuleInterface
+{
+
+public:
+
+	/**
+	 * Singleton-like access to this module's interface.  This is just for convenience!
+	 * Beware of calling this during the shutdown phase, though.  Your module might have been unloaded already.
+	 *
+	 * @return Returns singleton instance, loading the module on demand if needed
+	 */
+	static inline IDataflowEnginePlugin& Get()
+	{
+		return FModuleManager::LoadModuleChecked< IDataflowEnginePlugin >( "DataflowEnginePlugin" );
+	}
+
+	/**
+	 * Checks to see if this module is loaded and ready.  It is only valid to call Get() if IsAvailable() returns true.
+	 *
+	 * @return True if the module is loaded and ready to use
+	 */
+	static inline bool IsAvailable()
+	{
+		return FModuleManager::Get().IsModuleLoaded( "DataflowEnginePlugin" );
+	}
+};
+
diff --git a/Engine/Plugins/Experimental/Dataflow/Source/DataflowEnginePlugin/Public/Dataflow/DataflowRenderingActor.h b/Engine/Plugins/Experimental/Dataflow/Source/DataflowEnginePlugin/Public/Dataflow/DataflowRenderingActor.h
new file mode 100644
--- /dev/null
+++ b/Engine/Plugins/Experimental/Dataflow/Source/DataflowEnginePlugin/Public/Dataflow/DataflowRenderingActor.h
@@ -0,0 +1,30 @@
+// Copyright Epic Games, Inc. All Rights Reserved.
+
+#pragma once
+
+#include "CoreMinimal.h"
+#include "GameFramework/Actor.h"
+#include "UObject/ObjectMacros.h"
+#include "UObject/ObjectPtr.h"
+#include "Dataflow/DataflowRenderingComponent.h"
+
+#include "DataflowRenderingActor.generated.h"
+
+
+UCLASS()
+class DATAFLOWENGINEPLUGIN_API ADataflowRenderingActor : public AActor
+{
+	GENERATED_UCLASS_BODY()
+
+public:
+
+	/* DataflowComponent */
+	UPROPERTY(VisibleAnywhere, BlueprintReadOnly, Category = Destruction, meta = (ExposeFunctionCategories = "Components|Dataflow", AllowPrivateAccess = "true"))
+	TObjectPtr<UDataflowRenderingComponent> DataflowRenderingComponent;
+	UDataflowRenderingComponent* GetDataflowRenderingComponent() const { return DataflowRenderingComponent; }
+
+
+#if WITH_EDITOR
+	virtual bool GetReferencedContentObjects(TArray<UObject*>& Objects) const override;
+#endif
+};
diff --git a/Engine/Plugins/Experimental/Dataflow/Source/DataflowEnginePlugin/Public/Dataflow/DataflowRenderingComponent.h b/Engine/Plugins/Experimental/Dataflow/Source/DataflowEnginePlugin/Public/Dataflow/DataflowRenderingComponent.h
new file mode 100644
--- /dev/null
+++ b/Engine/Plugins/Experimental/Dataflow/Source/DataflowEnginePlugin/Public/Dataflow/DataflowRenderingComponent.h
@@ -0,0 +1,71 @@
+// Copyright Epic Games, Inc. All Rights Reserved.
+#pragma once
+
+#include "GeometryCollection/ManagedArrayCollection.h"
+#include "Components/MeshComponent.h"
+#include "Dataflow/DataflowObject.h"
+#include "UObject/ObjectMacros.h"
+#include "ProceduralMeshComponent.h"
+#include "DataflowRenderingComponent.generated.h"
+
+/**
+*	UDataflowRenderingComponent
+*/
+UCLASS(meta = (BlueprintSpawnableComponent))
+class DATAFLOWENGINEPLUGIN_API UDataflowRenderingComponent : public UMeshComponent
+{
+	GENERATED_UCLASS_BODY()
+
+public:
+	~UDataflowRenderingComponent();
+
+	void Invalidate();
+	void UpdateLocalBounds();
+
+	virtual void TickComponent(float DeltaTime, enum ELevelTick TickType, FActorComponentTickFunction *ThisTickFunction) override;
+	virtual FBoxSphereBounds CalcBounds(const FTransform& LocalToWorld) const override;
+
+	/** Render Targets*/
+	void ResetRenderTargets(); 
+	void AddRenderTarget(const UDataflowEdNode* InTarget); 
+	const TArray<const UDataflowEdNode*>& GetRenderTargets() const {return RenderTargets;}
+
+	/** Context */
+	void SetContext(TSharedPtr<Dataflow::FContext> InContext) { Context = InContext; }
+
+	/** RenderCollection */
+	void ResetRenderingCollection();
+	void SetRenderingCollection(FManagedArrayCollection&& InCollection);
+
+	/** Dataflow */
+	void SetDataflow(const UDataflow* InDataflow) { Dataflow = InDataflow; }
+	const UDataflow* GetDataflow() const { return Dataflow; }
+
+private:
+	TSharedPtr<Dataflow::FContext> Context;
+	TArray<const UDataflowEdNode*> RenderTargets;
+	TObjectPtr< const UDataflow> Dataflow;
+	FManagedArrayCollection RenderCollection;
+	TObjectPtr<UProceduralMeshComponent> Mesh;
+
+	//
+	// Render the Procedural Mesh
+	//
+	struct FRenderMesh
+	{
+		TArray<FVector> Vertices;
+		TArray<int32> Triangles;
+		TArray<FVector> Normals;
+		TArray<FVector2D> UVs;
+		TArray<FLinearColor> Colors;
+		TArray<FProcMeshTangent> Tangents;
+	};
+	FRenderMesh* RenderMesh = nullptr;
+	void RenderProceduralMesh();
+	void ResetProceduralMesh();
+
+	bool bUpdateRender = true;
+	bool bBoundsNeedsUpdate = true;
+	FBoxSphereBounds BoundingBox = FBoxSphereBounds(ForceInitToZero);
+};
+
diff --git a/Engine/Source/Runtime/Experimental/Chaos/Private/GeometryCollection/Facades/CollectionBoundsFacade.cpp b/Engine/Source/Runtime/Experimental/Chaos/Private/GeometryCollection/Facades/CollectionBoundsFacade.cpp
new file mode 100644
--- /dev/null
+++ b/Engine/Source/Runtime/Experimental/Chaos/Private/GeometryCollection/Facades/CollectionBoundsFacade.cpp
@@ -0,0 +1,106 @@
+// Copyright Epic Games, Inc. All Rights Reserved.
+
+/*=============================================================================
+	GeometryCollection.cpp: FGeometryCollection methods.
+=============================================================================*/
+
+#include "GeometryCollection/Facades/CollectionBoundsFacade.h"
+#include "GeometryCollection/GeometryCollection.h"
+
+namespace GeometryCollection::Facades
+{
+
+	// Attributes
+	const FName FBoundsFacade::AAAAttribute = "AAA";
+
+	FBoundsFacade::FBoundsFacade(FManagedArrayCollection& InCollection)
+		: Self(InCollection)
+		, BoundingBox(InCollection,"BoundingBox", FGeometryCollection::GeometryGroup)
+		, Vertex(InCollection,"Vertex", FGeometryCollection::VerticesGroup)
+		, BoneMap(InCollection,"BoneMap", FGeometryCollection::VerticesGroup)
+		, TransformToGeometryIndex(InCollection,"TransformToGeometryIndex", FTransformCollection::TransformGroup)
+	{
+		DefineSchema(Self);
+	}
+
+	//
+	//  Initialization
+	//
+
+		//
+	//  Initialization
+	//
+
+	void FBoundsFacade::DefineSchema(FManagedArrayCollection& InCollection)
+	{
+		FManagedArrayCollection::FConstructionParameters VertexDependency(FGeometryCollection::VerticesGroup);
+
+		// surface rendering attributes
+		if(!InCollection.HasGroup(FGeometryCollection::VerticesGroup))
+		{
+			InCollection.AddGroup(FGeometryCollection::VerticesGroup);
+		}
+		if( !InCollection.HasGroup(FGeometryCollection::FacesGroup) )
+		{
+			InCollection.AddGroup(FGeometryCollection::FacesGroup);
+		}
+
+		InCollection.AddAttribute<FVector3f>("Vertex", FGeometryCollection::VerticesGroup);
+		InCollection.AddAttribute<FIntVector>("Indices", FGeometryCollection::FacesGroup, VertexDependency);
+
+		ensure(InCollection.FindAttributeTyped<FVector3f>("Vertex", FGeometryCollection::VerticesGroup) != nullptr);
+		ensure(InCollection.FindAttributeTyped<FIntVector>("Indices", FGeometryCollection::FacesGroup) != nullptr);
+	}
+
+	bool FBoundsFacade::IsValid(const FManagedArrayCollection& InCollection)
+	{
+		return InCollection.HasAttribute("BoundingBox", FGeometryCollection::GeometryGroup)
+			&& InCollection.HasAttribute("Vertex", FGeometryCollection::VerticesGroup)
+			&& InCollection.HasAttribute("BoneMap", FGeometryCollection::VerticesGroup)
+			&& InCollection.HasAttribute("TransformToGeometryIndex", FTransformCollection::TransformGroup);
+	}
+
+
+
+	void FBoundsFacade::UpdateBoundingBox(FManagedArrayCollection& InCollection, bool bSkipCheck)
+	{
+		if (!bSkipCheck || !IsValid(InCollection))
+		{
+			return;
+		}
+
+		TManagedArray<FBox>& BoundingBox = InCollection.ModifyAttribute<FBox>("BoundingBox", FGeometryCollection::GeometryGroup);
+		const TManagedArray<FVector3f>& Vertex = InCollection.GetAttribute<FVector3f>("Vertex", FGeometryCollection::VerticesGroup);
+		const TManagedArray<int32>& BoneMap = InCollection.GetAttribute<int32>("BoneMap", FGeometryCollection::VerticesGroup);
+		const TManagedArray<int32>& TransformToGeometryIndex = InCollection.GetAttribute<int32>("TransformToGeometryIndex", FTransformCollection::TransformGroup);
+
+		if (BoundingBox.Num())
+		{
+			// Initialize BoundingBox
+			for (int32 Idx = 0; Idx < BoundingBox.Num(); ++Idx)
+			{
+				BoundingBox[Idx].Init();
+			}
+
+			// Compute BoundingBox
+			for (int32 Idx = 0; Idx < Vertex.Num(); ++Idx)
+			{
+				int32 TransformIndexValue = BoneMap[Idx];
+				BoundingBox[TransformToGeometryIndex[TransformIndexValue]] += FVector(Vertex[Idx]);
+			}
+		}
+	}
+
+	//
+	//  GetAttributes
+	//
+
+	const TManagedArray< FBox >* FBoundsFacade::GetBoundingBoxes(const FManagedArrayCollection& InCollection)
+	{
+		return InCollection.FindAttributeTyped <FBox>("BoundingBox", FGeometryCollection::GeometryGroup);
+	}
+
+
+}; // GeometryCollection::Facades
+
+
diff --git a/Engine/Source/Runtime/Experimental/Chaos/Private/GeometryCollection/Facades/CollectionRenderingFacade.cpp b/Engine/Source/Runtime/Experimental/Chaos/Private/GeometryCollection/Facades/CollectionRenderingFacade.cpp
new file mode 100644
--- /dev/null
+++ b/Engine/Source/Runtime/Experimental/Chaos/Private/GeometryCollection/Facades/CollectionRenderingFacade.cpp
@@ -0,0 +1,96 @@
+// Copyright Epic Games, Inc. All Rights Reserved.
+
+/*=============================================================================
+	GeometryCollection.cpp: FGeometryCollection methods.
+=============================================================================*/
+
+#include "GeometryCollection/Facades/CollectionRenderingFacade.h"
+#include "GeometryCollection/GeometryCollection.h"
+
+namespace GeometryCollection::Facades
+{
+
+	// Attributes
+	const FName FRenderingFacade::AAAAttribute = "AAA";
+
+	FRenderingFacade::FRenderingFacade(FManagedArrayCollection* InCollection)
+		: Self(InCollection)
+	{
+		DefineSchema(Self);
+	}
+
+	//
+	//  Initialization
+	//
+
+	void FRenderingFacade::DefineSchema(FManagedArrayCollection* InCollection)
+	{
+		FManagedArrayCollection::FConstructionParameters VertexDependency(FGeometryCollection::VerticesGroup);
+
+		// surface rendering attributes
+		if (!InCollection->HasGroup(FGeometryCollection::VerticesGroup))
+		{
+			InCollection->AddGroup(FGeometryCollection::VerticesGroup);
+		}
+		if (!InCollection->HasGroup(FGeometryCollection::FacesGroup))
+		{
+			InCollection->AddGroup(FGeometryCollection::FacesGroup);
+		}
+
+		InCollection->AddAttribute<FVector3f>("Vertex", FGeometryCollection::VerticesGroup);
+		InCollection->AddAttribute<FIntVector>("Indices", FGeometryCollection::FacesGroup, VertexDependency);
+
+		ensure(InCollection->FindAttributeTyped<FVector3f>("Vertex", FGeometryCollection::VerticesGroup) != nullptr);
+		ensure(InCollection->FindAttributeTyped<FIntVector>("Indices", FGeometryCollection::FacesGroup) != nullptr);
+	}
+
+	bool FRenderingFacade::CanRenderSurface(const FManagedArrayCollection* InCollection)
+	{
+		return InCollection->HasGroup(FGeometryCollection::VerticesGroup)
+			&& InCollection->HasGroup(FGeometryCollection::FacesGroup)
+			&& InCollection->FindAttributeTyped<FVector3f>("Vertex", FGeometryCollection::VerticesGroup)
+			&& InCollection->FindAttributeTyped<FIntVector>("Indices", FGeometryCollection::FacesGroup);
+	}
+
+	bool FRenderingFacade::IsValid(const FManagedArrayCollection* InCollection)
+	{
+		return FRenderingFacade::CanRenderSurface(InCollection);
+	}
+
+	void FRenderingFacade::AddTriangle(FManagedArrayCollection* InCollection, const Chaos::FTriangle& InTriangle)
+	{
+		if (IsValid(InCollection))
+		{
+			auto CollectionVert = [](const Chaos::FVec3& V) { return FVector3f(float(V.X), float(V.Y), float(V.Z)); };
+
+			TManagedArray<FVector3f>& Vertices = InCollection->ModifyAttribute<FVector3f>("Vertex", FGeometryCollection::VerticesGroup);
+			TManagedArray<FIntVector>& Indices = InCollection->ModifyAttribute<FIntVector>("Indices", FGeometryCollection::FacesGroup);
+
+			int32 IndicesStart = InCollection->AddElements(1, FGeometryCollection::FacesGroup);
+			int32 VertexStart = InCollection->AddElements(3, FGeometryCollection::VerticesGroup);
+
+			Indices[IndicesStart] = FIntVector(VertexStart, VertexStart + 1, VertexStart + 2);
+			Vertices[VertexStart] = CollectionVert(InTriangle[0]);
+			Vertices[VertexStart + 1] = CollectionVert(InTriangle[1]);
+			Vertices[VertexStart + 2] = CollectionVert(InTriangle[2]);
+		}
+	}
+
+	//
+	//  GetAttributes
+	//
+
+	const TManagedArray< FIntVector >* FRenderingFacade::GetIndices(const FManagedArrayCollection* InCollection)
+	{
+		return InCollection->FindAttribute<FIntVector>("Indices", FGeometryCollection::FacesGroup);
+	}
+
+	const TManagedArray< FVector3f >* FRenderingFacade::GetVertices(const FManagedArrayCollection* InCollection)
+	{
+		return InCollection->FindAttribute<FVector3f>("Vertex", FGeometryCollection::VerticesGroup);
+	}
+
+
+}; // GeometryCollection::Facades
+
+
diff --git a/Engine/Source/Runtime/Experimental/Chaos/Public/GeometryCollection/Facades/CollectionBoundsFacade.h b/Engine/Source/Runtime/Experimental/Chaos/Public/GeometryCollection/Facades/CollectionBoundsFacade.h
new file mode 100644
--- /dev/null
+++ b/Engine/Source/Runtime/Experimental/Chaos/Public/GeometryCollection/Facades/CollectionBoundsFacade.h
@@ -0,0 +1,83 @@
+// Copyright Epic Games, Inc. All Rights Reserved.
+
+#pragma once
+#include "Containers/UnrealString.h"
+#include "Containers/Array.h"
+#include "Math/MathFwd.h"
+#include "GeometryCollection/ManagedArrayAccessor.h"
+#include "GeometryCollection/ManagedArrayCollection.h"
+#include "GeometryCollection/Facades/CollectionSelectionFacade.h"
+#include "Chaos/Triangle.h"
+
+namespace GeometryCollection::Facades
+{
+
+	/**
+	* FBoundsFacade
+	* 
+	* Defines common API for calculating the bounding box on a collection
+	* 
+	* Usage:
+	*    FBoundsFacade::UpdateBoundingBox(this);
+	* 
+	* Then arrays can be accessed later by:
+	*	const TManagedArray< FIntVector >* FaceIndices = FRenderingFacade::GetIndices(this);
+	*	const TManagedArray< FVector3f >* Vertices = FRenderingFacade::GetVertices(this);
+	*
+	* The following attributes are created on the collection:
+	*
+	*	- FindAttribute<FIntVector>("Indices", FGeometryCollection::FacesGroup);
+	*	- FindAttribute<FVector3f>("Vertex", FGeometryCollection::VerticesGroup);
+	* 
+	*/
+	class CHAOS_API FBoundsFacade
+	{
+		FManagedArrayCollection& Self;
+
+	public:
+
+		// Attributes
+		static const FName AAAAttribute;
+
+		/**
+		* FBoundsFacade Constuctor
+		* @param VertixDependencyGroup : GroupName the index attribute is dependent on.
+		*/
+		FBoundsFacade(FManagedArrayCollection& InSelf);
+
+		/**
+		*  Create the facade.
+		*/
+		static void DefineSchema(FManagedArrayCollection& InCollection);
+		void DefineSchema() { return DefineSchema(Self); }
+
+
+		/**
+		*  Is the Facade defined on the collection?
+		*/
+		static bool IsValid(const FManagedArrayCollection& InCollection);
+		bool IsValid() { return IsValid(Self); }
+
+
+		/**
+		*  Does it support rendering surfaces. 
+		*/
+		static void UpdateBoundingBox(FManagedArrayCollection& InCollection, bool bSkipCheck = false);
+		void UpdateBoundingBox(bool bSkipCheck = false) { UpdateBoundingBox(Self); }
+
+
+		/**
+		* Return the vertex positions from the collection. Null if not initialized.
+		* @param FManagedArrayCollection : Collection
+		*/
+		static const TManagedArray< FBox >* GetBoundingBoxes(const FManagedArrayCollection& InCollection);
+		const TManagedArray< FBox >* GetBoundingBoxes() const { return GetBoundingBoxes(Self); }
+
+private:
+		TManagedArrayAccessor<FBox> BoundingBox;
+		TManagedArrayAccessor<FVector3f> Vertex;
+		TManagedArrayAccessor<int32> BoneMap;
+		TManagedArrayAccessor<int32> TransformToGeometryIndex;
+	};
+
+}
\ No newline at end of file
diff --git a/Engine/Source/Runtime/Experimental/Chaos/Public/GeometryCollection/Facades/CollectionRenderingFacade.h b/Engine/Source/Runtime/Experimental/Chaos/Public/GeometryCollection/Facades/CollectionRenderingFacade.h
new file mode 100644
--- /dev/null
+++ b/Engine/Source/Runtime/Experimental/Chaos/Public/GeometryCollection/Facades/CollectionRenderingFacade.h
@@ -0,0 +1,95 @@
+// Copyright Epic Games, Inc. All Rights Reserved.
+
+#pragma once
+#include "Containers/UnrealString.h"
+#include "Containers/Array.h"
+#include "Math/MathFwd.h"
+#include "GeometryCollection/ManagedArrayCollection.h"
+#include "GeometryCollection/Facades/CollectionSelectionFacade.h"
+#include "Chaos/Triangle.h"
+
+namespace GeometryCollection::Facades
+{
+
+	/**
+	* FRenderingFacade
+	*
+	* Defines common API for storing a vertex weights bound to a bone. This mapping is from the
+	* the vertex to the bone index. The FSelectionFacad will store the mapping from the BoneIndex
+	* to the vertex.
+	* Usage:
+	*    FRenderingFacade::AddBoneWeights(this, FSelectionFacade(this) );
+	*
+	* Then arrays can be accessed later by:
+	*	const TManagedArray< FIntVector >* FaceIndices = FRenderingFacade::GetIndices(this);
+	*	const TManagedArray< FVector3f >* Vertices = FRenderingFacade::GetVertices(this);
+	*
+	* The following attributes are created on the collection:
+	*
+	*	- FindAttribute<FIntVector>("Indices", FGeometryCollection::FacesGroup);
+	*	- FindAttribute<FVector3f>("Vertex", FGeometryCollection::VerticesGroup);
+	*
+	*/
+	class CHAOS_API FRenderingFacade
+	{
+		FManagedArrayCollection* Self;
+
+	public:
+
+		// Attributes
+		static const FName AAAAttribute;
+
+		/**
+		* FRenderingFacade Constuctor
+		* @param VertixDependencyGroup : GroupName the index attribute is dependent on.
+		*/
+		FRenderingFacade(FManagedArrayCollection* InSelf);
+
+		/**
+		*  Create the facade.
+		*/
+		static void DefineSchema(FManagedArrayCollection* InCollection);
+		void DefineSchema() { return DefineSchema(Self); }
+
+
+		/**
+		*  Is the Facade defined on the collection?
+		*/
+		static bool IsValid(const FManagedArrayCollection* InCollection);
+		bool IsValid() { return IsValid(Self); }
+
+
+		/**
+		*  Does it support rendering surfaces.
+		*/
+		static bool CanRenderSurface(const FManagedArrayCollection* InCollection);
+		bool CanRenderSurface() { return CanRenderSurface(Self); }
+
+
+		/**
+		* Using a FSelectionFacade::FSelectionKey, create indexes from the vertices to the driving bones.
+		* @param FManagedArrayCollection : Collection
+		* @param FSelectionFacade::FSelectionKey : Key for weights in the FSelectionFacade
+		*/
+		static void AddTriangle(FManagedArrayCollection* InCollection, const Chaos::FTriangle& InTriangle);
+		void AddTriangle(const Chaos::FTriangle& InTriangle) { AddTriangle(Self, InTriangle); }
+
+		/**
+		* Return the render indices from the collection. Null if not initialized.
+		* @param FManagedArrayCollection : Collection
+		*/
+		static const TManagedArray< FIntVector >* GetIndices(const FManagedArrayCollection* InCollection);
+		const TManagedArray< FIntVector >* GetIndices() const { return GetIndices(Self); }
+
+
+		/**
+		* Return the vertex positions from the collection. Null if not initialized.
+		* @param FManagedArrayCollection : Collection
+		*/
+		static const TManagedArray< FVector3f >* GetVertices(const FManagedArrayCollection* InCollection);
+		const TManagedArray< FVector3f >* GetVertices() const { return GetVertices(Self); }
+
+
+	};
+
+}
\ No newline at end of file
diff --git a/Engine/Source/Runtime/Experimental/Dataflow/Core/Private/Dataflow/DataflowNode.cpp b/Engine/Source/Runtime/Experimental/Dataflow/Core/Private/Dataflow/DataflowNode.cpp
--- a/Engine/Source/Runtime/Experimental/Dataflow/Core/Private/Dataflow/DataflowNode.cpp
+++ b/Engine/Source/Runtime/Experimental/Dataflow/Core/Private/Dataflow/DataflowNode.cpp
@@ -119,8 +119,6 @@ void FDataflowNode::AddOutput(FDataflowOutput* InPtr)
 	}
 }
 
-
-
 FDataflowOutput* FDataflowNode::FindOutput(FName InName)
 {
 	for (TPair<uint32, FDataflowOutput*> Elem : Outputs)
@@ -134,6 +132,19 @@ FDataflowOutput* FDataflowNode::FindOutput(FName InName)
 	return nullptr;
 }
 
+const FDataflowOutput* FDataflowNode::FindOutput(FName InName) const
+{
+	for (TPair<uint32, FDataflowOutput*> Elem : Outputs)
+	{
+		FDataflowOutput* Con = Elem.Value;
+		if (Con->GetName().IsEqual(InName))
+		{
+			return (FDataflowOutput*)Con;
+		}
+	}
+	return nullptr;
+}
+
 const FDataflowOutput* FDataflowNode::FindOutput(const void* Reference) const
 {
 	for (TPair<uint32, FDataflowOutput*> Elem : Outputs)
diff --git a/Engine/Source/Runtime/Experimental/Dataflow/Core/Private/Dataflow/DataflowNodeParameters.cpp b/Engine/Source/Runtime/Experimental/Dataflow/Core/Private/Dataflow/DataflowNodeParameters.cpp
--- a/Engine/Source/Runtime/Experimental/Dataflow/Core/Private/Dataflow/DataflowNodeParameters.cpp
+++ b/Engine/Source/Runtime/Experimental/Dataflow/Core/Private/Dataflow/DataflowNodeParameters.cpp
@@ -1,12 +1,14 @@
 // Copyright Epic Games, Inc. All Rights Reserved.
 
 #include "Dataflow/DataflowNodeParameters.h"
+
 #include "Dataflow/DataflowInputOutput.h"
 #include "Dataflow/DataflowNode.h"
 
 namespace Dataflow
 {
 	uint64 FTimestamp::Invalid = 0;
+	uint64 FTimestamp::Current() { return FPlatformTime::Cycles64();  }
 
 	void BeginContextEvaluation(FContext& Context, const FDataflowNode* Node, const FDataflowOutput* Output)
 	{
diff --git a/Engine/Source/Runtime/Experimental/Dataflow/Core/Public/Dataflow/DataflowGraph.h b/Engine/Source/Runtime/Experimental/Dataflow/Core/Public/Dataflow/DataflowGraph.h
--- a/Engine/Source/Runtime/Experimental/Dataflow/Core/Public/Dataflow/DataflowGraph.h
+++ b/Engine/Source/Runtime/Experimental/Dataflow/Core/Public/Dataflow/DataflowGraph.h
@@ -85,6 +85,18 @@ namespace Dataflow
 			return TSharedPtr<FDataflowNode>(nullptr);
 		}
 
+		TSharedPtr<const FDataflowNode> FindBaseNode(FGuid InGuid) const
+		{
+			for (TSharedPtr<FDataflowNode> Node : Nodes)
+			{
+				if (Node->GetGuid() == InGuid)
+				{
+					return Node;
+				}
+			}
+			return TSharedPtr<FDataflowNode>(nullptr);
+		}
+
 		TSharedPtr<FDataflowNode> FindBaseNode(FName InName)
 		{
 			for (TSharedPtr<FDataflowNode> Node : Nodes)
diff --git a/Engine/Source/Runtime/Experimental/Dataflow/Core/Public/Dataflow/DataflowNode.h b/Engine/Source/Runtime/Experimental/Dataflow/Core/Public/Dataflow/DataflowNode.h
--- a/Engine/Source/Runtime/Experimental/Dataflow/Core/Public/Dataflow/DataflowNode.h
+++ b/Engine/Source/Runtime/Experimental/Dataflow/Core/Public/Dataflow/DataflowNode.h
@@ -75,6 +75,7 @@ struct DATAFLOWCORE_API FDataflowNode
 	virtual FName GetCategory() const { return ""; }
 	virtual FString GetTags() const { return ""; }
 	virtual FString GetToolTip() const { return ""; }
+	virtual TArray<Dataflow::FRenderingParameter> GetRenderParameters() const { return GetRenderParametersImpl(); }
 
 	//
 	// Connections
@@ -99,6 +100,7 @@ struct DATAFLOWCORE_API FDataflowNode
 
 	FDataflowOutput* FindOutput(FName Name);
 	FDataflowOutput* FindOutput(void* Reference);
+	const FDataflowOutput* FindOutput(FName Name) const;
 	const FDataflowOutput* FindOutput(const void* Reference) const;
 
 
@@ -180,7 +182,10 @@ struct DATAFLOWCORE_API FDataflowNode
 
 	bool IsValid() const { return bValid; }
 
+
 private:
+	virtual TArray<Dataflow::FRenderingParameter> GetRenderParametersImpl() const { return TArray<Dataflow::FRenderingParameter>(); }
+
 
 	bool bValid = true;
 
@@ -200,6 +205,12 @@ namespace Dataflow
 				TUniquePtr<A> Val = MakeUnique<A>(FNodeParameters{InParam.Name}, InParam.Guid);    \
 				Val->ValidateConnections(); return Val;});
 
+#define DATAFLOW_NODE_RENDER_TYPE(A, B)												\
+	virtual TArray<Dataflow::FRenderingParameter> GetRenderParametersImpl() const {		\
+		TArray<Dataflow::FRenderingParameter> Array;								\
+		Array.Add({ A, {B,} });														\
+		return Array;}
+
 #define DATAFLOW_NODE_DEFINE_INTERNAL(TYPE, DISPLAY_NAME, CATEGORY, TAGS)			\
 public:																				\
 	static FName StaticType() {return #TYPE;}										\
diff --git a/Engine/Source/Runtime/Experimental/Dataflow/Core/Public/Dataflow/DataflowNodeParameters.h b/Engine/Source/Runtime/Experimental/Dataflow/Core/Public/Dataflow/DataflowNodeParameters.h
--- a/Engine/Source/Runtime/Experimental/Dataflow/Core/Public/Dataflow/DataflowNodeParameters.h
+++ b/Engine/Source/Runtime/Experimental/Dataflow/Core/Public/Dataflow/DataflowNodeParameters.h
@@ -15,10 +15,26 @@ namespace Dataflow
 {
 	struct DATAFLOWCORE_API FTimestamp
 	{
-		static uint64 Invalid;
-		FTimestamp(uint64 InValue) : Value(InValue) {}
-		uint64 Value = 0;
+		typedef uint64 Type;
+		Type Value = Type(0);
+
+		FTimestamp(Type InValue) : Value(InValue) {}
 		bool operator>=(const FTimestamp& InTimestamp) { return Value >= InTimestamp.Value; }
+		bool IsInvalid() { return Value == Invalid; }
+
+		static Type Current();
+		static Type Invalid; // 0
+	};
+
+	struct DATAFLOWCORE_API FRenderingParameter {
+		FRenderingParameter() {}
+		FRenderingParameter(FName InTypeName, const TArray<FName>& InOutputs)
+			: Type(InTypeName), Outputs(InOutputs) {}
+		FRenderingParameter(FName InTypeName, TArray<FName>&& InOutputs)
+			: Type(InTypeName), Outputs(InOutputs) {}
+
+		FName Type = FName("");
+		TArray<FName> Outputs;
 	};
 
 	struct FContextCacheElementBase 
@@ -104,7 +120,7 @@ namespace Dataflow
 		void SetData(size_t Key, FProperty* Property, const T& Value)
 		{
 			int64 IntKey = (int64)Key;
-			TUniquePtr<FContextCacheElement<T>> DataStoreEntry = MakeUnique<FContextCacheElement<T>>(Property, Value, FPlatformTime::Cycles64());
+			TUniquePtr<FContextCacheElement<T>> DataStoreEntry = MakeUnique<FContextCacheElement<T>>(Property, Value, FTimestamp::Current());
 
 			SetDataImpl(IntKey, MoveTemp(DataStoreEntry));
 		}
@@ -113,7 +129,7 @@ namespace Dataflow
 		void SetData(size_t Key, FProperty* Property, T&& Value)
 		{
 			int64 IntKey = (int64)Key;
-			TUniquePtr<FContextCacheElement<T>> DataStoreEntry = MakeUnique<FContextCacheElement<T>>(Property, Forward<T>(Value), FPlatformTime::Cycles64());
+			TUniquePtr<FContextCacheElement<T>> DataStoreEntry = MakeUnique<FContextCacheElement<T>>(Property, Forward<T>(Value), FTimestamp::Current());
 
 			SetDataImpl(IntKey, MoveTemp(DataStoreEntry));
 		}
@@ -139,7 +155,15 @@ namespace Dataflow
 			int64 IntKey = (int64)Key;
 			return HasDataImpl(Key, InTimestamp);
 		}
-		
+
+		virtual bool IsEmptyImpl() const = 0;
+
+		bool IsEmpty() const
+		{
+			return IsEmptyImpl();
+		}
+
+
 		FTimestamp GetTimestamp() const { return Timestamp; }
 		virtual void Evaluate(const FDataflowNode* Node, const FDataflowOutput* Output) = 0;
 		virtual bool Evaluate(const FDataflowOutput& Connection) = 0;
@@ -170,6 +194,11 @@ namespace Dataflow
 			return DataStore.Contains(Key) && DataStore[Key]->Timestamp >= InTimestamp;
 		}
 
+		virtual bool IsEmptyImpl() const override
+		{
+			return DataStore.IsEmpty();
+		}
+
 		virtual void Evaluate(const FDataflowNode* Node, const FDataflowOutput* Output) override;
 		virtual bool Evaluate(const FDataflowOutput& Connection) override;
 	};
@@ -208,6 +237,11 @@ namespace Dataflow
 			return DataStore.Contains(Key) && DataStore[Key]->Timestamp >= InTimestamp;
 		}
 
+		virtual bool IsEmptyImpl() const override
+		{
+			return DataStore.IsEmpty();
+		}
+
 		virtual void Evaluate(const FDataflowNode* Node, const FDataflowOutput* Output) override;
 		virtual bool Evaluate(const FDataflowOutput& Connection) override;
 	};
diff --git a/Engine/Source/Runtime/Experimental/Dataflow/Engine/Private/Dataflow/DataflowEdNode.cpp b/Engine/Source/Runtime/Experimental/Dataflow/Engine/Private/Dataflow/DataflowEdNode.cpp
--- a/Engine/Source/Runtime/Experimental/Dataflow/Engine/Private/Dataflow/DataflowEdNode.cpp
+++ b/Engine/Source/Runtime/Experimental/Dataflow/Engine/Private/Dataflow/DataflowEdNode.cpp
@@ -4,6 +4,9 @@
 
 #include "Dataflow/DataflowNode.h"
 #include "Dataflow/DataflowCore.h"
+#include "Dataflow/DataflowObject.h"
+#include "Dataflow/DataflowRenderingFactory.h"
+#include "GeometryCollection/Facades/CollectionRenderingFacade.h"
 #include "Logging/LogMacros.h"
 
 #include UE_INLINE_GENERATED_CPP_BY_NAME(DataflowEdNode)
@@ -25,6 +28,43 @@ UDataflowEdNode::UDataflowEdNode(const FObjectInitializer& ObjectInitializer)
 #endif // WITH_EDITOR && !UE_BUILD_SHIPPING
 }
 
+void UDataflowEdNode::SetAssetRender(bool bInRender)
+{
+	bRenderInAssetEditor = bInRender;
+	if (IsBound())
+	{
+#if WITH_EDITOR && !(UE_BUILD_SHIPPING || UE_BUILD_TEST)
+		if (UDataflow* DataflowObject = Cast< UDataflow>(GetGraph()))
+		{
+			if (bRenderInAssetEditor)
+				DataflowObject->AddRenderTarget(this);
+			else
+				DataflowObject->RemoveRenderTarget(this);
+		}
+#endif
+	}
+}
+
+
+TSharedPtr<FDataflowNode> UDataflowEdNode::GetDataflowNode()
+{
+	if(TSharedPtr<Dataflow::FGraph> Dataflow = GetDataflowGraph())
+	{
+		return Dataflow->FindBaseNode(GetDataflowNodeGuid());
+	}
+	return TSharedPtr<FDataflowNode>(nullptr);
+}
+
+TSharedPtr<const FDataflowNode> UDataflowEdNode::GetDataflowNode() const
+{
+	if (TSharedPtr<const Dataflow::FGraph> Dataflow = GetDataflowGraph())
+	{
+		return Dataflow->FindBaseNode(GetDataflowNodeGuid());
+	}
+	return TSharedPtr<FDataflowNode>(nullptr);
+}
+
+
 void UDataflowEdNode::AllocateDefaultPins()
 {
 	UE_LOG(DATAFLOWNODE_LOG, Verbose, TEXT("UDataflowEdNode::AllocateDefaultPins()"));
@@ -161,6 +201,37 @@ FLinearColor UDataflowEdNode::GetNodeBodyTintColor() const
 
 #endif
 
+
+TArray<Dataflow::FRenderingParameter> UDataflowEdNode::GetRenderParameters() const
+{
+	if (TSharedPtr<const FDataflowNode> DataflowNode = GetDataflowNode())
+	{
+		return DataflowNode->GetRenderParameters();
+	}
+	return 	TArray<Dataflow::FRenderingParameter>();
+}
+
+
+bool UDataflowEdNode::Render(GeometryCollection::Facades::FRenderingFacade& RenderData, TSharedPtr<Dataflow::FContext> Context) const
+{
+	bool bNeedsRefresh = false;
+	if (DataflowGraph)
+	{
+		if (TSharedPtr<const FDataflowNode> NodeTarget = DataflowGraph->FindBaseNode(FName(GetName())))
+		{
+			if (Dataflow::FRenderingFactory* Factory = Dataflow::FRenderingFactory::GetInstance())
+			{
+				for (Dataflow::FRenderingParameter& Parameter : GetRenderParameters())
+				{
+					Factory->RenderNodeOutput(RenderData, { NodeTarget.Get(), Parameter, *Context });
+					bNeedsRefresh = true;
+				}
+			}
+		}
+	}
+	return bNeedsRefresh;
+}
+
 #undef LOCTEXT_NAMESPACE
 
 
diff --git a/Engine/Source/Runtime/Experimental/Dataflow/Engine/Private/Dataflow/DataflowObject.cpp b/Engine/Source/Runtime/Experimental/Dataflow/Engine/Private/Dataflow/DataflowObject.cpp
--- a/Engine/Source/Runtime/Experimental/Dataflow/Engine/Private/Dataflow/DataflowObject.cpp
+++ b/Engine/Source/Runtime/Experimental/Dataflow/Engine/Private/Dataflow/DataflowObject.cpp
@@ -3,6 +3,7 @@
 #include "Dataflow/DataflowObject.h"
 #include "Dataflow/DataflowCore.h"
 #include "Dataflow/DataflowEdNode.h"
+#include "Dataflow/DataflowNodeParameters.h"
 
 #include UE_INLINE_GENERATED_CPP_BY_NAME(DataflowObject)
 
@@ -56,8 +57,10 @@ void UDataflow::PostLoad()
 
 	for (UEdGraphNode* EdNode : Nodes)
 	{
+		UDataflowEdNode* DataflowEdNode = Cast<UDataflowEdNode>(EdNode);
+
 		// Not all nodes are UDataflowEdNode (There is now UDataflowEdNodeComment)
-		if (UDataflowEdNode* DataflowEdNode = Cast<UDataflowEdNode>(EdNode))
+		if (DataflowEdNode)
 		{
 			DataflowEdNode->SetDataflowGraph(Dataflow);
 		}
@@ -66,12 +69,37 @@ void UDataflow::PostLoad()
 		{
 			EdNode->SetEnabledState(ENodeEnabledState::Disabled);
 		}
+		else
+		{
+			if (DataflowEdNode)
+			{
+				if (DataflowEdNode->DoAssetRender())
+				{
+					RenderTargets.Add(DataflowEdNode);
+				}
+			}
+		}
 	}
 #endif
 
+	LastModifiedRenderTarget = Dataflow::FTimestamp::Current();
 	UObject::PostLoad();
 }
 
+void UDataflow::AddRenderTarget(UDataflowEdNode* InNode)
+{
+	LastModifiedRenderTarget = Dataflow::FTimestamp::Current();
+	InNode->bRenderInAssetEditor = true;
+	RenderTargets.AddUnique(InNode);
+}
+
+void UDataflow::RemoveRenderTarget(UDataflowEdNode* InNode)
+{
+	LastModifiedRenderTarget = Dataflow::FTimestamp::Current();
+	InNode->bRenderInAssetEditor = false;
+	RenderTargets.Remove(InNode);
+}
+
 void UDataflow::Serialize(FArchive& Ar)
 {
 	Super::Serialize(Ar);
diff --git a/Engine/Source/Runtime/Experimental/Dataflow/Engine/Private/Dataflow/DataflowRenderingFactory.cpp b/Engine/Source/Runtime/Experimental/Dataflow/Engine/Private/Dataflow/DataflowRenderingFactory.cpp
new file mode 100644
--- /dev/null
+++ b/Engine/Source/Runtime/Experimental/Dataflow/Engine/Private/Dataflow/DataflowRenderingFactory.cpp
@@ -0,0 +1,29 @@
+// Copyright Epic Games, Inc. All Rights Reserved.
+
+#include "Dataflow/DataflowRenderingFactory.h"
+
+#include "Dataflow/DataflowNode.h"
+#include "Dataflow/DataflowInputOutput.h"
+#include "Misc/MessageDialog.h"
+
+namespace Dataflow
+{
+	FRenderingFactory* FRenderingFactory::Instance = nullptr;
+
+
+	void FRenderingFactory::RenderNodeOutput(GeometryCollection::Facades::FRenderingFacade& RenderData, const FGraphRenderingState& State)
+	{
+		FName OutputType = State.GetRenderType();
+		if (RenderMap.Contains(State.GetRenderType()))
+		{
+			RenderMap[State.GetRenderType()](RenderData, State);
+		}
+		else
+		{
+			UE_LOG(LogChaos, Warning,
+				TEXT("Warning : Dataflow missing output rendering type(%s)"), *State.GetRenderType().ToString());
+
+		}
+	}
+}
+
diff --git a/Engine/Source/Runtime/Experimental/Dataflow/Engine/Public/Dataflow/DataflowEdNode.h b/Engine/Source/Runtime/Experimental/Dataflow/Engine/Public/Dataflow/DataflowEdNode.h
--- a/Engine/Source/Runtime/Experimental/Dataflow/Engine/Public/Dataflow/DataflowEdNode.h
+++ b/Engine/Source/Runtime/Experimental/Dataflow/Engine/Public/Dataflow/DataflowEdNode.h
@@ -4,8 +4,10 @@
 
 #include "CoreMinimal.h"
 #include "Dataflow/DataflowCore.h"
+#include "Dataflow/DataflowRenderingFactory.h"
 #include "EdGraph/EdGraph.h"
 #include "EdGraph/EdGraphNode.h"
+#include "GeometryCollection/GeometryCollection.h"
 #include "Internationalization/Text.h"
 #include "Misc/Build.h"
 #include "Misc/Guid.h"
@@ -18,8 +20,8 @@
 class FArchive;
 class UEdGraphPin;
 class UObject;
-namespace Dataflow { class FGraph; }
-
+namespace Dataflow { class FGraph; class FRenderingParameters; }
+namespace GeometryCollection::Facades { class FRenderingFacade; }
 
 UCLASS()
 class DATAFLOWENGINE_API UDataflowEdNode : public UEdGraphNode
@@ -37,24 +39,38 @@ public:
 #if WITH_EDITOR && !(UE_BUILD_SHIPPING || UE_BUILD_TEST)
 	virtual void PinConnectionListChanged(UEdGraphPin* Pin) override;
 #endif // WITH_EDITOR && !UE_BUILD_SHIPPING
+#if WITH_EDITOR
+	virtual FLinearColor GetNodeTitleColor() const override;
+	virtual FLinearColor GetNodeBodyTintColor() const override;
+#endif // WITH_EDITOR
 	// End of UEdGraphNode interface
 
+	// UObject interface
+	void Serialize(FArchive& Ar);
+	// End UObject interface
+
 	bool IsBound() { return DataflowGraph && DataflowNodeGuid.IsValid(); }
 
-	TSharedPtr<Dataflow::FGraph> GetDataflowGraph() { return DataflowGraph;}
+	TSharedPtr<Dataflow::FGraph> GetDataflowGraph() { return DataflowGraph; }
+	TSharedPtr<const Dataflow::FGraph> GetDataflowGraph() const { return DataflowGraph; }
 	void SetDataflowGraph(TSharedPtr<Dataflow::FGraph> InDataflowGraph) { DataflowGraph = InDataflowGraph; }
 
 	FGuid GetDataflowNodeGuid() const { return DataflowNodeGuid; }
 	void SetDataflowNodeGuid(FGuid InGuid) { DataflowNodeGuid = InGuid; }
 
-#if WITH_EDITOR
-	virtual FLinearColor GetNodeTitleColor() const override;
-	virtual FLinearColor GetNodeBodyTintColor() const override;
-#endif
+	TSharedPtr<FDataflowNode> GetDataflowNode();
+	TSharedPtr<const FDataflowNode> GetDataflowNode() const;
 
-	// UObject interface
-	void Serialize(FArchive& Ar);
-	// End UObject interface
+	//
+	// Node Rendering
+	//
+	void SetAssetRender(bool bInRender);
+	bool DoAssetRender() { return bRenderInAssetEditor; }
+	TArray<Dataflow::FRenderingParameter> GetRenderParameters() const;
+	virtual bool Render(GeometryCollection::Facades::FRenderingFacade& RenderData, TSharedPtr<Dataflow::FContext> Context) const;
+
+	UPROPERTY()
+	bool bRenderInAssetEditor = false;
 
 };
 
diff --git a/Engine/Source/Runtime/Experimental/Dataflow/Engine/Public/Dataflow/DataflowObject.h b/Engine/Source/Runtime/Experimental/Dataflow/Engine/Public/Dataflow/DataflowObject.h
--- a/Engine/Source/Runtime/Experimental/Dataflow/Engine/Public/Dataflow/DataflowObject.h
+++ b/Engine/Source/Runtime/Experimental/Dataflow/Engine/Public/Dataflow/DataflowObject.h
@@ -17,6 +17,7 @@
 class FArchive;
 class UDataflow;
 class UObject;
+class UDataflowEdNode;
 namespace Dataflow { class FGraph; }
 
 
@@ -57,6 +58,8 @@ class DATAFLOWENGINE_API UDataflow : public UEdGraph
 {
 	GENERATED_UCLASS_BODY()
 
+	Dataflow::FTimestamp LastModifiedRenderTarget = Dataflow::FTimestamp::Invalid; 
+	TArray<const UDataflowEdNode*> RenderTargets; // Not Serialized
 	TSharedPtr<Dataflow::FGraph, ESPMode::ThreadSafe> Dataflow;
 	void PostEditCallback();
 
@@ -87,5 +90,14 @@ public:
 		UDataflow* ThisNC = const_cast<UDataflow*>(this);
 		return FDataflowAssetEdit(ThisNC, [ThisNC]() {ThisNC->PostEditCallback(); });
 	}
+
+	//
+	// Render Targets
+	//
+	void AddRenderTarget(UDataflowEdNode*);
+	void RemoveRenderTarget(UDataflowEdNode*);
+	const TArray<const UDataflowEdNode*>& GetRenderTargets() const { return RenderTargets; }
+	const Dataflow::FTimestamp& GetRenderingTimestamp() const { return LastModifiedRenderTarget; }
+
 };
 
diff --git a/Engine/Source/Runtime/Experimental/Dataflow/Engine/Public/Dataflow/DataflowRenderingFactory.h b/Engine/Source/Runtime/Experimental/Dataflow/Engine/Public/Dataflow/DataflowRenderingFactory.h
new file mode 100644
--- /dev/null
+++ b/Engine/Source/Runtime/Experimental/Dataflow/Engine/Public/Dataflow/DataflowRenderingFactory.h
@@ -0,0 +1,94 @@
+// Copyright Epic Games, Inc. All Rights Reserved.
+
+#pragma once
+
+#include "ChaosLog.h"
+#include "CoreMinimal.h"
+#include "Dataflow/DataflowNodeParameters.h"
+#include "Dataflow/DataflowInputOutput.h"
+#include "Dataflow/DataflowNode.h"
+#include "GeometryCollection/Facades/CollectionRenderingFacade.h"
+
+namespace Dataflow
+{
+	class FContext;
+
+	struct DATAFLOWENGINE_API FGraphRenderingState {
+		FGraphRenderingState(const FDataflowNode* InNode, const FRenderingParameter& InParameters, Dataflow::FContext& InContext)
+			: Node(InNode)
+			, RenderType(InParameters.Type)
+			, RenderOutputs(InParameters.Outputs)
+			, Context(InContext)
+		{}
+
+		const FName& GetRenderType() const { return RenderType; }
+		const TArray<FName>& GetRenderOutputs() const { return RenderOutputs; }
+
+		template<class T>
+		const T& GetValue(FName OutputName, const T& Default) const
+		{
+			if (Node)
+			{
+				if (const FDataflowOutput* Output = Node->FindOutput(OutputName))
+				{
+					return Output->GetValue<T>(Context, Default);
+				}
+			}
+			return Default;
+		}
+
+	private:
+		const FDataflowNode* Node = nullptr;
+
+		FName RenderType;
+		TArray<FName> RenderOutputs;
+
+		Dataflow::FContext& Context;
+	};
+
+	//
+	//
+	//
+	class DATAFLOWENGINE_API FRenderingFactory
+	{
+		typedef TFunction<void(GeometryCollection::Facades::FRenderingFacade& RenderData, const FGraphRenderingState& State)> FOutputRenderingFunction;
+
+		// All Maps indexed by TypeName
+		TMap<FName, FOutputRenderingFunction > RenderMap;		// [TypeName] -> NewNodeFunction
+		static FRenderingFactory* Instance;
+		FRenderingFactory() {}
+
+	public:
+		~FRenderingFactory() { delete Instance; }
+
+		static FRenderingFactory* GetInstance()
+		{
+			if (!Instance)
+			{
+				Instance = new FRenderingFactory();
+			}
+			return Instance;
+		}
+
+		void RegisterOutput(const FName& Type, FOutputRenderingFunction InFunction)
+		{
+			if (RenderMap.Contains(Type))
+			{
+				UE_LOG(LogChaos, Warning,
+					TEXT("Warning : Dataflow output rendering registration conflicts with "
+						"existing type(%s)"), *Type.ToString());
+			}
+			else
+			{
+				RenderMap.Add(Type, InFunction);
+			}
+		}
+
+		void RenderNodeOutput(GeometryCollection::Facades::FRenderingFacade& RenderData, const FGraphRenderingState& State);
+
+		bool Contains(FName InType) const { return RenderMap.Contains(InType); }
+
+	};
+
+}
+
