Dataflow : Render toggle button

diff --git a/Engine/Plugins/Experimental/Dataflow/Source/DataflowEditor/Private/Dataflow/DataflowEditorStyle.h b/Engine/Plugins/Experimental/Dataflow/Source/DataflowEditor/Private/Dataflow/DataflowEditorStyle.h
--- a/Engine/Plugins/Experimental/Dataflow/Source/DataflowEditor/Private/Dataflow/DataflowEditorStyle.h
+++ b/Engine/Plugins/Experimental/Dataflow/Source/DataflowEditor/Private/Dataflow/DataflowEditorStyle.h
@@ -5,6 +5,7 @@
 #include "Styling/SlateStyleMacros.h"
 #include "Styling/SlateStyleRegistry.h"
 #include "Interfaces/IPluginManager.h"
+#include "Misc/Paths.h"
 
 class FDataflowEditorStyle final : public FSlateStyleSet
 {
@@ -20,9 +21,13 @@ public:
 		Set("ClassIcon.Dataflow", new FSlateVectorImageBrush(RootToContentDir(TEXT("DataflowAsset_16.svg")), Icon16x16));
 		Set("ClassThumbnail.Dataflow", new FSlateVectorImageBrush(RootToContentDir(TEXT("DataflowAsset_64.svg")), Icon64x64));
 
-		Set("Dataflow.Render.Unknown", new FSlateImageBrush(RootToContentDir(TEXT("Slate/Switch_Undetermined_56x_28x")), Icon56x28));
-		Set("Dataflow.Render.Disabled", new FSlateImageBrush(RootToContentDir(TEXT("Slate/Switch_OFF_56x_28x")), Icon56x28));
-		Set("Dataflow.Render.Enabled", new FSlateImageBrush(RootToContentDir(TEXT("Slate/Switch_ON_56x_28x")), Icon56x28));
+		Set("Dataflow.Render.Unknown", new FSlateImageBrush(RootToContentDir(TEXT("Slate/Switch_Undetermined_56x_28x.png")), Icon28x14));
+		Set("Dataflow.Render.Disabled", new FSlateImageBrush(RootToContentDir(TEXT("Slate/Switch_OFF_56x_28x.png")), Icon28x14));
+		Set("Dataflow.Render.Enabled", new FSlateImageBrush(RootToContentDir(TEXT("Slate/Switch_ON_56x_28x.png")), Icon28x14));
+
+		Set("Dataflow.Cached.False", new FSlateImageBrush(RootToContentDir(TEXT("Slate/status_grey.png")), Icon16x16));
+		Set("Dataflow.Cached.True", new FSlateImageBrush(RootToContentDir(TEXT("Slate/status_green.png")), Icon16x16));
+
 
 		FSlateStyleRegistry::RegisterSlateStyle(*this);
 	}
diff --git a/Engine/Plugins/Experimental/Dataflow/Source/DataflowEditor/Private/Dataflow/DataflowSNode.cpp b/Engine/Plugins/Experimental/Dataflow/Source/DataflowEditor/Private/Dataflow/DataflowSNode.cpp
--- a/Engine/Plugins/Experimental/Dataflow/Source/DataflowEditor/Private/Dataflow/DataflowSNode.cpp
+++ b/Engine/Plugins/Experimental/Dataflow/Source/DataflowEditor/Private/Dataflow/DataflowSNode.cpp
@@ -29,10 +29,13 @@ void SDataflowEdNode::Construct(const FArguments& InArgs, UDataflowEdNode* InNod
 	DataflowGraphNode = Cast<UDataflowEdNode>(InNode);
 	UpdateGraphNode();
 
-	/*
+	
 	const FSlateBrush* DisabledSwitchBrush = FDataflowEditorStyle::Get().GetBrush(TEXT("Dataflow.Render.Disabled"));
 	const FSlateBrush* EnabledSwitchBrush = FDataflowEditorStyle::Get().GetBrush(TEXT("Dataflow.Render.Enabled"));
 
+	//
+	// Render
+	//
 	CheckBoxStyle = FCheckBoxStyle()
 		.SetCheckBoxType(ESlateCheckBoxType::CheckBox)
 		.SetUncheckedImage(*DisabledSwitchBrush)
@@ -42,10 +45,9 @@ void SDataflowEdNode::Construct(const FArguments& InArgs, UDataflowEdNode* InNod
 		.SetCheckedHoveredImage(*EnabledSwitchBrush)
 		.SetCheckedPressedImage(*EnabledSwitchBrush)
 		.SetPadding(FMargin(0, 0, 0, 1));
-	*/
-	//TSharedPtr<SNodePanel> ParentPanel = GetParentPanel();
-
+	
 	RenderCheckBoxWidget = SNew(SCheckBox)
+		.Style(&CheckBoxStyle)
 		.IsChecked_Lambda([this]()-> ECheckBoxState
 			{
 				if (DataflowGraphNode && DataflowGraphNode->DoAssetRender())
@@ -64,8 +66,37 @@ void SDataflowEdNode::Construct(const FArguments& InArgs, UDataflowEdNode* InNod
 						DataflowGraphNode->SetAssetRender(false);
 				}
 			});
-		//.Style(&CheckBoxStyle);
 
+
+	//
+	//  Cached
+	//
+	/*
+	const FSlateBrush* CachedFalseBrush = FDataflowEditorStyle::Get().GetBrush(TEXT("Dataflow.Cached.False"));
+	const FSlateBrush* CachedTrueBrush = FDataflowEditorStyle::Get().GetBrush(TEXT("Dataflow.Cached.True"));
+
+	CacheStatusStyle = FCheckBoxStyle()
+		.SetCheckBoxType(ESlateCheckBoxType::CheckBox)
+		.SetUncheckedImage(*CachedFalseBrush)
+		.SetUncheckedHoveredImage(*CachedFalseBrush)
+		.SetUncheckedPressedImage(*CachedFalseBrush)
+		.SetCheckedImage(*CachedTrueBrush)
+		.SetCheckedHoveredImage(*CachedTrueBrush)
+		.SetCheckedPressedImage(*CachedTrueBrush)
+		.SetPadding(FMargin(0, 0, 0, 1));
+
+		CacheStatus = SNew(SCheckBox)
+		.Style(&CacheStatusStyle)
+		.IsChecked_Lambda([this]()-> ECheckBoxState
+		{
+			if (DataflowGraphNode)
+			{
+				return ECheckBoxState::Checked;
+				// @todo(dataflow) : Missing connection to FToolkit
+			}
+			return ECheckBoxState::Unchecked;
+		});
+	*/
 }
 
 TArray<FOverlayWidgetInfo> SDataflowEdNode::GetOverlayWidgets(bool bSelected, const FVector2D& WidgetSize) const
@@ -83,6 +114,19 @@ TArray<FOverlayWidgetInfo> SDataflowEdNode::GetOverlayWidgets(bool bSelected, co
 		Widgets.Add(Info);
 	}
 
+	/*
+	if (DataflowGraphNode )
+	{
+		// @todo(dataflow) : Need to bump the title box over 6px
+		const FVector2D ImageSize = CacheStatus->GetDesiredSize();
+
+		FOverlayWidgetInfo Info;
+		Info.OverlayOffset = FVector2D(6.f, 6.f);
+		Info.Widget = CacheStatus;
+
+		Widgets.Add(Info);
+	}
+	*/
 	return Widgets;
 }
 
diff --git a/Engine/Plugins/Experimental/Dataflow/Source/DataflowEditor/Public/Dataflow/DataflowSNode.h b/Engine/Plugins/Experimental/Dataflow/Source/DataflowEditor/Public/Dataflow/DataflowSNode.h
--- a/Engine/Plugins/Experimental/Dataflow/Source/DataflowEditor/Public/Dataflow/DataflowSNode.h
+++ b/Engine/Plugins/Experimental/Dataflow/Source/DataflowEditor/Public/Dataflow/DataflowSNode.h
@@ -39,6 +39,9 @@ private:
 
 	FCheckBoxStyle CheckBoxStyle;
 	TSharedPtr<SCheckBox> RenderCheckBoxWidget;
+
+	//FCheckBoxStyle CacheStatusStyle;
+	//TSharedPtr<SCheckBox> CacheStatus;
 };
 
 //
