Dataflow : Invalidation
- Force node's asset to be modified when its property is changed in node editor.

diff --git a/Engine/Plugins/Experimental/Dataflow/Source/DataflowEditor/Private/Dataflow/DataflowEditorActions.cpp b/Engine/Plugins/Experimental/Dataflow/Source/DataflowEditor/Private/Dataflow/DataflowEditorActions.cpp
--- a/Engine/Plugins/Experimental/Dataflow/Source/DataflowEditor/Private/Dataflow/DataflowEditorActions.cpp
+++ b/Engine/Plugins/Experimental/Dataflow/Source/DataflowEditor/Private/Dataflow/DataflowEditorActions.cpp
@@ -181,6 +181,8 @@ void FDataflowEditorCommands::OnPropertyValueChanged(UDataflow* OutDataflow, TSh
 		TSharedPtr<const FDataflowNode> UpdatedNode = nullptr;
 		if (OutDataflow && InPropertyChangedEvent.Property && InPropertyChangedEvent.Property->GetOwnerUObject())
 		{
+			OutDataflow->Modify(true);
+
 			FString Name = InPropertyChangedEvent.Property->GetOwnerUObject()->GetName();
 			Name = Name.Left(Name.Len() - FString("DataflowNode").Len());
 			{
