Dataflow : Changed the way unique name for node works
- first uses the type name and checks if any node with the same name exists in the graph already, if so , then generate a unique name by adding a number at the end

diff --git a/Engine/Plugins/Experimental/Dataflow/Source/DataflowEditor/Private/Dataflow/DataflowSNode.cpp b/Engine/Plugins/Experimental/Dataflow/Source/DataflowEditor/Private/Dataflow/DataflowSNode.cpp
--- a/Engine/Plugins/Experimental/Dataflow/Source/DataflowEditor/Private/Dataflow/DataflowSNode.cpp
+++ b/Engine/Plugins/Experimental/Dataflow/Source/DataflowEditor/Private/Dataflow/DataflowSNode.cpp
@@ -182,12 +182,22 @@ UEdGraphNode* FAssetSchemaAction_Dataflow_CreateNode_DataflowEdNode::PerformActi
 {
 	if (UDataflow* Dataflow = Cast<UDataflow>(ParentGraph))
 	{
-		const FName NodeName = MakeUniqueObjectName(Dataflow, UDataflowEdNode::StaticClass(), FName(GetMenuDescription().ToString()));
+		// by default use the type name and check if it is unique in the context of the graph
+		// if not, then generate a unique name 
+		const FString NodeBaseName = GetMenuDescription().ToString();
+		FName NodeUniqueName{ NodeBaseName };
+		int32 NameIndex= 0;
+		while (Dataflow->GetDataflow()->FindBaseNode(FName(NodeUniqueName)) != nullptr)
+		{ 
+			NodeUniqueName = FName(NodeBaseName + FString::Printf(TEXT("_%d"), NameIndex));
+			NameIndex++;
+		}
+
 		if (Dataflow::FNodeFactory* Factory = Dataflow::FNodeFactory::GetInstance())
 		{
-			if (TSharedPtr<FDataflowNode> DataflowNode = Factory->NewNodeFromRegisteredType(*Dataflow->GetDataflow(), { FGuid::NewGuid(),NodeTypeName,NodeName}))
+			if (TSharedPtr<FDataflowNode> DataflowNode = Factory->NewNodeFromRegisteredType(*Dataflow->GetDataflow(), { FGuid::NewGuid(), NodeTypeName, NodeUniqueName }))
 			{
-				if (UDataflowEdNode* EdNode = NewObject<UDataflowEdNode>(Dataflow, UDataflowEdNode::StaticClass(), NodeName))
+				if (UDataflowEdNode* EdNode = NewObject<UDataflowEdNode>(Dataflow, UDataflowEdNode::StaticClass(), NodeUniqueName))
 				{
 					Dataflow->Modify();
 					if (FromPin != nullptr)
