Dataflow : add abiolity to rename nodes in the geometry collection editor

diff --git a/Engine/Plugins/Experimental/Dataflow/Source/DataflowEditor/Private/Dataflow/DataflowEditorToolkit.cpp b/Engine/Plugins/Experimental/Dataflow/Source/DataflowEditor/Private/Dataflow/DataflowEditorToolkit.cpp
--- a/Engine/Plugins/Experimental/Dataflow/Source/DataflowEditor/Private/Dataflow/DataflowEditorToolkit.cpp
+++ b/Engine/Plugins/Experimental/Dataflow/Source/DataflowEditor/Private/Dataflow/DataflowEditorToolkit.cpp
@@ -94,8 +94,13 @@ TSharedRef<SGraphEditor> FDataflowEditorToolkit::CreateGraphEditorWidget(UDatafl
 		}
 	};
 
+	SGraphEditor::FGraphEditorEvents InEvents;
+	InEvents.OnVerifyTextCommit = FOnNodeVerifyTextCommit::CreateSP(this, &FDataflowEditorToolkit::OnNodeVerifyTitleCommit);
+	InEvents.OnTextCommitted = FOnNodeTextCommitted::CreateSP(this, &FDataflowEditorToolkit::OnNodeTitleCommitted);
+
 	return SNew(SDataflowGraphEditor, DataflowToEdit)
 		.GraphToEdit(DataflowToEdit)
+		.GraphEvents(InEvents)
 		.DetailsView(InNodeDetailsEditor)
 		.EvaluateGraph(Evaluate);
 
@@ -160,6 +165,15 @@ TSharedPtr<IDetailsView> FDataflowEditorToolkit::CreateAssetDetailsEditorWidget(
 
 }
 
+bool FDataflowEditorToolkit::OnNodeVerifyTitleCommit(const FText& NewText, UEdGraphNode* GraphNode, FText& OutErrorMessage)
+{
+	return FDataflowEditorCommands::OnNodeVerifyTitleCommit(NewText, GraphNode, OutErrorMessage);
+}
+
+void FDataflowEditorToolkit::OnNodeTitleCommitted(const FText& InNewText, ETextCommit::Type InCommitType, UEdGraphNode* GraphNode)
+{
+	FDataflowEditorCommands::OnNodeTitleCommitted(InNewText, InCommitType, GraphNode);
+}
 
 TSharedRef<SDockTab> FDataflowEditorToolkit::SpawnTab_GraphCanvas(const FSpawnTabArgs& Args)
 {
diff --git a/Engine/Plugins/Experimental/Dataflow/Source/DataflowEditor/Private/Dataflow/DataflowEditorToolkit.h b/Engine/Plugins/Experimental/Dataflow/Source/DataflowEditor/Private/Dataflow/DataflowEditorToolkit.h
--- a/Engine/Plugins/Experimental/Dataflow/Source/DataflowEditor/Private/Dataflow/DataflowEditorToolkit.h
+++ b/Engine/Plugins/Experimental/Dataflow/Source/DataflowEditor/Private/Dataflow/DataflowEditorToolkit.h
@@ -63,4 +63,7 @@ private:
 	static const FName NodeDetailsTabId;
 	TSharedPtr<IStructureDetailsView> NodeDetailsEditor;
 	TSharedPtr<IStructureDetailsView> CreateNodeDetailsEditorWidget(UObject* ObjectToEdit);
+
+	bool OnNodeVerifyTitleCommit(const FText& NewText, UEdGraphNode* GraphNode, FText& OutErrorMessage);
+	void OnNodeTitleCommitted(const FText& InNewText, ETextCommit::Type InCommitType, UEdGraphNode* GraphNode);
 };
