Dataflow : Generic evaluation from DataflowEditor
- Allow dataflow to perform graph evaluations.

diff --git a/Engine/Plugins/Experimental/ChaosFlesh/Source/ChaosFleshEditor/Private/ChaosFlesh/Asset/AssetTypeActions_FleshAsset.cpp b/Engine/Plugins/Experimental/ChaosFlesh/Source/ChaosFleshEditor/Private/ChaosFlesh/Asset/AssetTypeActions_FleshAsset.cpp
--- a/Engine/Plugins/Experimental/ChaosFlesh/Source/ChaosFleshEditor/Private/ChaosFlesh/Asset/AssetTypeActions_FleshAsset.cpp
+++ b/Engine/Plugins/Experimental/ChaosFlesh/Source/ChaosFleshEditor/Private/ChaosFlesh/Asset/AssetTypeActions_FleshAsset.cpp
@@ -3,7 +3,8 @@
 #include "ChaosFlesh/Asset/AssetTypeActions_FleshAsset.h"
 
 #include "ChaosFlesh/FleshAsset.h"
-#include "ChaosFlesh/ChaosFleshEditorPlugin.h"
+#include "Dataflow/DataflowEditorPlugin.h"
+#include "Dataflow/DataflowEditorToolkit.h"
 #include "ThumbnailRendering/SceneThumbnailInfo.h"
 #include "ToolMenus.h"
 
@@ -34,15 +35,24 @@ void FAssetTypeActions_FleshAsset::GetActions(const TArray<UObject*>& InObjects,
 
 void FAssetTypeActions_FleshAsset::OpenAssetEditor(const TArray<UObject*>& InObjects, TSharedPtr<class IToolkitHost> EditWithinLevelEditor)
 {
+	bool bNeedsBaseEditor = true;
 	EToolkitMode::Type Mode = EditWithinLevelEditor.IsValid() ? EToolkitMode::WorldCentric : EToolkitMode::Standalone;
 	for (auto ObjIt = InObjects.CreateConstIterator(); ObjIt; ++ObjIt)
 	{
 		if (auto Object = Cast<UFleshAsset>(*ObjIt))
 		{
-			IChaosFleshEditorPlugin* FleshEditorModule = &FModuleManager::LoadModuleChecked<IChaosFleshEditorPlugin>("ChaosFleshEditor");
-			FleshEditorModule->CreateFleshAssetEditor(Mode, EditWithinLevelEditor, Object);
+			if (FDataflowEditorToolkit::CanOpenDataflowEditor(Object))
+			{
+				bNeedsBaseEditor = false;
+				IDataflowEditorPlugin* DataflowEditorPlugin = &FModuleManager::LoadModuleChecked<IDataflowEditorPlugin>("DataflowEditor");
+				DataflowEditorPlugin->CreateDataflowAssetEditor(Mode, EditWithinLevelEditor, Object);
+			}
 		}
 	}
+	if (bNeedsBaseEditor)
+	{
+		FAssetTypeActions_Base::OpenAssetEditor(InObjects, EditWithinLevelEditor);
+	}
 }
 
 #undef LOCTEXT_NAMESPACE
diff --git a/Engine/Plugins/Experimental/ChaosFlesh/Source/ChaosFleshEditor/Private/ChaosFlesh/ChaosFleshEditorPlugin.cpp b/Engine/Plugins/Experimental/ChaosFlesh/Source/ChaosFleshEditor/Private/ChaosFlesh/ChaosFleshEditorPlugin.cpp
--- a/Engine/Plugins/Experimental/ChaosFlesh/Source/ChaosFleshEditor/Private/ChaosFlesh/ChaosFleshEditorPlugin.cpp
+++ b/Engine/Plugins/Experimental/ChaosFlesh/Source/ChaosFleshEditor/Private/ChaosFlesh/ChaosFleshEditorPlugin.cpp
@@ -13,6 +13,7 @@
 #include "ChaosFlesh/Asset/FleshAssetThumbnailRenderer.h"
 #include "ChaosFlesh/Editor/FleshEditorToolkit.h"
 #include "CoreMinimal.h"
+#include "Dataflow/DataflowEditorToolkit.h"
 #include "Features/IModularFeatures.h"
 #include "HAL/ConsoleManager.h"
 #include "LevelEditor.h"
@@ -68,9 +69,9 @@ void IChaosFleshEditorPlugin::ShutdownModule()
 
 TSharedRef<FAssetEditorToolkit> IChaosFleshEditorPlugin::CreateFleshAssetEditor(const EToolkitMode::Type Mode, const TSharedPtr<IToolkitHost>& InitToolkitHost, UObject* FleshAsset)
 {
-	TSharedPtr<FFleshEditorToolkit> NewFleshAssetEditor = MakeShared<FFleshEditorToolkit>();
-	NewFleshAssetEditor->InitFleshAssetEditor(Mode, InitToolkitHost, FleshAsset);
-	return StaticCastSharedPtr<FAssetEditorToolkit>(NewFleshAssetEditor).ToSharedRef();
+	TSharedPtr<FDataflowEditorToolkit> NewAssetEditor = MakeShared<FDataflowEditorToolkit>();
+	NewAssetEditor->InitDataflowEditor(Mode, InitToolkitHost, FleshAsset);
+	return StaticCastSharedPtr<FAssetEditorToolkit>(NewAssetEditor).ToSharedRef();
 }
 
 IMPLEMENT_MODULE(IChaosFleshEditorPlugin, FleshAssetEditor)
diff --git a/Engine/Plugins/Experimental/ChaosFlesh/Source/ChaosFleshEditor/Private/ChaosFlesh/Editor/FleshEditorToolkit.cpp b/Engine/Plugins/Experimental/ChaosFlesh/Source/ChaosFleshEditor/Private/ChaosFlesh/Editor/FleshEditorToolkit.cpp
--- a/Engine/Plugins/Experimental/ChaosFlesh/Source/ChaosFleshEditor/Private/ChaosFlesh/Editor/FleshEditorToolkit.cpp
+++ b/Engine/Plugins/Experimental/ChaosFlesh/Source/ChaosFleshEditor/Private/ChaosFlesh/Editor/FleshEditorToolkit.cpp
@@ -51,17 +51,17 @@ void FFleshEditorToolkit::InitFleshAssetEditor(const EToolkitMode::Type Mode, co
 	FleshAsset = CastChecked<UFleshAsset>(ObjectToEdit);
 	if (FleshAsset != nullptr)
 	{
-		if (FleshAsset->Dataflow == nullptr)
+		if (FleshAsset->DataflowAsset == nullptr)
 		{
 			const FName NodeName = MakeUniqueObjectName(FleshAsset, UDataflow::StaticClass(), FName("DataflowFleshAsset"));
-			FleshAsset->Dataflow = NewObject<UDataflow>(FleshAsset, NodeName);
+			FleshAsset->DataflowAsset = NewObject<UDataflow>(FleshAsset, NodeName);
 		}
-		FleshAsset->Dataflow->Schema = UDataflowSchema::StaticClass();
-		Dataflow = FleshAsset->Dataflow;
+		FleshAsset->DataflowAsset->Schema = UDataflowSchema::StaticClass();
+		Dataflow = FleshAsset->DataflowAsset;
 
 		NodeDetailsEditor = CreateNodeDetailsEditorWidget(ObjectToEdit);
 		AssetDetailsEditor = CreateAssetDetailsEditorWidget(FleshAsset);
-		GraphEditor = CreateGraphEditorWidget(FleshAsset->Dataflow, NodeDetailsEditor);
+		GraphEditor = CreateGraphEditorWidget(FleshAsset->DataflowAsset, NodeDetailsEditor);
 		SkeletalEditor = CreateSkeletalEditorWidget(FleshAsset->SkeletalMesh);
 
 		Context = TSharedPtr< Dataflow::FEngineContext>(new Dataflow::FFleshContext(FleshAsset, Dataflow, FPlatformTime::Cycles64()));
@@ -133,7 +133,7 @@ void FFleshEditorToolkit::Tick(float DeltaTime)
 				LastNodeTimestamp = Dataflow::FTimestamp::Invalid;
 			}
 
-			FDataflowEditorCommands::EvaluateNode(*Context.Get(), LastNodeTimestamp, Dataflow, nullptr, nullptr, FleshAsset->Terminal);
+			FDataflowEditorCommands::EvaluateNode(*Context.Get(), LastNodeTimestamp, Dataflow, nullptr, nullptr, FleshAsset->DataflowTerminal);
 		}
 	}
 }
diff --git a/Engine/Plugins/Experimental/ChaosFlesh/Source/ChaosFleshEngine/Public/ChaosFlesh/FleshAsset.h b/Engine/Plugins/Experimental/ChaosFlesh/Source/ChaosFleshEngine/Public/ChaosFlesh/FleshAsset.h
--- a/Engine/Plugins/Experimental/ChaosFlesh/Source/ChaosFleshEngine/Public/ChaosFlesh/FleshAsset.h
+++ b/Engine/Plugins/Experimental/ChaosFlesh/Source/ChaosFleshEngine/Public/ChaosFlesh/FleshAsset.h
@@ -98,10 +98,10 @@ public:
 	// Dataflow
 	//
 	UPROPERTY(EditAnywhere, Category = "Dataflow")
-	TObjectPtr<UDataflow> Dataflow;
+	TObjectPtr<UDataflow> DataflowAsset;
 
 	UPROPERTY(EditAnywhere, Category = "Dataflow")
-	FString Terminal;
+	FString DataflowTerminal = "FleshAssetTerminal";
 
 	//
 	// SkeletalMesh
diff --git a/Engine/Plugins/Experimental/ChaosFlesh/Source/ChaosFleshNodes/Private/Dataflow/ChaosFleshEngineAssetNodes.cpp b/Engine/Plugins/Experimental/ChaosFlesh/Source/ChaosFleshNodes/Private/Dataflow/ChaosFleshEngineAssetNodes.cpp
--- a/Engine/Plugins/Experimental/ChaosFlesh/Source/ChaosFleshNodes/Private/Dataflow/ChaosFleshEngineAssetNodes.cpp
+++ b/Engine/Plugins/Experimental/ChaosFlesh/Source/ChaosFleshNodes/Private/Dataflow/ChaosFleshEngineAssetNodes.cpp
@@ -19,8 +19,7 @@ namespace Dataflow
 	void RegisterChaosFleshEngineAssetNodes()
 	{
 		DATAFLOW_NODE_REGISTER_CREATION_FACTORY(FGetFleshAssetDataflowNode);
-		DATAFLOW_NODE_REGISTER_CREATION_FACTORY(FExampleFleshEditDataflowNode);
-		DATAFLOW_NODE_REGISTER_CREATION_FACTORY(FSetFleshAssetDataflowNode);
+		DATAFLOW_NODE_REGISTER_CREATION_FACTORY(FFleshAssetTerminalDataflowNode);
 		DATAFLOW_NODE_REGISTER_CREATION_FACTORY(FImportFleshDataflowNode);
 		DATAFLOW_NODE_REGISTER_CREATION_FACTORY(FConstructTetGridNode);
 		DATAFLOW_NODE_REGISTER_CREATION_FACTORY(FComputeFleshMassNode);
@@ -45,38 +44,19 @@ void FGetFleshAssetDataflowNode::Evaluate(Dataflow::FContext& Context, const FDa
 	}
 }
 
-void FExampleFleshEditDataflowNode::Evaluate(Dataflow::FContext& Context, const FDataflowOutput* Out) const 
+void FFleshAssetTerminalDataflowNode::SetAssetValue(TObjectPtr<UObject> Asset, Dataflow::FContext& Context) const
 {
-	if (Out->IsA<FManagedArrayCollection>(&Collection))
+	if (UFleshAsset* FleshAsset = Cast<UFleshAsset>(Asset.Get()))
 	{
-		FManagedArrayCollection InputCollection = GetValue<FManagedArrayCollection>(Context, &Collection);
-
-		if (TManagedArray<FVector3f>* Vertex = InputCollection.FindAttribute<FVector3f>("Vertex", "Vertices"))
-		{
-			for (int i = 0; i < Vertex->Num(); i++)
-			{
-				(*Vertex)[i][1] *= Scale;
-			}
-		}
-
-		SetValue<FManagedArrayCollection>(Context, InputCollection, &Collection);
+		const FManagedArrayCollection& InCollection = GetValue<FManagedArrayCollection>(Context, &Collection);
+		FleshAsset->SetCollection(InCollection.NewCopy<FFleshCollection>());
 	}
 }
 
-void FSetFleshAssetDataflowNode::Evaluate(Dataflow::FContext& Context, const FDataflowOutput* Out) const
+void FFleshAssetTerminalDataflowNode::Evaluate(Dataflow::FContext& Context, const FDataflowOutput* Out) const
 {
-	if (Out == nullptr)
-	{
-		const FManagedArrayCollection& Collection = GetValue<FManagedArrayCollection>(Context, &Input);
-
-		if (const Dataflow::FEngineContext* EngineContext = Context.AsType<Dataflow::FEngineContext>())
-		{
-			if (UFleshAsset* FleshAsset = Cast<UFleshAsset>(EngineContext->Owner))
-			{
-				FleshAsset->SetCollection(Collection.NewCopy<FFleshCollection>());
-			}
-		}
-	}
+	const FManagedArrayCollection& InCollection = GetValue<FManagedArrayCollection>(Context, &Collection);
+	SetValue<FManagedArrayCollection>(Context, InCollection, &Collection);
 }
 
 void FImportFleshDataflowNode::Evaluate(Dataflow::FContext& Context, const FDataflowOutput* Out) const
diff --git a/Engine/Plugins/Experimental/ChaosFlesh/Source/ChaosFleshNodes/Public/Dataflow/ChaosFleshEngineAssetNodes.h b/Engine/Plugins/Experimental/ChaosFlesh/Source/ChaosFleshNodes/Public/Dataflow/ChaosFleshEngineAssetNodes.h
--- a/Engine/Plugins/Experimental/ChaosFlesh/Source/ChaosFleshNodes/Public/Dataflow/ChaosFleshEngineAssetNodes.h
+++ b/Engine/Plugins/Experimental/ChaosFlesh/Source/ChaosFleshNodes/Public/Dataflow/ChaosFleshEngineAssetNodes.h
@@ -8,6 +8,7 @@
 #include "ChaosLog.h"
 #include "ChaosFlesh/TetrahedralCollection.h"
 #include "Dataflow/DataflowEngine.h"
+#include "ChaosFlesh/FleshAsset.h"
 
 #include "ChaosFleshEngineAssetNodes.generated.h"
 
@@ -33,51 +34,26 @@ public:
 };
 
 
-
-
 USTRUCT()
-struct FExampleFleshEditDataflowNode : public FDataflowNode
+struct FFleshAssetTerminalDataflowNode : public FDataflowTerminalNode
 {
 	GENERATED_USTRUCT_BODY()
-	DATAFLOW_NODE_DEFINE_INTERNAL(FExampleFleshEditDataflowNode, "ExampleFleshEdit", "Flesh", "")
-	DATAFLOW_NODE_RENDER_TYPE(FGeometryCollection::StaticType(), "Collection")
+	DATAFLOW_NODE_DEFINE_INTERNAL(FFleshAssetTerminalDataflowNode, "FleshAssetTerminal", "Flesh", "")
 
 public:
 
-	UPROPERTY(EditAnywhere, Category = "Dataflow")
-	float Scale = 10.0;
-
-	UPROPERTY(meta = (DataflowInput, DataflowOutput, DisplayName = "Collection", Passthrough = "Collection"))
+	UPROPERTY(meta = (DataflowInput, DataflowOutput, Passthrough = "Collection", DisplayName = "Collection"))
 	FManagedArrayCollection Collection;
 
-	FExampleFleshEditDataflowNode(const Dataflow::FNodeParameters& InParam, FGuid InGuid = FGuid::NewGuid())
-		: FDataflowNode(InParam, InGuid)
-	{
-		RegisterInputConnection(&Collection);
-		RegisterOutputConnection(&Collection, &Collection);
-	}
-
-	virtual void Evaluate(Dataflow::FContext& Context, const FDataflowOutput* Out) const override;
-};
-
-
-USTRUCT()
-struct FSetFleshAssetDataflowNode : public FDataflowTerminalNode
-{
-	GENERATED_USTRUCT_BODY()
-	DATAFLOW_NODE_DEFINE_INTERNAL(FSetFleshAssetDataflowNode, "SetFleshAsset", "Flesh", "")
-
-public:
 
-	UPROPERTY(meta = (DataflowInput, DisplayName = "Collection"))
-	FManagedArrayCollection Input;
-
-	FSetFleshAssetDataflowNode(const Dataflow::FNodeParameters& InParam, FGuid InGuid = FGuid::NewGuid())
+	FFleshAssetTerminalDataflowNode(const Dataflow::FNodeParameters& InParam, FGuid InGuid = FGuid::NewGuid())
 	: FDataflowTerminalNode(InParam, InGuid)
 	{
-		RegisterInputConnection(&Input);
+		RegisterInputConnection(&Collection);
+		RegisterOutputConnection(&Collection, &Collection);
 	}
 
+	virtual void SetAssetValue(TObjectPtr<UObject> Asset, Dataflow::FContext& Context) const override;
 	virtual void Evaluate(Dataflow::FContext& Context, const FDataflowOutput* Out) const override;
 };
 
diff --git a/Engine/Plugins/Experimental/Dataflow/Source/DataflowEditor/Private/Dataflow/DataflowEditorActions.cpp b/Engine/Plugins/Experimental/Dataflow/Source/DataflowEditor/Private/Dataflow/DataflowEditorActions.cpp
--- a/Engine/Plugins/Experimental/Dataflow/Source/DataflowEditor/Private/Dataflow/DataflowEditorActions.cpp
+++ b/Engine/Plugins/Experimental/Dataflow/Source/DataflowEditor/Private/Dataflow/DataflowEditorActions.cpp
@@ -121,6 +121,52 @@ void FDataflowEditorCommands::EvaluateNode(Dataflow::FContext& Context, Dataflow
 	}
 }
 
+void FDataflowEditorCommands::EvaluateTerminalNode(Dataflow::FContext& Context, Dataflow::FTimestamp& OutLastNodeTimestamp,
+	const UDataflow* Dataflow, const FDataflowNode* InNode, const FDataflowOutput* Output, UObject* InAsset, FString NodeName)
+{
+	if (Dataflow)
+	{
+		const FDataflowNode* Node = InNode;
+		if (Node == nullptr)
+		{
+			if (const TSharedPtr<Dataflow::FGraph> Graph = Dataflow->GetDataflow())
+			{
+				if (TSharedPtr<const FDataflowNode> GraphNode = Graph->FindBaseNode(FName(NodeName)))
+				{
+					Node = GraphNode.Get();
+				}
+			}
+		}
+
+		if (Node != nullptr)
+		{
+			if (Output == nullptr)
+			{
+				if (Node->GetTimestamp() >= OutLastNodeTimestamp)
+				{
+					Context.Evaluate(Node, nullptr);
+					OutLastNodeTimestamp = Context.GetTimestamp();
+
+					if (const FDataflowTerminalNode* TerminalNode = Node->AsType<const FDataflowTerminalNode>())
+					{
+						if (InAsset)
+						{
+							TerminalNode->SetAssetValue(InAsset, Context);
+						}
+					}
+				}
+			}
+			else // Output != nullptr
+			{
+				if (!Context.HasData(Output->CacheKey(), Context.GetTimestamp()))
+				{
+					Context.Evaluate(Node, Output);
+				}
+			}
+		}
+	}
+}
+
 
 bool FDataflowEditorCommands::OnNodeVerifyTitleCommit(const FText& NewText, UEdGraphNode* GraphNode, FText& OutErrorMessage)
 {
@@ -181,7 +227,7 @@ void FDataflowEditorCommands::OnPropertyValueChanged(UDataflow* OutDataflow, TSh
 		TSharedPtr<const FDataflowNode> UpdatedNode = nullptr;
 		if (OutDataflow && InPropertyChangedEvent.Property && InPropertyChangedEvent.Property->GetOwnerUObject())
 		{
-			OutDataflow->Modify(true);
+			OutDataflow->MarkPackageDirty();
 
 			FString Name = InPropertyChangedEvent.Property->GetOwnerUObject()->GetName();
 			Name = Name.Left(Name.Len() - FString("DataflowNode").Len());
diff --git a/Engine/Plugins/Experimental/Dataflow/Source/DataflowEditor/Private/Dataflow/DataflowEditorPlugin.cpp b/Engine/Plugins/Experimental/Dataflow/Source/DataflowEditor/Private/Dataflow/DataflowEditorPlugin.cpp
--- a/Engine/Plugins/Experimental/Dataflow/Source/DataflowEditor/Private/Dataflow/DataflowEditorPlugin.cpp
+++ b/Engine/Plugins/Experimental/Dataflow/Source/DataflowEditor/Private/Dataflow/DataflowEditorPlugin.cpp
@@ -6,7 +6,7 @@
 #include "AssetToolsModule.h"
 #include "CoreMinimal.h"
 #include "EdGraphUtilities.h"
-#include "DataflowEditorToolkit.h"
+#include "Dataflow/DataflowEditorToolkit.h"
 #include "Dataflow/DataflowNodeFactory.h"
 #include "Dataflow/DataflowAssetActions.h"
 #include "Dataflow/DataflowSNodeFactories.h"
@@ -46,10 +46,10 @@ void IDataflowEditorPlugin::ShutdownModule()
 	}
 }
 
-TSharedRef<FAssetEditorToolkit> IDataflowEditorPlugin::CreateDataflowAssetEditor(const EToolkitMode::Type Mode, const TSharedPtr<IToolkitHost>& InitToolkitHost, UObject* FleshAsset)
+TSharedRef<FAssetEditorToolkit> IDataflowEditorPlugin::CreateDataflowAssetEditor(const EToolkitMode::Type Mode, const TSharedPtr<IToolkitHost>& InitToolkitHost, UObject* Asset)
 {
 	TSharedPtr<FDataflowEditorToolkit> NewDataflowAssetEditor = MakeShared<FDataflowEditorToolkit>();
-	NewDataflowAssetEditor->InitDataflowEditor(Mode, InitToolkitHost, FleshAsset);
+	NewDataflowAssetEditor->InitDataflowEditor(Mode, InitToolkitHost, Asset);
 	return StaticCastSharedPtr<FAssetEditorToolkit>(NewDataflowAssetEditor).ToSharedRef();
 }
 
diff --git a/Engine/Plugins/Experimental/Dataflow/Source/DataflowEditor/Private/Dataflow/DataflowEditorToolkit.cpp b/Engine/Plugins/Experimental/Dataflow/Source/DataflowEditor/Private/Dataflow/DataflowEditorToolkit.cpp
--- a/Engine/Plugins/Experimental/Dataflow/Source/DataflowEditor/Private/Dataflow/DataflowEditorToolkit.cpp
+++ b/Engine/Plugins/Experimental/Dataflow/Source/DataflowEditor/Private/Dataflow/DataflowEditorToolkit.cpp
@@ -1,6 +1,6 @@
 // Copyright Epic Games, Inc. All Rights Reserved.
 
-#include "DataflowEditorToolkit.h"
+#include "Dataflow/DataflowEditorToolkit.h"
 
 #include "Dataflow/DataflowEditorActions.h"
 #include "Dataflow/DataflowGraphEditor.h"
@@ -29,9 +29,55 @@ const FName FDataflowEditorToolkit::AssetDetailsTabId(TEXT("DataflowEditor_Asset
 const FName FDataflowEditorToolkit::NodeDetailsTabId(TEXT("DataflowEditor_NodeDetails"));
 
 
+UDataflow* GetDataflowFrom(UObject* InObject)
+{
+	if (UClass* Class = InObject->GetClass())
+	{
+		if (FProperty* Property = Class->FindPropertyByName(FName("DataflowAsset")))
+		{
+			return *Property->ContainerPtrToValuePtr<UDataflow*>(InObject);
+		}
+	}
+	return nullptr;
+
+}
+
+FString GetDataflowTerminalFrom(UObject* InObject)
+{
+	if (UClass* Class = InObject->GetClass())
+	{
+		if (FProperty* Property = Class->FindPropertyByName(FName("DataflowTerminal")))
+		{
+			return *Property->ContainerPtrToValuePtr<FString>(InObject);
+		}
+	}
+	return FString();
+}
+
+bool FDataflowEditorToolkit::CanOpenDataflowEditor(UObject* ObjectToEdit)
+{
+	UDataflow* Dataflow = Cast<UDataflow>(ObjectToEdit);
+	if (Dataflow == nullptr)
+	{
+		Dataflow = GetDataflowFrom(ObjectToEdit);
+	}
+	return Dataflow != nullptr;
+}
+
 void FDataflowEditorToolkit::InitDataflowEditor(const EToolkitMode::Type Mode, const TSharedPtr<IToolkitHost>& InitToolkitHost, UObject* ObjectToEdit)
 {
-	Dataflow = CastChecked<UDataflow>(ObjectToEdit);
+	Asset = nullptr;
+	Dataflow = Cast<UDataflow>(ObjectToEdit);
+	if (Dataflow == nullptr)
+	{
+		Dataflow = GetDataflowFrom(ObjectToEdit);
+		if (Dataflow)
+		{
+			Asset = ObjectToEdit;
+			TerminalPath = GetDataflowTerminalFrom(ObjectToEdit);
+		}
+	}
+
 	if (Dataflow != nullptr)
 	{
 		Dataflow->Schema = UDataflowSchema::StaticClass();
@@ -40,6 +86,9 @@ void FDataflowEditorToolkit::InitDataflowEditor(const EToolkitMode::Type Mode, c
 		AssetDetailsEditor = CreateAssetDetailsEditorWidget(ObjectToEdit);
 		GraphEditor = CreateGraphEditorWidget(Dataflow, NodeDetailsEditor);
 
+		Context = TSharedPtr< Dataflow::FEngineContext>(new Dataflow::FAssetContext(Asset, Dataflow, FPlatformTime::Cycles64()));
+		LastNodeTimestamp = Context->GetTimestamp();
+
 		const TSharedRef<FTabManager::FLayout> StandaloneDefaultLayout = FTabManager::NewLayout("Dataflow_Layout.V1")
 			->AddArea
 			(
@@ -75,10 +124,39 @@ void FDataflowEditorToolkit::InitDataflowEditor(const EToolkitMode::Type Mode, c
 		const bool bCreateDefaultStandaloneMenu = true;
 		const bool bCreateDefaultToolbar = true;
 		FAssetEditorToolkit::InitAssetEditor(Mode, InitToolkitHost, FName(TEXT("DataflowEditorApp")), StandaloneDefaultLayout, bCreateDefaultToolbar, bCreateDefaultStandaloneMenu, ObjectToEdit);
+
+
+		AddEditingObject(Dataflow);
+		if (Asset) 
+		{
+			AddEditingObject(Asset);
+		}
 	}
 }
 
+void FDataflowEditorToolkit::OnPropertyValueChanged(const FPropertyChangedEvent& PropertyChangedEvent)
+{
+	FDataflowEditorCommands::OnPropertyValueChanged(this->GetDataflow(), Context, LastNodeTimestamp, PropertyChangedEvent);
+}
 
+void FDataflowEditorToolkit::Tick(float DeltaTime)
+{
+	if (Dataflow && Asset)
+	{
+		if (!Context)
+		{
+			Context = TSharedPtr< Dataflow::FEngineContext>(new Dataflow::FAssetContext(Asset, Dataflow, Dataflow::FTimestamp::Invalid));
+			LastNodeTimestamp = Dataflow::FTimestamp::Invalid;
+		}
+		TerminalPath = GetDataflowTerminalFrom(Asset);
+		FDataflowEditorCommands::EvaluateTerminalNode(*Context.Get(), LastNodeTimestamp, Dataflow, nullptr, nullptr, Asset, TerminalPath);
+	}
+}
+
+TStatId FDataflowEditorToolkit::GetStatId() const
+{
+	RETURN_QUICK_DECLARE_CYCLE_STAT(FFleshEditorToolkit, STATGROUP_Tickables);
+}
 
 TSharedRef<SGraphEditor> FDataflowEditorToolkit::CreateGraphEditorWidget(UDataflow* DataflowToEdit, TSharedPtr<IStructureDetailsView> InNodeDetailsEditor)
 {
@@ -89,8 +167,13 @@ TSharedRef<SGraphEditor> FDataflowEditorToolkit::CreateGraphEditorWidget(UDatafl
 	{
 		if (DataflowToEdit)
 		{
-			FMessageDialog::Open(EAppMsgType::Ok, NSLOCTEXT("Dataflow", "UnknownEvaluationContext_Dataflow", 
-				"Unsupported dataflow evaluation.\nEvaluate the dataflow graph through its associated asset editor."));
+			if (!Context)
+			{
+				Context = TSharedPtr< Dataflow::FEngineContext>(new Dataflow::FAssetContext(Asset, Dataflow, Dataflow::FTimestamp::Invalid));
+			}
+			LastNodeTimestamp = Dataflow::FTimestamp::Invalid;
+
+			FDataflowEditorCommands::EvaluateTerminalNode(*Context.Get(), LastNodeTimestamp, Dataflow, Node, Out, Asset, TerminalPath);
 		}
 	};
 
@@ -133,16 +216,9 @@ TSharedPtr<IStructureDetailsView> FDataflowEditorToolkit::CreateNodeDetailsEdito
 	}
 	TSharedPtr<IStructureDetailsView> DetailsView = PropertyEditorModule.CreateStructureDetailView(DetailsViewArgs, StructureViewArgs, nullptr);
 	DetailsView->GetDetailsView()->SetObject(ObjectToEdit);
-	DetailsView->GetOnFinishedChangingPropertiesDelegate().AddLambda([this](const FPropertyChangedEvent& PropertyChangedEvent) 
-	{
-			if (PropertyChangedEvent.ChangeType == EPropertyChangeType::ValueSet)
-			{
-			}
-	});
-
+	DetailsView->GetOnFinishedChangingPropertiesDelegate().AddSP(this, &FDataflowEditorToolkit::OnPropertyValueChanged);
 
 	return DetailsView;
-
 }
 
 TSharedPtr<IDetailsView> FDataflowEditorToolkit::CreateAssetDetailsEditorWidget(UObject* ObjectToEdit)
@@ -232,17 +308,34 @@ void FDataflowEditorToolkit::RegisterTabSpawners(const TSharedRef<FTabManager>&
 	FAssetEditorToolkit::RegisterTabSpawners(InTabManager);
 }
 
-
 FName FDataflowEditorToolkit::GetToolkitFName() const
 {
 	return FName("DataflowEditor");
 }
 
+FText FDataflowEditorToolkit::GetToolkitName() const
+{
+	if (Asset)
+	{
+		return  GetLabelForObject(Asset);
+	}
+	if (Dataflow)
+	{
+		return  GetLabelForObject(Dataflow);
+	}
+	return  LOCTEXT("ToolkitName", "Empty Dataflow Editor");
+}
+
 FText FDataflowEditorToolkit::GetBaseToolkitName() const
 {
 	return LOCTEXT("AppLabel", "Dataflow Editor");
 }
 
+FText FDataflowEditorToolkit::GetToolkitToolTipText() const
+{
+	return  LOCTEXT("ToolkitToolTipText", "Dataflow Editor");
+}
+
 FString FDataflowEditorToolkit::GetWorldCentricTabPrefix() const
 {
 	return LOCTEXT("WorldCentricTabPrefix", "Dataflow").ToString();
diff --git a/Engine/Plugins/Experimental/Dataflow/Source/DataflowEditor/Private/Dataflow/DataflowEditorToolkit.h b/Engine/Plugins/Experimental/Dataflow/Source/DataflowEditor/Private/Dataflow/DataflowEditorToolkit.h
deleted file mode 100644
diff --git a/Engine/Plugins/Experimental/Dataflow/Source/DataflowEditor/Private/Dataflow/DataflowGraphEditor.cpp b/Engine/Plugins/Experimental/Dataflow/Source/DataflowEditor/Private/Dataflow/DataflowGraphEditor.cpp
--- a/Engine/Plugins/Experimental/Dataflow/Source/DataflowEditor/Private/Dataflow/DataflowGraphEditor.cpp
+++ b/Engine/Plugins/Experimental/Dataflow/Source/DataflowEditor/Private/Dataflow/DataflowGraphEditor.cpp
@@ -2,7 +2,7 @@
 
 #include "Dataflow/DataflowGraphEditor.h"
 
-#include "DataflowEditorToolkit.h"
+#include "Dataflow/DataflowEditorToolkit.h"
 #include "Dataflow/DataflowEngine.h"
 #include "Dataflow/DataflowSNodeFactories.h"
 #include "BoneDragDropOp.h"
diff --git a/Engine/Plugins/Experimental/Dataflow/Source/DataflowEditor/Private/Dataflow/DataflowSNode.cpp b/Engine/Plugins/Experimental/Dataflow/Source/DataflowEditor/Private/Dataflow/DataflowSNode.cpp
--- a/Engine/Plugins/Experimental/Dataflow/Source/DataflowEditor/Private/Dataflow/DataflowSNode.cpp
+++ b/Engine/Plugins/Experimental/Dataflow/Source/DataflowEditor/Private/Dataflow/DataflowSNode.cpp
@@ -103,7 +103,7 @@ TArray<FOverlayWidgetInfo> SDataflowEdNode::GetOverlayWidgets(bool bSelected, co
 {
 	TArray<FOverlayWidgetInfo> Widgets = SGraphNode::GetOverlayWidgets(bSelected, WidgetSize);
 
-	if (DataflowGraphNode && DataflowGraphNode->GetDataflowNode()->GetRenderParameters().Num())
+	if (DataflowGraphNode && DataflowGraphNode->GetDataflowNode() && DataflowGraphNode->GetDataflowNode()->GetRenderParameters().Num())
 	{
 		const FVector2D ImageSize = RenderCheckBoxWidget->GetDesiredSize();
 
diff --git a/Engine/Plugins/Experimental/Dataflow/Source/DataflowEditor/Public/Dataflow/DataflowEditorActions.h b/Engine/Plugins/Experimental/Dataflow/Source/DataflowEditor/Public/Dataflow/DataflowEditorActions.h
--- a/Engine/Plugins/Experimental/Dataflow/Source/DataflowEditor/Public/Dataflow/DataflowEditorActions.h
+++ b/Engine/Plugins/Experimental/Dataflow/Source/DataflowEditor/Public/Dataflow/DataflowEditorActions.h
@@ -65,7 +65,11 @@ public:
 	*/
 	static void EvaluateNode(Dataflow::FContext& Context, Dataflow::FTimestamp& OutLastNodeTimestamp,
 		const UDataflow* Dataflow, const FDataflowNode* Node = nullptr, const FDataflowOutput* Out = nullptr, 
-		FString NodeName = FString());
+		FString NodeName = FString()); // @todo(Dataflow) deprecate  
+
+	static void EvaluateTerminalNode(Dataflow::FContext& Context, Dataflow::FTimestamp& OutLastNodeTimestamp,
+		const UDataflow* Dataflow, const FDataflowNode* Node = nullptr, const FDataflowOutput* Out = nullptr,
+		UObject* InAsset = nullptr, FString NodeName = FString());
 
 	/*
 	*  DeleteNodes
diff --git a/Engine/Plugins/Experimental/Dataflow/Source/DataflowEditor/Public/Dataflow/DataflowEditorPlugin.h b/Engine/Plugins/Experimental/Dataflow/Source/DataflowEditor/Public/Dataflow/DataflowEditorPlugin.h
--- a/Engine/Plugins/Experimental/Dataflow/Source/DataflowEditor/Public/Dataflow/DataflowEditorPlugin.h
+++ b/Engine/Plugins/Experimental/Dataflow/Source/DataflowEditor/Public/Dataflow/DataflowEditorPlugin.h
@@ -16,7 +16,7 @@ class FDataflowAssetActions;
 /**
  * The public interface to this module
  */
-class IDataflowEditorPlugin : public IModuleInterface
+class DATAFLOWEDITOR_API IDataflowEditorPlugin : public IModuleInterface
 {
 public:
 	virtual void StartupModule();
diff --git a/Engine/Plugins/Experimental/Dataflow/Source/DataflowEditor/Public/Dataflow/DataflowEditorToolkit.h b/Engine/Plugins/Experimental/Dataflow/Source/DataflowEditor/Public/Dataflow/DataflowEditorToolkit.h
new file mode 100644
--- /dev/null
+++ b/Engine/Plugins/Experimental/Dataflow/Source/DataflowEditor/Public/Dataflow/DataflowEditorToolkit.h
@@ -0,0 +1,111 @@
+// Copyright Epic Games, Inc. All Rights Reserved.
+
+#pragma once
+
+#include "CoreMinimal.h"
+#include "Dataflow/DataflowObjectInterface.h"
+#include "Toolkits/SimpleAssetEditor.h"
+#include "Misc/NotifyHook.h"
+#include "GraphEditor.h"
+#include "TickableEditorObject.h"
+
+class IDetailsView;
+class FTabManager;
+class IStructureDetailsView;
+class IToolkitHost;
+class UDataflow;
+
+namespace Dataflow
+{
+	class DATAFLOWEDITOR_API FAssetContext : public TEngineContext<FContextSingle>
+	{
+	public:
+		DATAFLOW_CONTEXT_INTERNAL(TEngineContext<FContextSingle>, FAssetContext);
+
+		FAssetContext(UObject* InOwner, UDataflow* InGraph, FTimestamp InTimestamp)
+			: Super(InOwner, InGraph, InTimestamp)
+		{}
+	};
+}
+
+class DATAFLOWEDITOR_API FDataflowEditorToolkit : public FAssetEditorToolkit, public FTickableEditorObject, public FNotifyHook, public FGCObject
+{
+public:
+	static bool CanOpenDataflowEditor(UObject* ObjectToEdit);
+
+	void InitDataflowEditor(const EToolkitMode::Type Mode, const TSharedPtr<IToolkitHost>& InitToolkitHost, UObject* ObjectToEdit);
+
+	//~ Begin DataflowEditorActions
+	void OnPropertyValueChanged(const FPropertyChangedEvent& PropertyChangedEvent);
+	//~ End DataflowEditorActions
+
+	// IToolkit Interface
+	virtual FName GetToolkitFName() const override;
+	virtual FText GetToolkitName() const override;
+	virtual FText GetBaseToolkitName() const override;
+	virtual FText GetToolkitToolTipText() const override;
+	virtual FString GetWorldCentricTabPrefix() const override;
+	virtual FLinearColor GetWorldCentricTabColorScale() const override;
+
+	//~ Begin FTickableEditorObject interface
+	virtual void Tick(float DeltaTime) override;
+	virtual bool IsTickable() const override { return true; }
+	virtual TStatId GetStatId() const override;
+	//~ End FTickableEditorObject Interface
+
+	// FGCObject Interface
+	virtual void AddReferencedObjects(FReferenceCollector& Collector) override;
+	virtual FString GetReferencerName() const override; 
+
+
+
+	// Tab spawners 
+	virtual void RegisterTabSpawners(const TSharedRef<FTabManager>& TabManager) override;
+	TSharedRef<SDockTab> SpawnTab_GraphCanvas(const FSpawnTabArgs& Args);
+	TSharedRef<SDockTab> SpawnTab_AssetDetails(const FSpawnTabArgs& Args);
+	TSharedRef<SDockTab> SpawnTab_NodeDetails(const FSpawnTabArgs& Args);
+
+	// Member Access
+	UObject* GetAsset() { return Asset; }
+	const UObject* GetAsset() const { return Asset; }
+
+	UDataflow* GetDataflow() { return Dataflow; }
+	const UDataflow* GetDataflow() const { return Dataflow; }
+
+	TSharedPtr<Dataflow::FEngineContext> GetContext() { return Context; }
+	TSharedPtr<const Dataflow::FEngineContext> GetContext() const { return Context; }
+
+	TSharedPtr<IDetailsView> GetAssetDetailsEditor() {return AssetDetailsEditor;}
+	const TSharedPtr<IDetailsView> GetAssetDetailsEditor() const { return AssetDetailsEditor; }
+
+	TSharedPtr<IStructureDetailsView> GetNodeDetailsEditor() { return NodeDetailsEditor; }
+	const TSharedPtr<IStructureDetailsView> GetNodeDetailsEditor() const { return NodeDetailsEditor; }
+
+	TSharedPtr<SGraphEditor> GetGraphEditor() { return GraphEditor; }
+	const TSharedPtr<SGraphEditor> GetGraphEditor() const { return GraphEditor; }
+
+private:
+
+	UObject* Asset = nullptr;
+	UDataflow* Dataflow = nullptr;
+	FString TerminalPath = "";
+
+	static const FName GraphCanvasTabId;
+	TSharedPtr<SGraphEditor> GraphEditor;
+	TSharedPtr<FUICommandList> GraphEditorCommands;
+	TSharedRef<SGraphEditor> CreateGraphEditorWidget(UDataflow* ObjectToEdit, TSharedPtr<IStructureDetailsView> PropertiesEditor);
+
+	static const FName AssetDetailsTabId;
+	TSharedPtr<IDetailsView> AssetDetailsEditor;
+	TSharedPtr<IDetailsView> CreateAssetDetailsEditorWidget(UObject* ObjectToEdit);
+
+	static const FName NodeDetailsTabId;
+	TSharedPtr<IStructureDetailsView> NodeDetailsEditor;
+	TSharedPtr<IStructureDetailsView> CreateNodeDetailsEditorWidget(UObject* ObjectToEdit);
+
+	bool OnNodeVerifyTitleCommit(const FText& NewText, UEdGraphNode* GraphNode, FText& OutErrorMessage);
+	void OnNodeTitleCommitted(const FText& InNewText, ETextCommit::Type InCommitType, UEdGraphNode* GraphNode);
+
+	TSharedPtr<Dataflow::FEngineContext> Context;
+	Dataflow::FTimestamp LastNodeTimestamp = Dataflow::FTimestamp::Invalid;
+};
diff --git a/Engine/Plugins/Experimental/GeometryCollectionPlugin/Source/GeometryCollectionEditor/Private/GeometryCollection/AssetTypeActions_GeometryCollection.cpp b/Engine/Plugins/Experimental/GeometryCollectionPlugin/Source/GeometryCollectionEditor/Private/GeometryCollection/AssetTypeActions_GeometryCollection.cpp
--- a/Engine/Plugins/Experimental/GeometryCollectionPlugin/Source/GeometryCollectionEditor/Private/GeometryCollection/AssetTypeActions_GeometryCollection.cpp
+++ b/Engine/Plugins/Experimental/GeometryCollectionPlugin/Source/GeometryCollectionEditor/Private/GeometryCollection/AssetTypeActions_GeometryCollection.cpp
@@ -2,16 +2,12 @@
 
 #include "GeometryCollection/AssetTypeActions_GeometryCollection.h"
 
-#include "GeometryCollection/GeometryCollectionEditorToolkit.h"
-#include "GeometryCollection/GeometryCollectionEditorPlugin.h"
+#include "Dataflow/DataflowEditorPlugin.h"
+#include "Dataflow/DataflowEditorToolkit.h"
 #include "ThumbnailRendering/SceneThumbnailInfo.h"
 #include "GeometryCollection/GeometryCollectionObject.h"
 #include "ToolMenus.h"
 
-bool bGeometryCollectionDataflowEditor = false;
-FAutoConsoleVariableRef CVarGeometryCollectionDataflowEditor(TEXT("p.Chaos.GeometryCollection.DataflowEditor"), bGeometryCollectionDataflowEditor, TEXT("Enable dataflow asset editor on geometry collection assets(Curently Dev-Only)"));
-
-
 #define LOCTEXT_NAMESPACE "AssetTypeActions"
 
 UClass* FAssetTypeActions_GeometryCollection::GetSupportedClass() const
@@ -34,21 +30,25 @@ UThumbnailInfo* FAssetTypeActions_GeometryCollection::GetThumbnailInfo(UObject*
 
 void FAssetTypeActions_GeometryCollection::OpenAssetEditor(const TArray<UObject*>& InObjects, TSharedPtr<class IToolkitHost> EditWithinLevelEditor)
 {
-	if (bGeometryCollectionDataflowEditor)
+	bool bNeedsBaseEditor = true;
+
+	EToolkitMode::Type Mode = EditWithinLevelEditor.IsValid() ? EToolkitMode::WorldCentric : EToolkitMode::Standalone;
+	for (auto ObjIt = InObjects.CreateConstIterator(); ObjIt; ++ObjIt)
 	{
-		EToolkitMode::Type Mode = EditWithinLevelEditor.IsValid() ? EToolkitMode::WorldCentric : EToolkitMode::Standalone;
-		for (auto ObjIt = InObjects.CreateConstIterator(); ObjIt; ++ObjIt)
+		if (auto Object = Cast<UGeometryCollection>(*ObjIt))
 		{
-			if (auto Object = Cast<UGeometryCollection>(*ObjIt))
+			if (FDataflowEditorToolkit::CanOpenDataflowEditor(Object))
 			{
-				IGeometryCollectionEditorPlugin* EditorModule = &FModuleManager::LoadModuleChecked<IGeometryCollectionEditorPlugin>("GeometryCollectionEditor");
-				EditorModule->CreateGeometryCollectionAssetEditor(Mode, EditWithinLevelEditor, Object);
+				bNeedsBaseEditor = false;
+				IDataflowEditorPlugin* DataflowEditorPlugin = &FModuleManager::LoadModuleChecked<IDataflowEditorPlugin>("DataflowEditor");
+				DataflowEditorPlugin->CreateDataflowAssetEditor(Mode, EditWithinLevelEditor, Object);
 			}
 		}
 	}
-	else
-	{ 
-		FAssetTypeActions_Base::OpenAssetEditor(InObjects,EditWithinLevelEditor);
+
+	if (bNeedsBaseEditor)
+	{
+		FAssetTypeActions_Base::OpenAssetEditor(InObjects, EditWithinLevelEditor);
 	}
 }
 
diff --git a/Engine/Plugins/Experimental/GeometryCollectionPlugin/Source/GeometryCollectionEditor/Private/GeometryCollection/GeometryCollectionEditorToolkit.cpp b/Engine/Plugins/Experimental/GeometryCollectionPlugin/Source/GeometryCollectionEditor/Private/GeometryCollection/GeometryCollectionEditorToolkit.cpp
--- a/Engine/Plugins/Experimental/GeometryCollectionPlugin/Source/GeometryCollectionEditor/Private/GeometryCollection/GeometryCollectionEditorToolkit.cpp
+++ b/Engine/Plugins/Experimental/GeometryCollectionPlugin/Source/GeometryCollectionEditor/Private/GeometryCollection/GeometryCollectionEditorToolkit.cpp
@@ -42,13 +42,13 @@ void FGeometryCollectionEditorToolkit::InitGeometryCollectionAssetEditor(const E
 
 	if (GeometryCollection != nullptr)
 	{
-		if (GeometryCollection->Dataflow == nullptr)
+		if (GeometryCollection->DataflowAsset == nullptr)
 		{
 			const FName NodeName = MakeUniqueObjectName(GeometryCollection, UDataflow::StaticClass(), FName("GeometryCollectionDataflowAsset"));
-			GeometryCollection->Dataflow = NewObject<UDataflow>(GeometryCollection, NodeName);
+			GeometryCollection->DataflowAsset = NewObject<UDataflow>(GeometryCollection, NodeName);
 		}
-		Dataflow = GeometryCollection->Dataflow;
-		GeometryCollection->Dataflow->Schema = UDataflowSchema::StaticClass();
+		Dataflow = GeometryCollection->DataflowAsset;
+		GeometryCollection->DataflowAsset->Schema = UDataflowSchema::StaticClass();
 
 		NodeDetailsEditor = CreateNodeDetailsEditorWidget(ObjectToEdit);
 		AssetDetailsEditor = CreateAssetDetailsEditorWidget(GeometryCollection);
@@ -112,7 +112,7 @@ void FGeometryCollectionEditorToolkit::Tick(float DeltaTime)
 				LastNodeTimestamp = Dataflow::FTimestamp::Invalid;
 			}
 
-			FDataflowEditorCommands::EvaluateNode(*Context.Get(), LastNodeTimestamp, Dataflow, nullptr, nullptr, GeometryCollection->Terminal);
+			FDataflowEditorCommands::EvaluateNode(*Context.Get(), LastNodeTimestamp, Dataflow, nullptr, nullptr, GeometryCollection->DataflowTerminal);
 		}
 	}
 }
diff --git a/Engine/Plugins/Experimental/GeometryCollectionPlugin/Source/GeometryCollectionNodes/Private/Dataflow/Nodes/GeometryCollectionAssetNodes.cpp b/Engine/Plugins/Experimental/GeometryCollectionPlugin/Source/GeometryCollectionNodes/Private/Dataflow/Nodes/GeometryCollectionAssetNodes.cpp
--- a/Engine/Plugins/Experimental/GeometryCollectionPlugin/Source/GeometryCollectionNodes/Private/Dataflow/Nodes/GeometryCollectionAssetNodes.cpp
+++ b/Engine/Plugins/Experimental/GeometryCollectionPlugin/Source/GeometryCollectionNodes/Private/Dataflow/Nodes/GeometryCollectionAssetNodes.cpp
@@ -11,7 +11,7 @@ namespace Dataflow
 {
 	void GeometryCollectionEngineAssetNodes()
 	{
-		DATAFLOW_NODE_REGISTER_CREATION_FACTORY(FSetGeometryCollectionAssetDataflowNode);
+		DATAFLOW_NODE_REGISTER_CREATION_FACTORY(FGeometryCollectionTerminalDataflowNode);
 		DATAFLOW_NODE_REGISTER_CREATION_FACTORY(FGetGeometryCollectionAssetDataflowNode);
 		DATAFLOW_NODE_REGISTER_CREATION_FACTORY(FGetGeometryCollectionSourcesDataflowNode);
 		DATAFLOW_NODE_REGISTER_CREATION_FACTORY(FCreateGeometryCollectionFromSourcesDataflowNode);
@@ -21,34 +21,33 @@ namespace Dataflow
 
 // ===========================================================================================================================
 
-FSetGeometryCollectionAssetDataflowNode::FSetGeometryCollectionAssetDataflowNode(const Dataflow::FNodeParameters& InParam, FGuid InGuid)
-	: FDataflowTerminalNode(InParam, InGuid)
-{
-	RegisterInputConnection(&Collection);
-	RegisterInputConnection(&Materials);
-}
-
-void FSetGeometryCollectionAssetDataflowNode::Evaluate(Dataflow::FContext& Context) const
+void FGeometryCollectionTerminalDataflowNode::SetAssetValue(TObjectPtr<UObject> Asset, Dataflow::FContext& Context) const
 {
 	using FGeometryCollectionPtr = TSharedPtr<FGeometryCollection, ESPMode::ThreadSafe>;
 	using FMaterialArray = TArray<TObjectPtr<UMaterial>>;
 
-	// todo(chaos) : should certainly have some form of errors handling  
-	if (const Dataflow::FEngineContext* EngineContext = Context.AsType<Dataflow::FEngineContext>())
+	if (UGeometryCollection* CollectionAsset = Cast<UGeometryCollection>(Asset.Get()))
 	{
-		if (UGeometryCollection* CollectionAsset = Cast<UGeometryCollection>(EngineContext->Owner))
+		if (FGeometryCollectionPtr GeometryCollection = CollectionAsset->GetGeometryCollection())
 		{
-			if (FGeometryCollectionPtr GeometryCollection = CollectionAsset->GetGeometryCollection())
-			{
-				const FManagedArrayCollection& InCollection = GetValue<FManagedArrayCollection>(Context, &Collection);
-				const FMaterialArray& InMaterials = GetValue<FMaterialArray>(Context, &Materials);
-
-				CollectionAsset->ResetFrom(InCollection, InMaterials);
-			}
+			const FManagedArrayCollection& InCollection = GetValue<FManagedArrayCollection>(Context, &Collection);
+			const FMaterialArray& InMaterials = GetValue<FMaterialArray>(Context, &Materials);
+			CollectionAsset->ResetFrom(InCollection, InMaterials);
 		}
 	}
 }
 
+void FGeometryCollectionTerminalDataflowNode::Evaluate(Dataflow::FContext& Context, const FDataflowOutput* Out) const
+{
+	using FMaterialArray = TArray<TObjectPtr<UMaterial>>;
+
+	const FManagedArrayCollection& InCollection = GetValue<FManagedArrayCollection>(Context, &Collection);
+	const FMaterialArray& InMaterials = GetValue<FMaterialArray>(Context, &Materials);
+
+	SetValue<FManagedArrayCollection>(Context, InCollection, &Collection);
+	SetValue<FMaterialArray>(Context, InMaterials, &Materials);
+}
+
 // ===========================================================================================================================
 
 FGetGeometryCollectionAssetDataflowNode::FGetGeometryCollectionAssetDataflowNode(const Dataflow::FNodeParameters& InParam, FGuid InGuid)
diff --git a/Engine/Plugins/Experimental/GeometryCollectionPlugin/Source/GeometryCollectionNodes/Public/Dataflow/Nodes/GeometryCollectionAssetNodes.h b/Engine/Plugins/Experimental/GeometryCollectionPlugin/Source/GeometryCollectionNodes/Public/Dataflow/Nodes/GeometryCollectionAssetNodes.h
--- a/Engine/Plugins/Experimental/GeometryCollectionPlugin/Source/GeometryCollectionNodes/Public/Dataflow/Nodes/GeometryCollectionAssetNodes.h
+++ b/Engine/Plugins/Experimental/GeometryCollectionPlugin/Source/GeometryCollectionNodes/Public/Dataflow/Nodes/GeometryCollectionAssetNodes.h
@@ -11,28 +11,36 @@
 
 class UGeometryCollection;
 
-/**
- * Set the geometry collection asset   
-  */
 USTRUCT()
-struct FSetGeometryCollectionAssetDataflowNode : public FDataflowTerminalNode
+struct FGeometryCollectionTerminalDataflowNode : public FDataflowTerminalNode
 {
 	GENERATED_USTRUCT_BODY()
-	DATAFLOW_NODE_DEFINE_INTERNAL(FSetGeometryCollectionAssetDataflowNode, "SetGeometryCollectionAsset", "GeometryCollection", "")
+	DATAFLOW_NODE_DEFINE_INTERNAL(FGeometryCollectionTerminalDataflowNode, "GeometryCollectionTerminal", "GeometryCollection", "")
 
 public:
-	FSetGeometryCollectionAssetDataflowNode(const Dataflow::FNodeParameters& InParam, FGuid InGuid = FGuid::NewGuid());
-	virtual void Evaluate(Dataflow::FContext& Context) const override;
 
-	/** Attribute collection to use for this asset */ 
-	UPROPERTY(meta = (DataflowInput, DisplayName = "Collection"))
+	UPROPERTY(meta = (DataflowInput, DataflowOutput, Passthrough = "Collection", DisplayName = "Collection"))
 	FManagedArrayCollection Collection;
 
 	/** Materials array to use for this asset */
-	UPROPERTY(meta = (DataflowInput, DisplayName = "Materials"))
+	UPROPERTY(meta = (DataflowInput, DataflowOutput, Passthrough = "Materials", DisplayName = "Materials"))
 	TArray<TObjectPtr<UMaterial>> Materials;
+
+	FGeometryCollectionTerminalDataflowNode(const Dataflow::FNodeParameters& InParam, FGuid InGuid = FGuid::NewGuid())
+		: FDataflowTerminalNode(InParam, InGuid)
+	{
+		RegisterInputConnection(&Collection);
+		RegisterOutputConnection(&Collection, &Collection);
+		RegisterInputConnection(&Materials);
+		RegisterOutputConnection(&Materials, &Materials);
+	}
+
+	virtual void SetAssetValue(TObjectPtr<UObject> Asset, Dataflow::FContext& Context) const override;
+	virtual void Evaluate(Dataflow::FContext& Context, const FDataflowOutput* Out) const override;
 };
 
+
+
 /**
  * Get Current geometry collection asset 
  * Note : Use with caution as this may get replaced in a near future for a more generic getAsset node
diff --git a/Engine/Source/Runtime/Experimental/Dataflow/Core/Private/Dataflow/DataflowTerminalNode.cpp b/Engine/Source/Runtime/Experimental/Dataflow/Core/Private/Dataflow/DataflowTerminalNode.cpp
--- a/Engine/Source/Runtime/Experimental/Dataflow/Core/Private/Dataflow/DataflowTerminalNode.cpp
+++ b/Engine/Source/Runtime/Experimental/Dataflow/Core/Private/Dataflow/DataflowTerminalNode.cpp
@@ -7,12 +7,3 @@
 #include "Dataflow/DataflowArchive.h"
 
 #include UE_INLINE_GENERATED_CPP_BY_NAME(DataflowTerminalNode)
-
-bool FDataflowTerminalNode::ValidateConnections()
-{
-	if (ensureMsgf(NumOutputs() == 0, TEXT("Error: Terminal nodes can not have outputs.")))
-	{
-		return Super::ValidateConnections();
-	}
-	return false;
-}
diff --git a/Engine/Source/Runtime/Experimental/Dataflow/Core/Public/Dataflow/DataflowTerminalNode.h b/Engine/Source/Runtime/Experimental/Dataflow/Core/Public/Dataflow/DataflowTerminalNode.h
--- a/Engine/Source/Runtime/Experimental/Dataflow/Core/Public/Dataflow/DataflowTerminalNode.h
+++ b/Engine/Source/Runtime/Experimental/Dataflow/Core/Public/Dataflow/DataflowTerminalNode.h
@@ -37,16 +37,10 @@ struct DATAFLOWCORE_API FDataflowTerminalNode : public FDataflowNode
 
 	virtual ~FDataflowTerminalNode() { }
 
-	//
-	// Connections
-	//
-
-	virtual void AddOutput(FDataflowOutput* InPtr) override { ensure(false); }
-
 	//
 	// Error Checking
 	//
-	virtual bool ValidateConnections() override ;
+	virtual void SetAssetValue(TObjectPtr<UObject> Asset, Dataflow::FContext& Context) const { ensure(false); }
 
 	//
 	// Evaluate
diff --git a/Engine/Source/Runtime/Experimental/GeometryCollectionEngine/Public/GeometryCollection/GeometryCollectionObject.h b/Engine/Source/Runtime/Experimental/GeometryCollectionEngine/Public/GeometryCollection/GeometryCollectionObject.h
--- a/Engine/Source/Runtime/Experimental/GeometryCollectionEngine/Public/GeometryCollection/GeometryCollectionObject.h
+++ b/Engine/Source/Runtime/Experimental/GeometryCollectionEngine/Public/GeometryCollection/GeometryCollectionObject.h
@@ -702,10 +702,10 @@ public:
 	// Dataflow
 	//
 	UPROPERTY(EditAnywhere, BlueprintReadWrite, Category = "Dataflow")
-	TObjectPtr<UDataflow> Dataflow;
+	TObjectPtr<UDataflow> DataflowAsset;
 
 	UPROPERTY(EditAnywhere, Category = "Dataflow")
-	FString Terminal;
+	FString DataflowTerminal = "GeometryCollectionTerminal";
 
 
 private:
