Dataflow : Node Rendering

diff --git a/Engine/Plugins/Experimental/ChaosFlesh/Source/ChaosFleshNodes/Public/Dataflow/ChaosFleshEngineAssetNodes.h b/Engine/Plugins/Experimental/ChaosFlesh/Source/ChaosFleshNodes/Public/Dataflow/ChaosFleshEngineAssetNodes.h
--- a/Engine/Plugins/Experimental/ChaosFlesh/Source/ChaosFleshNodes/Public/Dataflow/ChaosFleshEngineAssetNodes.h
+++ b/Engine/Plugins/Experimental/ChaosFlesh/Source/ChaosFleshNodes/Public/Dataflow/ChaosFleshEngineAssetNodes.h
@@ -17,7 +17,7 @@ struct FGetFleshAssetDataflowNode : public FDataflowNode
 {
 	GENERATED_USTRUCT_BODY()
 	DATAFLOW_NODE_DEFINE_INTERNAL(FGetFleshAssetDataflowNode, "GetFleshAsset", "Flesh", "")
-	DATAFLOW_NODE_RENDER_TYPE(FGeometryCollection::StaticType(), "Output")
+	DATAFLOW_NODE_RENDER_TYPE(FGeometryCollection::StaticType(), "Collection")
 
 public:
 
@@ -39,6 +39,7 @@ struct FFleshAssetTerminalDataflowNode : public FDataflowTerminalNode
 {
 	GENERATED_USTRUCT_BODY()
 	DATAFLOW_NODE_DEFINE_INTERNAL(FFleshAssetTerminalDataflowNode, "FleshAssetTerminal", "Flesh", "")
+	DATAFLOW_NODE_RENDER_TYPE(FGeometryCollection::StaticType(), "Collection")
 
 public:
 
diff --git a/Engine/Plugins/Experimental/Dataflow/Source/DataflowEditor/Private/Dataflow/DataflowEditorToolkit.cpp b/Engine/Plugins/Experimental/Dataflow/Source/DataflowEditor/Private/Dataflow/DataflowEditorToolkit.cpp
--- a/Engine/Plugins/Experimental/Dataflow/Source/DataflowEditor/Private/Dataflow/DataflowEditorToolkit.cpp
+++ b/Engine/Plugins/Experimental/Dataflow/Source/DataflowEditor/Private/Dataflow/DataflowEditorToolkit.cpp
@@ -2,15 +2,19 @@
 
 #include "Dataflow/DataflowEditorToolkit.h"
 
+#include "Dataflow/DataflowCore.h"
 #include "Dataflow/DataflowEditorActions.h"
-#include "Dataflow/DataflowGraphEditor.h"
+#include "Dataflow/DataflowEditorViewport.h"
 #include "Dataflow/DataflowEdNode.h"
+#include "Dataflow/DataflowGraphEditor.h"
 #include "Dataflow/DataflowNodeFactory.h"
 #include "Dataflow/DataflowObject.h"
 #include "Dataflow/DataflowObjectInterface.h"
 #include "Dataflow/DataflowSchema.h"
-#include "Dataflow/DataflowCore.h"
 #include "EditorStyleSet.h"
+#include "EditorViewportTabContent.h"
+#include "EditorViewportLayout.h"
+#include "EditorViewportCommands.h"
 #include "GraphEditorActions.h"
 #include "Modules/ModuleManager.h"
 #include "PropertyEditorModule.h"
@@ -24,6 +28,7 @@
 
 //DEFINE_LOG_CATEGORY_STATIC(FDataflowEditorToolkitLog, Log, All);
 
+const FName FDataflowEditorToolkit::ViewportTabId(TEXT("DataflowEditor_Viewport"));
 const FName FDataflowEditorToolkit::GraphCanvasTabId(TEXT("DataflowEditor_GraphCanvas"));
 const FName FDataflowEditorToolkit::AssetDetailsTabId(TEXT("DataflowEditor_AssetDetails"));
 const FName FDataflowEditorToolkit::NodeDetailsTabId(TEXT("DataflowEditor_NodeDetails"));
@@ -139,6 +144,16 @@ void FDataflowEditorToolkit::OnPropertyValueChanged(const FPropertyChangedEvent&
 	FDataflowEditorCommands::OnPropertyValueChanged(this->GetDataflow(), Context, LastNodeTimestamp, PropertyChangedEvent);
 }
 
+bool FDataflowEditorToolkit::OnNodeVerifyTitleCommit(const FText& NewText, UEdGraphNode* GraphNode, FText& OutErrorMessage)
+{
+	return FDataflowEditorCommands::OnNodeVerifyTitleCommit(NewText, GraphNode, OutErrorMessage);
+}
+
+void FDataflowEditorToolkit::OnNodeTitleCommitted(const FText& InNewText, ETextCommit::Type InCommitType, UEdGraphNode* GraphNode)
+{
+	FDataflowEditorCommands::OnNodeTitleCommitted(InNewText, InCommitType, GraphNode);
+}
+
 void FDataflowEditorToolkit::Tick(float DeltaTime)
 {
 	if (Dataflow && Asset)
@@ -155,7 +170,7 @@ void FDataflowEditorToolkit::Tick(float DeltaTime)
 
 TStatId FDataflowEditorToolkit::GetStatId() const
 {
-	RETURN_QUICK_DECLARE_CYCLE_STAT(FFleshEditorToolkit, STATGROUP_Tickables);
+	RETURN_QUICK_DECLARE_CYCLE_STAT(FDataflowEditorToolkit, STATGROUP_Tickables);
 }
 
 TSharedRef<SGraphEditor> FDataflowEditorToolkit::CreateGraphEditorWidget(UDataflow* DataflowToEdit, TSharedPtr<IStructureDetailsView> InNodeDetailsEditor)
@@ -241,14 +256,22 @@ TSharedPtr<IDetailsView> FDataflowEditorToolkit::CreateAssetDetailsEditorWidget(
 
 }
 
-bool FDataflowEditorToolkit::OnNodeVerifyTitleCommit(const FText& NewText, UEdGraphNode* GraphNode, FText& OutErrorMessage)
+TSharedRef<SDockTab> FDataflowEditorToolkit::SpawnTab_Viewport(const FSpawnTabArgs& Args)
 {
-	return FDataflowEditorCommands::OnNodeVerifyTitleCommit(NewText, GraphNode, OutErrorMessage);
-}
+	check(Args.GetTabId() == ViewportTabId);
 
-void FDataflowEditorToolkit::OnNodeTitleCommitted(const FText& InNewText, ETextCommit::Type InCommitType, UEdGraphNode* GraphNode)
-{
-	FDataflowEditorCommands::OnNodeTitleCommitted(InNewText, InCommitType, GraphNode);
+	TSharedRef< SDockTab > DockableTab = SNew(SDockTab);
+	ViewportEditor = MakeShareable(new FEditorViewportTabContent());
+	TWeakPtr<FDataflowEditorToolkit> WeakSharedThis = SharedThis(this);
+
+	const FString LayoutId = FString("DataflowEditorViewport");
+	ViewportEditor->Initialize([WeakSharedThis](const FAssetEditorViewportConstructionArgs& InConstructionArgs)
+		{
+			return SNew(SDataflowEditorViewport)
+				.DataflowEditorToolkit(WeakSharedThis);
+		}, DockableTab, LayoutId);
+
+	return DockableTab;
 }
 
 TSharedRef<SDockTab> FDataflowEditorToolkit::SpawnTab_GraphCanvas(const FSpawnTabArgs& Args)
@@ -290,6 +313,11 @@ void FDataflowEditorToolkit::RegisterTabSpawners(const TSharedRef<FTabManager>&
 {
 	TSharedRef<FWorkspaceItem> WorkspaceMenuCategoryRef = InTabManager->AddLocalWorkspaceMenuCategory(LOCTEXT("WorkspaceMenu_DataflowEditor", "Dataflow Editor"));
 
+	InTabManager->RegisterTabSpawner(ViewportTabId, FOnSpawnTab::CreateSP(this, &FDataflowEditorToolkit::SpawnTab_Viewport))
+		.SetDisplayName(LOCTEXT("DataflowViewportTab", "Dataflow Viewport"))
+		.SetGroup(WorkspaceMenuCategoryRef)
+		.SetIcon(FSlateIcon(FAppStyle::GetAppStyleSetName(), "GraphEditor.EventGraph_16x"));
+
 	InTabManager->RegisterTabSpawner(GraphCanvasTabId, FOnSpawnTab::CreateSP(this, &FDataflowEditorToolkit::SpawnTab_GraphCanvas))
 		.SetDisplayName(LOCTEXT("DataflowTab", "Graph"))
 		.SetGroup(WorkspaceMenuCategoryRef)
@@ -357,5 +385,9 @@ void FDataflowEditorToolkit::AddReferencedObjects(FReferenceCollector& Collector
 	{
 		Collector.AddReferencedObject(Dataflow);
 	}
+	if (Asset)
+	{
+		Collector.AddReferencedObject(Asset);
+	}
 }
 #undef LOCTEXT_NAMESPACE
diff --git a/Engine/Plugins/Experimental/Dataflow/Source/DataflowEditor/Private/Dataflow/DataflowEditorViewport.cpp b/Engine/Plugins/Experimental/Dataflow/Source/DataflowEditor/Private/Dataflow/DataflowEditorViewport.cpp
new file mode 100644
--- /dev/null
+++ b/Engine/Plugins/Experimental/Dataflow/Source/DataflowEditor/Private/Dataflow/DataflowEditorViewport.cpp
@@ -0,0 +1,151 @@
+// Copyright Epic Games, Inc. All Rights Reserved.
+#include "Dataflow/DataflowEditorViewport.h"
+
+#include "AdvancedPreviewScene.h"
+#include "Dataflow/DataflowEdNode.h"
+#include "Dataflow/DataflowObject.h"
+#include "Dataflow/DataflowRenderingActor.h"
+#include "Dataflow/DataflowRenderingComponent.h"
+#include "EditorViewportClient.h"
+#include "Dataflow/DataflowEditorToolkit.h"
+#include "SCommonEditorViewportToolbarBase.h"
+
+SDataflowEditorViewport::SDataflowEditorViewport()
+{
+	// Temporarily allow water subsystem to be created on preview worlds because we need one here : 
+	//UWaterSubsystem::FScopedAllowWaterSubsystemOnPreviewWorld AllowWaterSubsystemOnPreviewWorldScope(true);
+	PreviewScene = MakeShareable(new FAdvancedPreviewScene(FPreviewScene::ConstructionValues()));
+	PreviewScene->SetFloorVisibility(false);
+}
+
+void SDataflowEditorViewport::Construct(const FArguments& InArgs)
+{
+	DataflowEditorToolkitPtr = InArgs._DataflowEditorToolkit;
+	TSharedPtr<FDataflowEditorToolkit> DataflowEditorToolkit = DataflowEditorToolkitPtr.Pin();
+	check(DataflowEditorToolkitPtr.IsValid());
+
+	SEditorViewport::Construct(SEditorViewport::FArguments());
+
+	FBoxSphereBounds SphereBounds = FBoxSphereBounds(EForceInit::ForceInitToZero);
+	CustomDataflowRenderingActor = CastChecked<ADataflowRenderingActor>(PreviewScene->GetWorld()->SpawnActor(ADataflowRenderingActor::StaticClass()));
+
+	EditorViewportClient->SetDataflowRenderingActor(CustomDataflowRenderingActor);
+	EditorViewportClient->FocusViewportOnBox( SphereBounds.GetBox());
+}
+
+TSharedRef<SEditorViewport> SDataflowEditorViewport::GetViewportWidget()
+{
+	return SharedThis(this);
+}
+
+TSharedPtr<FExtender> SDataflowEditorViewport::GetExtenders() const
+{
+	TSharedPtr<FExtender> Result(MakeShareable(new FExtender));
+	return Result;
+}
+
+void SDataflowEditorViewport::OnFloatingButtonClicked()
+{
+}
+
+void SDataflowEditorViewport::AddReferencedObjects(FReferenceCollector& Collector)
+{							
+	Collector.AddReferencedObject(CustomDataflowRenderingActor);
+}
+
+TSharedRef<FEditorViewportClient> SDataflowEditorViewport::MakeEditorViewportClient()
+{
+	EditorViewportClient = MakeShareable(new FDataflowEditorViewportClient(PreviewScene.Get(), SharedThis(this), DataflowEditorToolkitPtr));
+	return EditorViewportClient.ToSharedRef();
+}
+
+TSharedPtr<SWidget> SDataflowEditorViewport::MakeViewportToolbar()
+{
+	return SNew(SCommonEditorViewportToolbarBase, SharedThis(this));
+}
+
+
+// ----------------------------------------------------------------------------------
+
+FDataflowEditorViewportClient::FDataflowEditorViewportClient(FPreviewScene* InPreviewScene, 
+	const TWeakPtr<SEditorViewport> InEditorViewportWidget,
+	TWeakPtr<FDataflowEditorToolkit> InDataflowEditorToolkitPtr)
+	: 
+	FEditorViewportClient(nullptr, InPreviewScene, InEditorViewportWidget)
+	, DataflowEditorToolkitPtr(InDataflowEditorToolkitPtr)
+{
+	bSetListenerPosition = false;
+	SetRealtime(true);
+	EngineShowFlags.Grid = false;
+}
+
+Dataflow::FTimestamp FDataflowEditorViewportClient::LatestTimestamp(const UDataflow* Dataflow, const Dataflow::FContext* Context)
+{
+	if (Dataflow && Context)
+	{
+		return FMath::Max(Dataflow->GetRenderingTimestamp().Value, Context->GetTimestamp().Value);
+	}
+	return Dataflow::FTimestamp::Invalid;
+}
+
+
+void FDataflowEditorViewportClient::Tick(float DeltaSeconds)
+{
+	Super::Tick(DeltaSeconds);
+
+	TSharedPtr<FDataflowEditorToolkit> DataflowEditorToolkit = DataflowEditorToolkitPtr.Pin();
+
+	if (DataflowRenderingActor && DataflowEditorToolkitPtr.IsValid())
+	{
+		if (TSharedPtr<Dataflow::FContext> Context = DataflowEditorToolkit->GetContext())
+		{
+			if (UDataflowRenderingComponent* DataflowRenderingComponent = DataflowRenderingActor->GetDataflowRenderingComponent())
+			{
+				if (const UDataflow* Dataflow = DataflowEditorToolkit->GetDataflow())
+				{
+					Dataflow::FTimestamp SystemTimestamp = LatestTimestamp(Dataflow, Context.Get());
+					if (SystemTimestamp >= LastModifiedTimestamp)
+					{
+						if (Dataflow->GetRenderTargets().Num())
+						{
+							// @todo(dataflow) : Check the Time on the target outs first instead.
+							//					 to avoid invalidation during unrelated edits. 
+
+							DataflowRenderingComponent->ResetRenderTargets();
+
+							DataflowRenderingComponent->SetDataflow(Dataflow);
+							DataflowRenderingComponent->SetContext(Context);
+							for (const UDataflowEdNode* Node : Dataflow->GetRenderTargets())
+							{
+								DataflowRenderingComponent->AddRenderTarget(Node);
+							}
+						}
+						else
+						{
+							DataflowRenderingComponent->ResetRenderTargets();
+						}
+						LastModifiedTimestamp = LatestTimestamp(Dataflow, Context.Get()).Value + 1;
+					}
+				}
+			}
+		}
+	}
+	// Tick the preview scene world.
+	if (!GIntraFrameDebuggingGameThread)
+	{
+		PreviewScene->GetWorld()->Tick(IsRealtime() ? LEVELTICK_All : LEVELTICK_TimeOnly, DeltaSeconds);
+	}
+}
+
+
+void FDataflowEditorViewportClient::AddReferencedObjects(FReferenceCollector& Collector)
+{
+	Super::AddReferencedObjects(Collector);
+
+	if (DataflowRenderingActor)
+	{
+		Collector.AddReferencedObject(DataflowRenderingActor);
+	}
+}
+
+
diff --git a/Engine/Plugins/Experimental/Dataflow/Source/DataflowEditor/Public/Dataflow/DataflowEditorToolkit.h b/Engine/Plugins/Experimental/Dataflow/Source/DataflowEditor/Public/Dataflow/DataflowEditorToolkit.h
--- a/Engine/Plugins/Experimental/Dataflow/Source/DataflowEditor/Public/Dataflow/DataflowEditorToolkit.h
+++ b/Engine/Plugins/Experimental/Dataflow/Source/DataflowEditor/Public/Dataflow/DataflowEditorToolkit.h
@@ -9,6 +9,7 @@
 #include "GraphEditor.h"
 #include "TickableEditorObject.h"
 
+class FEditorViewportTabContent;
 class IDetailsView;
 class FTabManager;
 class IStructureDetailsView;
@@ -61,6 +62,7 @@ public:
 
 	// Tab spawners 
 	virtual void RegisterTabSpawners(const TSharedRef<FTabManager>& TabManager) override;
+	TSharedRef<SDockTab> SpawnTab_Viewport(const FSpawnTabArgs& Args);
 	TSharedRef<SDockTab> SpawnTab_GraphCanvas(const FSpawnTabArgs& Args);
 	TSharedRef<SDockTab> SpawnTab_AssetDetails(const FSpawnTabArgs& Args);
 	TSharedRef<SDockTab> SpawnTab_NodeDetails(const FSpawnTabArgs& Args);
@@ -90,6 +92,9 @@ private:
 	UDataflow* Dataflow = nullptr;
 	FString TerminalPath = "";
 
+	static const FName ViewportTabId;
+	TSharedPtr<FEditorViewportTabContent> ViewportEditor;
+
 	static const FName GraphCanvasTabId;
 	TSharedPtr<SGraphEditor> GraphEditor;
 	TSharedPtr<FUICommandList> GraphEditorCommands;
diff --git a/Engine/Plugins/Experimental/Dataflow/Source/DataflowEditor/Public/Dataflow/DataflowEditorViewport.h b/Engine/Plugins/Experimental/Dataflow/Source/DataflowEditor/Public/Dataflow/DataflowEditorViewport.h
new file mode 100644
--- /dev/null
+++ b/Engine/Plugins/Experimental/Dataflow/Source/DataflowEditor/Public/Dataflow/DataflowEditorViewport.h
@@ -0,0 +1,83 @@
+// Copyright Epic Games, Inc. All Rights Reserved.
+#pragma once
+
+#include "CoreMinimal.h"
+#include "Dataflow/DataflowNodeParameters.h"
+#include "EditorViewportClient.h"
+#include "SAssetEditorViewport.h"
+#include "SCommonEditorViewportToolbarBase.h"
+
+class FDataflowEditorToolkit;
+class FAdvancedPreviewScene;
+class FEditorViewportClient;
+class FDataflowEditorViewportClient;
+class SEditorViewport;
+class ADataflowActor;
+class ADataflowRenderingActor;
+
+// ----------------------------------------------------------------------------------
+
+class SDataflowEditorViewport : public SAssetEditorViewport, public ICommonEditorViewportToolbarInfoProvider, public FGCObject
+{
+public:
+	SLATE_BEGIN_ARGS(SDataflowEditorViewport) {}
+		SLATE_ARGUMENT(TWeakPtr<FDataflowEditorToolkit>, DataflowEditorToolkit)
+	SLATE_END_ARGS()
+
+	SDataflowEditorViewport();
+
+	void Construct(const FArguments& InArgs);
+
+	//~ ICommonEditorViewportToolbarInfoProvider interface
+	virtual TSharedRef<SEditorViewport> GetViewportWidget() override;
+	virtual TSharedPtr<FExtender> GetExtenders() const override;
+	virtual void OnFloatingButtonClicked() override;
+
+	//~ FGCObject Interface
+	virtual void AddReferencedObjects(FReferenceCollector& Collector) override;
+	virtual FString GetReferencerName() const override{return TEXT("SDataflowEditorViewport");}
+
+protected:
+	virtual TSharedRef<FEditorViewportClient> MakeEditorViewportClient() override;
+	virtual TSharedPtr<SWidget> MakeViewportToolbar() override;
+private:
+	/// The scene for this viewport. 
+	TSharedPtr<FAdvancedPreviewScene> PreviewScene;
+
+	/// Editor viewport client 
+	TSharedPtr<FDataflowEditorViewportClient> EditorViewportClient;
+
+	TWeakPtr<FDataflowEditorToolkit> DataflowEditorToolkitPtr;
+
+	ADataflowActor* CustomDataflowActor = nullptr;
+	ADataflowRenderingActor* CustomDataflowRenderingActor = nullptr;
+};
+
+
+// ----------------------------------------------------------------------------------
+
+class FDataflowEditorViewportClient : public FEditorViewportClient
+{
+public:
+	using Super = FEditorViewportClient;
+	
+	FDataflowEditorViewportClient(FPreviewScene* InPreviewScene,
+		const TWeakPtr<SEditorViewport> InEditorViewportWidget = nullptr,
+		TWeakPtr<FDataflowEditorToolkit> InDataflowEditorToolkitPtr = nullptr);
+
+
+	void SetDataflowRenderingActor(ADataflowRenderingActor* InActor) { DataflowRenderingActor = InActor; }
+	Dataflow::FTimestamp LatestTimestamp(const UDataflow* Dataflow, const Dataflow::FContext* Context);
+
+	// FEditorViewportClient interface
+	virtual void Tick(float DeltaSeconds) override;
+	// End of FEditorViewportClient
+
+	//~ FGCObject Interface
+	virtual void AddReferencedObjects(FReferenceCollector& Collector) override;
+	virtual FString GetReferencerName() const override { return TEXT("FDataflowEditorViewportClient"); }
+
+	TWeakPtr<FDataflowEditorToolkit> DataflowEditorToolkitPtr = nullptr;
+	ADataflowRenderingActor* DataflowRenderingActor = nullptr;
+	Dataflow::FTimestamp LastModifiedTimestamp = Dataflow::FTimestamp::Invalid;
+};
diff --git a/Engine/Plugins/Experimental/GeometryCollectionPlugin/Source/GeometryCollectionNodes/Public/Dataflow/Nodes/GeometryCollectionAssetNodes.h b/Engine/Plugins/Experimental/GeometryCollectionPlugin/Source/GeometryCollectionNodes/Public/Dataflow/Nodes/GeometryCollectionAssetNodes.h
--- a/Engine/Plugins/Experimental/GeometryCollectionPlugin/Source/GeometryCollectionNodes/Public/Dataflow/Nodes/GeometryCollectionAssetNodes.h
+++ b/Engine/Plugins/Experimental/GeometryCollectionPlugin/Source/GeometryCollectionNodes/Public/Dataflow/Nodes/GeometryCollectionAssetNodes.h
@@ -16,6 +16,7 @@ struct FGeometryCollectionTerminalDataflowNode : public FDataflowTerminalNode
 {
 	GENERATED_USTRUCT_BODY()
 	DATAFLOW_NODE_DEFINE_INTERNAL(FGeometryCollectionTerminalDataflowNode, "GeometryCollectionTerminal", "GeometryCollection", "")
+	DATAFLOW_NODE_RENDER_TYPE(FGeometryCollection::StaticType(), "Collection")
 
 public:
 
