Dataflow : Skeletal View

diff --git a/Engine/Plugins/Experimental/Dataflow/Source/DataflowEditor/Private/Dataflow/DataflowEditorToolkit.cpp b/Engine/Plugins/Experimental/Dataflow/Source/DataflowEditor/Private/Dataflow/DataflowEditorToolkit.cpp
--- a/Engine/Plugins/Experimental/Dataflow/Source/DataflowEditor/Private/Dataflow/DataflowEditorToolkit.cpp
+++ b/Engine/Plugins/Experimental/Dataflow/Source/DataflowEditor/Private/Dataflow/DataflowEditorToolkit.cpp
@@ -2,6 +2,7 @@
 
 #include "Dataflow/DataflowEditorToolkit.h"
 
+#include "Animation/Skeleton.h"
 #include "Dataflow/DataflowCore.h"
 #include "Dataflow/DataflowEditorActions.h"
 #include "Dataflow/DataflowEditorViewport.h"
@@ -15,12 +16,15 @@
 #include "EditorViewportTabContent.h"
 #include "EditorViewportLayout.h"
 #include "EditorViewportCommands.h"
+#include "Engine/SkeletalMesh.h"
 #include "GraphEditorActions.h"
 #include "Modules/ModuleManager.h"
 #include "PropertyEditorModule.h"
 #include "Styling/SlateStyleRegistry.h"
 #include "Framework/Commands/GenericCommands.h"
 #include "Widgets/Docking/SDockTab.h"
+#include "ISkeletonTree.h"
+#include "ISkeletonEditorModule.h"
 #include "IStructureDetailsView.h"
 #include "Misc/MessageDialog.h"
 
@@ -32,6 +36,7 @@ const FName FDataflowEditorToolkit::ViewportTabId(TEXT("DataflowEditor_Viewport"
 const FName FDataflowEditorToolkit::GraphCanvasTabId(TEXT("DataflowEditor_GraphCanvas"));
 const FName FDataflowEditorToolkit::AssetDetailsTabId(TEXT("DataflowEditor_AssetDetails"));
 const FName FDataflowEditorToolkit::NodeDetailsTabId(TEXT("DataflowEditor_NodeDetails"));
+const FName FDataflowEditorToolkit::SkeletalTabId(TEXT("DataflowEditor_Skeletal"));
 
 
 UDataflow* GetDataflowFrom(UObject* InObject)
@@ -47,6 +52,19 @@ UDataflow* GetDataflowFrom(UObject* InObject)
 
 }
 
+USkeletalMesh* GetSkeletalMeshFrom(UObject* InObject)
+{
+	if (UClass* Class = InObject->GetClass())
+	{
+		if (FProperty* Property = Class->FindPropertyByName(FName("SkeletalMesh")))
+		{
+			return *Property->ContainerPtrToValuePtr<USkeletalMesh*>(InObject);
+		}
+	}
+	return nullptr;
+
+}
+
 FString GetDataflowTerminalFrom(UObject* InObject)
 {
 	if (UClass* Class = InObject->GetClass())
@@ -90,6 +108,7 @@ void FDataflowEditorToolkit::InitDataflowEditor(const EToolkitMode::Type Mode, c
 		NodeDetailsEditor = CreateNodeDetailsEditorWidget(ObjectToEdit);
 		AssetDetailsEditor = CreateAssetDetailsEditorWidget(ObjectToEdit);
 		GraphEditor = CreateGraphEditorWidget(Dataflow, NodeDetailsEditor);
+		SkeletalEditor = CreateSkeletalEditorWidget(ObjectToEdit);
 
 		Context = TSharedPtr< Dataflow::FEngineContext>(new Dataflow::FAssetContext(Asset, Dataflow, FPlatformTime::Cycles64()));
 		LastNodeTimestamp = Context->GetTimestamp();
@@ -256,6 +275,39 @@ TSharedPtr<IDetailsView> FDataflowEditorToolkit::CreateAssetDetailsEditorWidget(
 
 }
 
+TSharedPtr<ISkeletonTree> FDataflowEditorToolkit::CreateSkeletalEditorWidget(UObject* ObjectToEdit)
+{
+	if (Dataflow)
+	{
+		if (!StubSkeletalMesh)
+		{
+			const FName NodeName = MakeUniqueObjectName(Dataflow, UDataflow::StaticClass(), FName("USkeleton"));
+			StubSkeleton = NewObject<USkeleton>(Dataflow, NodeName);
+			const FName NodeName2 = MakeUniqueObjectName(Dataflow, UDataflow::StaticClass(), FName("USkeleton"));
+			StubSkeletalMesh = NewObject<USkeletalMesh>(Dataflow, NodeName2);
+			StubSkeletalMesh->SetSkeleton(StubSkeleton);
+		}
+
+		FSkeletonTreeArgs SkeletonTreeArgs;
+		//SkeletonTreeArgs.OnSelectionChanged = FOnSkeletonTreeSelectionChanged::CreateSP(this, &FAnimationBlueprintEditor::HandleSelectionChanged);
+		//SkeletonTreeArgs.PreviewScene = GetPreviewScene();
+		//SkeletonTreeArgs.ContextName = GetToolkitFName();
+
+		USkeleton* Skeleton = StubSkeleton;
+		if (Asset)
+		{
+			if (USkeletalMesh* SkeletalMesh = GetSkeletalMeshFrom(Asset))
+			{
+				Skeleton = SkeletalMesh->GetSkeleton();
+			}
+		}
+		ISkeletonEditorModule& SkeletonEditorModule = FModuleManager::LoadModuleChecked<ISkeletonEditorModule>("SkeletonEditor");
+		TSharedPtr<ISkeletonTree> SkeletonTree = SkeletonEditorModule.CreateSkeletonTree(Skeleton, SkeletonTreeArgs);
+		return SkeletonTree;
+	}
+	return TSharedPtr<ISkeletonTree>(nullptr);
+}
+
 TSharedRef<SDockTab> FDataflowEditorToolkit::SpawnTab_Viewport(const FSpawnTabArgs& Args)
 {
 	check(Args.GetTabId() == ViewportTabId);
@@ -309,6 +361,25 @@ TSharedRef<SDockTab> FDataflowEditorToolkit::SpawnTab_NodeDetails(const FSpawnTa
 		];
 }
 
+TSharedRef<SDockTab> FDataflowEditorToolkit::SpawnTab_Skeletal(const FSpawnTabArgs& Args)
+{
+	check(Args.GetTabId() == SkeletalTabId);
+
+	USkeletalMesh* SkeletalMesh = StubSkeletalMesh;
+	if (Asset)
+	{
+		SkeletalMesh = GetSkeletalMeshFrom(Asset);
+	}
+
+	SkeletalEditor->SetSkeletalMesh(SkeletalMesh);
+
+	return SNew(SDockTab)
+		.Label(LOCTEXT("FleshEditorSkeletal_TabTitle", "Skeletal Hierarchy"))
+		[
+			SkeletalEditor.ToSharedRef()
+		];
+}
+
 void FDataflowEditorToolkit::RegisterTabSpawners(const TSharedRef<FTabManager>& InTabManager)
 {
 	TSharedRef<FWorkspaceItem> WorkspaceMenuCategoryRef = InTabManager->AddLocalWorkspaceMenuCategory(LOCTEXT("WorkspaceMenu_DataflowEditor", "Dataflow Editor"));
@@ -333,6 +404,11 @@ void FDataflowEditorToolkit::RegisterTabSpawners(const TSharedRef<FTabManager>&
 		.SetGroup(WorkspaceMenuCategoryRef)
 		.SetIcon(FSlateIcon(FAppStyle::GetAppStyleSetName(), "LevelEditor.Tabs.Details"));
 
+	InTabManager->RegisterTabSpawner(SkeletalTabId, FOnSpawnTab::CreateSP(this, &FDataflowEditorToolkit::SpawnTab_Skeletal))
+		.SetDisplayName(LOCTEXT("DataflowSkeletalTab", "Skeletal Hierarchy"))
+		.SetGroup(WorkspaceMenuCategoryRef)
+		.SetIcon(FSlateIcon(FAppStyle::GetAppStyleSetName(), "LevelEditor.Tabs.SkeletalHierarchy"));
+
 	FAssetEditorToolkit::RegisterTabSpawners(InTabManager);
 }
 
diff --git a/Engine/Plugins/Experimental/Dataflow/Source/DataflowEditor/Public/Dataflow/DataflowEditorToolkit.h b/Engine/Plugins/Experimental/Dataflow/Source/DataflowEditor/Public/Dataflow/DataflowEditorToolkit.h
--- a/Engine/Plugins/Experimental/Dataflow/Source/DataflowEditor/Public/Dataflow/DataflowEditorToolkit.h
+++ b/Engine/Plugins/Experimental/Dataflow/Source/DataflowEditor/Public/Dataflow/DataflowEditorToolkit.h
@@ -15,6 +15,7 @@ class FTabManager;
 class IStructureDetailsView;
 class IToolkitHost;
 class UDataflow;
+class USkeletalMesh;
 
 namespace Dataflow
 {
@@ -36,10 +37,6 @@ public:
 
 	void InitDataflowEditor(const EToolkitMode::Type Mode, const TSharedPtr<IToolkitHost>& InitToolkitHost, UObject* ObjectToEdit);
 
-	//~ Begin DataflowEditorActions
-	void OnPropertyValueChanged(const FPropertyChangedEvent& PropertyChangedEvent);
-	//~ End DataflowEditorActions
-
 	// IToolkit Interface
 	virtual FName GetToolkitFName() const override;
 	virtual FText GetToolkitName() const override;
@@ -48,24 +45,23 @@ public:
 	virtual FString GetWorldCentricTabPrefix() const override;
 	virtual FLinearColor GetWorldCentricTabColorScale() const override;
 
-	//~ Begin FTickableEditorObject interface
+	// FTickableEditorObject interface
 	virtual void Tick(float DeltaTime) override;
 	virtual bool IsTickable() const override { return true; }
 	virtual TStatId GetStatId() const override;
-	//~ End FTickableEditorObject Interface
 
 	// FGCObject Interface
 	virtual void AddReferencedObjects(FReferenceCollector& Collector) override;
 	virtual FString GetReferencerName() const override; 
 
-
-
 	// Tab spawners 
 	virtual void RegisterTabSpawners(const TSharedRef<FTabManager>& TabManager) override;
 	TSharedRef<SDockTab> SpawnTab_Viewport(const FSpawnTabArgs& Args);
 	TSharedRef<SDockTab> SpawnTab_GraphCanvas(const FSpawnTabArgs& Args);
 	TSharedRef<SDockTab> SpawnTab_AssetDetails(const FSpawnTabArgs& Args);
 	TSharedRef<SDockTab> SpawnTab_NodeDetails(const FSpawnTabArgs& Args);
+	TSharedRef<SDockTab> SpawnTab_Skeletal(const FSpawnTabArgs& Args);
+
 
 	// Member Access
 	UObject* GetAsset() { return Asset; }
@@ -86,6 +82,14 @@ public:
 	TSharedPtr<SGraphEditor> GetGraphEditor() { return GraphEditor; }
 	const TSharedPtr<SGraphEditor> GetGraphEditor() const { return GraphEditor; }
 
+protected:
+
+	//~ Begin DataflowEditorActions
+	void OnPropertyValueChanged(const FPropertyChangedEvent& PropertyChangedEvent);
+	bool OnNodeVerifyTitleCommit(const FText& NewText, UEdGraphNode* GraphNode, FText& OutErrorMessage);
+	void OnNodeTitleCommitted(const FText& InNewText, ETextCommit::Type InCommitType, UEdGraphNode* GraphNode);
+	//~ End DataflowEditorActions
+
 private:
 
 	UObject* Asset = nullptr;
@@ -108,8 +112,11 @@ private:
 	TSharedPtr<IStructureDetailsView> NodeDetailsEditor;
 	TSharedPtr<IStructureDetailsView> CreateNodeDetailsEditorWidget(UObject* ObjectToEdit);
 
-	bool OnNodeVerifyTitleCommit(const FText& NewText, UEdGraphNode* GraphNode, FText& OutErrorMessage);
-	void OnNodeTitleCommitted(const FText& InNewText, ETextCommit::Type InCommitType, UEdGraphNode* GraphNode);
+	static const FName SkeletalTabId;
+	TObjectPtr<USkeleton> StubSkeleton;
+	TObjectPtr<USkeletalMesh> StubSkeletalMesh;
+	TSharedPtr<class ISkeletonTree> SkeletalEditor;
+	TSharedPtr<ISkeletonTree> CreateSkeletalEditorWidget(UObject* ObjectToEdit);
 
 	TSharedPtr<Dataflow::FEngineContext> Context;
 	Dataflow::FTimestamp LastNodeTimestamp = Dataflow::FTimestamp::Invalid;
