Dataflow : Render surface

diff --git a/Engine/Plugins/Experimental/Dataflow/Source/DataflowEnginePlugin/Private/Dataflow/DataflowEnginePlugin.cpp b/Engine/Plugins/Experimental/Dataflow/Source/DataflowEnginePlugin/Private/Dataflow/DataflowEnginePlugin.cpp
--- a/Engine/Plugins/Experimental/Dataflow/Source/DataflowEnginePlugin/Private/Dataflow/DataflowEnginePlugin.cpp
+++ b/Engine/Plugins/Experimental/Dataflow/Source/DataflowEnginePlugin/Private/Dataflow/DataflowEnginePlugin.cpp
@@ -57,15 +57,34 @@ void FDataflowEnginePlugin::StartupModule()
 					GeometryCollectionAlgo::GlobalMatrices(Transforms, Parents, M);
 
 					auto ToD = [](FVector3f V) { return FVector3d(V.X, V.Y, V.Z); };
+					auto ToF = [](FVector3d V) { return FVector3f(V.X, V.Y, V.Z); };
+
 
 					const TManagedArray<FVector3f>& Vertex = Collection.GetAttribute<FVector3f>("Vertex", FGeometryCollection::VerticesGroup);
 					const TManagedArray<FIntVector>& Faces = Collection.GetAttribute<FIntVector>("Indices", FGeometryCollection::FacesGroup);
+
+					TArray<FVector3f> Vertices; Vertices.AddUninitialized(Vertex.Num());
+					TArray<FIntVector> Tris; Tris.AddUninitialized(Faces.Num());
+					TArray<bool> Visited; Visited.Init(false, Vertices.Num());
+
+					int32 Tdx = 0;
 					for (const FIntVector& Face : Faces)
 					{
-						FTransform Ms[3] = { M[BoneIndex[Face[0]]], M[BoneIndex[Face[1]]], M[BoneIndex[Face[2]]] };
-						RenderCollection.AddTriangle(Chaos::FTriangle(Ms[0].TransformPosition(ToD(Vertex[Face[0]])),
-							Ms[1].TransformPosition(ToD(Vertex[Face[1]])), Ms[2].TransformPosition(ToD(Vertex[Face[2]]))));
+						FIntVector Tri = FIntVector(Face[0], Face[1], Face[2]);
+						FTransform Ms[3] = { M[BoneIndex[Tri[0]]], M[BoneIndex[Tri[1]]], M[BoneIndex[Tri[2]]] };
+
+						Tris[Tdx++] = Tri;
+						if (!Visited[Tri[0]]) Vertices[Tri[0]] = ToF(Ms[0].TransformPosition(ToD(Vertex[Tri[0]])));
+						if (!Visited[Tri[1]]) Vertices[Tri[1]] = ToF(Ms[1].TransformPosition(ToD(Vertex[Tri[1]])));
+						if (!Visited[Tri[2]]) Vertices[Tri[2]] = ToF(Ms[2].TransformPosition(ToD(Vertex[Tri[2]])));
+
+						Visited[Tri[0]] = true; Visited[Tri[1]] = true; Visited[Tri[2]] = true;
 					}
+
+					// Maybe these buffers should be shrunk, but there are unused vertices in the buffer. 
+					for (int i = 0; i < Visited.Num(); i++) if (!Visited[i]) Vertices[i] = FVector3f(0);
+
+					RenderCollection.AddSurface(MoveTemp(Vertices),MoveTemp(Tris));
 				}
 			}
 		});
diff --git a/Engine/Source/Runtime/Experimental/Chaos/Private/GeometryCollection/Facades/CollectionRenderingFacade.cpp b/Engine/Source/Runtime/Experimental/Chaos/Private/GeometryCollection/Facades/CollectionRenderingFacade.cpp
--- a/Engine/Source/Runtime/Experimental/Chaos/Private/GeometryCollection/Facades/CollectionRenderingFacade.cpp
+++ b/Engine/Source/Runtime/Experimental/Chaos/Private/GeometryCollection/Facades/CollectionRenderingFacade.cpp
@@ -10,9 +10,6 @@
 namespace GeometryCollection::Facades
 {
 
-	// Attributes
-	const FName FRenderingFacade::AAAAttribute = "AAA";
-
 	FRenderingFacade::FRenderingFacade(FManagedArrayCollection* InCollection)
 		: Self(InCollection)
 	{
@@ -76,6 +73,27 @@ namespace GeometryCollection::Facades
 		}
 	}
 
+
+	void FRenderingFacade::AddSurface(FManagedArrayCollection* InCollection, TArray<FVector3f>&& InVertices, TArray<FIntVector>&& InIndices)
+	{
+		if (IsValid(InCollection))
+		{
+			auto CollectionVert = [](const Chaos::FVec3& V) { return FVector3f(float(V.X), float(V.Y), float(V.Z)); };
+
+			TManagedArray<FVector3f>& Vertices = InCollection->ModifyAttribute<FVector3f>("Vertex", FGeometryCollection::VerticesGroup);
+			TManagedArray<FIntVector>& Indices = InCollection->ModifyAttribute<FIntVector>("Indices", FGeometryCollection::FacesGroup);
+
+			int32 IndicesStart = InCollection->AddElements(InIndices.Num(), FGeometryCollection::FacesGroup);
+			int32 VertexStart = InCollection->AddElements(InVertices.Num(), FGeometryCollection::VerticesGroup);
+
+			const FIntVector * IndiciesDest = Indices.GetData() + IndicesStart;
+			FMemory::Memmove((void*)IndiciesDest, InIndices.GetData(), sizeof(FIntVector) * InIndices.Num());
+
+			const FVector3f * VerticesDest = Vertices.GetData() + VertexStart;
+			FMemory::Memmove((void*)VerticesDest, InVertices.GetData(), sizeof(FVector3f) * InVertices.Num());
+		}
+	}
+
 	//
 	//  GetAttributes
 	//
diff --git a/Engine/Source/Runtime/Experimental/Chaos/Public/GeometryCollection/Facades/CollectionRenderingFacade.h b/Engine/Source/Runtime/Experimental/Chaos/Public/GeometryCollection/Facades/CollectionRenderingFacade.h
--- a/Engine/Source/Runtime/Experimental/Chaos/Public/GeometryCollection/Facades/CollectionRenderingFacade.h
+++ b/Engine/Source/Runtime/Experimental/Chaos/Public/GeometryCollection/Facades/CollectionRenderingFacade.h
@@ -36,9 +36,6 @@ namespace GeometryCollection::Facades
 
 	public:
 
-		// Attributes
-		static const FName AAAAttribute;
-
 		/**
 		* FRenderingFacade Constuctor
 		* @param VertixDependencyGroup : GroupName the index attribute is dependent on.
@@ -67,13 +64,22 @@ namespace GeometryCollection::Facades
 
 
 		/**
-		* Using a FSelectionFacade::FSelectionKey, create indexes from the vertices to the driving bones.
+		* Add a triangle to the rendering view. 
 		* @param FManagedArrayCollection : Collection
 		* @param FSelectionFacade::FSelectionKey : Key for weights in the FSelectionFacade
 		*/
 		static void AddTriangle(FManagedArrayCollection* InCollection, const Chaos::FTriangle& InTriangle);
 		void AddTriangle(const Chaos::FTriangle& InTriangle) { AddTriangle(Self, InTriangle); }
 
+		/**
+		* Add a surface to the rendering view.
+		* @param FManagedArrayCollection : Collection
+		* @param FSelectionFacade::FSelectionKey : Key for weights in the FSelectionFacade
+		*/
+		static void AddSurface(FManagedArrayCollection* InCollection, TArray<FVector3f>&& InVertices, TArray<FIntVector>&& InIndices);
+		void AddSurface(TArray<FVector3f>&& InVertices, TArray<FIntVector>&& InIndices) { AddSurface(Self, MoveTemp(InVertices), MoveTemp(InIndices) ); }
+
+
 		/**
 		* Return the render indices from the collection. Null if not initialized.
 		* @param FManagedArrayCollection : Collection
