Fix node evaluation not working in dataflow editor when right clicking

diff --git a/Engine/Plugins/Experimental/Dataflow/Source/DataflowEditor/Private/Dataflow/DataflowEditorToolkit.cpp b/Engine/Plugins/Experimental/Dataflow/Source/DataflowEditor/Private/Dataflow/DataflowEditorToolkit.cpp
--- a/Engine/Plugins/Experimental/Dataflow/Source/DataflowEditor/Private/Dataflow/DataflowEditorToolkit.cpp
+++ b/Engine/Plugins/Experimental/Dataflow/Source/DataflowEditor/Private/Dataflow/DataflowEditorToolkit.cpp
@@ -199,16 +199,13 @@ TSharedRef<SGraphEditor> FDataflowEditorToolkit::CreateGraphEditorWidget(UDatafl
 
 	FDataflowEditorCommands::FGraphEvaluationCallback Evaluate = [&](FDataflowNode* Node, FDataflowOutput* Out)
 	{
-		if (DataflowToEdit)
+		if (!Context)
 		{
-			if (!Context)
-			{
-				Context = TSharedPtr< Dataflow::FEngineContext>(new Dataflow::FAssetContext(Asset, Dataflow, Dataflow::FTimestamp::Invalid));
-			}
-			LastNodeTimestamp = Dataflow::FTimestamp::Invalid;
-
-			FDataflowEditorCommands::EvaluateTerminalNode(*Context.Get(), LastNodeTimestamp, Dataflow, Node, Out, Asset, TerminalPath);
+			Context = TSharedPtr< Dataflow::FEngineContext>(new Dataflow::FAssetContext(Asset, Dataflow, Dataflow::FTimestamp::Invalid));
 		}
+		LastNodeTimestamp = Dataflow::FTimestamp::Invalid;
+
+		FDataflowEditorCommands::EvaluateTerminalNode(*Context.Get(), LastNodeTimestamp, Dataflow, Node, Out, Asset, TerminalPath);
 	};
 
 	SGraphEditor::FGraphEditorEvents InEvents;
