Dataflow : Editor rendering with component base.

diff --git a/Engine/Plugins/Experimental/Dataflow/Source/DataflowEditor/DataflowEditor.Build.cs b/Engine/Plugins/Experimental/Dataflow/Source/DataflowEditor/DataflowEditor.Build.cs
--- a/Engine/Plugins/Experimental/Dataflow/Source/DataflowEditor/DataflowEditor.Build.cs
+++ b/Engine/Plugins/Experimental/Dataflow/Source/DataflowEditor/DataflowEditor.Build.cs
@@ -34,6 +34,7 @@ namespace UnrealBuildTool.Rules
 					"DataflowEngine",
 					"DataflowEnginePlugin",
 					"Slate",
+					"Chaos",
 				}
 			);
 		}
diff --git a/Engine/Plugins/Experimental/Dataflow/Source/DataflowEditor/Private/Dataflow/DataflowEditorViewport.cpp b/Engine/Plugins/Experimental/Dataflow/Source/DataflowEditor/Private/Dataflow/DataflowEditorViewport.cpp
--- a/Engine/Plugins/Experimental/Dataflow/Source/DataflowEditor/Private/Dataflow/DataflowEditorViewport.cpp
+++ b/Engine/Plugins/Experimental/Dataflow/Source/DataflowEditor/Private/Dataflow/DataflowEditorViewport.cpp
@@ -6,9 +6,18 @@
 #include "Dataflow/DataflowObject.h"
 #include "Dataflow/DataflowRenderingActor.h"
 #include "Dataflow/DataflowRenderingComponent.h"
-#include "EditorViewportClient.h"
+#include "DynamicMeshBuilder.h"
 #include "Dataflow/DataflowEditorToolkit.h"
+#include "EditorModes.h"
+#include "EditorViewportClient.h"
+#include "Materials/Material.h"
 #include "SCommonEditorViewportToolbarBase.h"
+#include "Settings/EditorStyleSettings.h"
+
+
+static int32 bEnableDataflowClientDrawing = 0;
+FAutoConsoleVariableRef CVarbEnableDataflowClientDrawing(
+	TEXT("p.Dataflow.Rendering.ClientDraw"), bEnableDataflowClientDrawing,TEXT("Enable the client rendering, will be slower on complex bodies. [default:0]"));
 
 SDataflowEditorViewport::SDataflowEditorViewport()
 {
@@ -49,7 +58,7 @@ void SDataflowEditorViewport::OnFloatingButtonClicked()
 }
 
 void SDataflowEditorViewport::AddReferencedObjects(FReferenceCollector& Collector)
-{							
+{
 	Collector.AddReferencedObject(CustomDataflowRenderingActor);
 }
 
@@ -64,7 +73,6 @@ TSharedPtr<SWidget> SDataflowEditorViewport::MakeViewportToolbar()
 	return SNew(SCommonEditorViewportToolbarBase, SharedThis(this));
 }
 
-
 // ----------------------------------------------------------------------------------
 
 FDataflowEditorViewportClient::FDataflowEditorViewportClient(FPreviewScene* InPreviewScene, 
@@ -74,8 +82,9 @@ FDataflowEditorViewportClient::FDataflowEditorViewportClient(FPreviewScene* InPr
 	FEditorViewportClient(nullptr, InPreviewScene, InEditorViewportWidget)
 	, DataflowEditorToolkitPtr(InDataflowEditorToolkitPtr)
 {
-	bSetListenerPosition = false;
 	SetRealtime(true);
+	SetViewModes(VMI_Lit, VMI_Lit);
+	bSetListenerPosition = false;
 	EngineShowFlags.Grid = false;
 }
 
@@ -88,6 +97,124 @@ Dataflow::FTimestamp FDataflowEditorViewportClient::LatestTimestamp(const UDataf
 	return Dataflow::FTimestamp::Invalid;
 }
 
+void FDataflowEditorViewportClient::ProcessClick(FSceneView& View, HHitProxy* HitProxy, FKey Key, EInputEvent Event, uint32 HitX, uint32 HitY)
+{
+	Super::ProcessClick(View, HitProxy, Key, Event, HitX, HitY);
+
+}
+
+
+void FDataflowEditorViewportClient::Draw(const FSceneView* View, FPrimitiveDrawInterface* PDI)
+{
+	Super::Draw(View, PDI);
+	if (bEnableDataflowClientDrawing)
+	{
+		GeometryCollection::Facades::FRenderingFacade Facade(&RenderCollection);
+		if (Facade.IsValid())
+		{
+			if (bRenderMesh)
+			{
+				MeshBuilder.Reset(new FDynamicMeshBuilder(View->GetFeatureLevel()));
+
+				FLinearColor UnselectedColor = GEngine->C_BrushWire;
+				//UnselectedColor.A = .1f;
+
+				FLinearColor SelectedColor = GetDefault<UEditorStyleSettings>()->SelectionColor;
+				SelectedColor.A = .5f;
+
+				// Allocate the material proxy and register it so it can be deleted properly once the rendering is done with it.
+				FDynamicColoredMaterialRenderProxy* SelectedColorInstance = new FDynamicColoredMaterialRenderProxy(GEngine->GeomMaterial->GetRenderProxy(), SelectedColor);
+				PDI->RegisterDynamicResource(SelectedColorInstance);
+
+				FDynamicColoredMaterialRenderProxy* UnselectedColorInstance = new FDynamicColoredMaterialRenderProxy(GEngine->GeomMaterial->GetRenderProxy(), UnselectedColor);
+				PDI->RegisterDynamicResource(UnselectedColorInstance);
+
+				// mesh surface
+				bool bRenderSurface = false;
+				if (bRenderSurface)
+				{
+					MeshBuilder->AddVertices(VertexBuffer);
+					MeshBuilder->AddTriangles(IndexBuffer);
+					MeshBuilder->Draw(PDI, FMatrix::Identity, UnselectedColorInstance, ESceneDepthPriorityGroup::SDPG_World, false);
+				}
+				else
+				{
+					// wireframes
+					auto ToF = [](FVector3f V) { return FVector(V.X, V.Y, V.Z); };
+					for (int i = 0; i < IndexBuffer.Num(); i += 3)
+					{
+						PDI->DrawLine(ToF(VertexBuffer[i].Position), ToF(VertexBuffer[i + 1].Position), SelectedColor, SDPG_World, 1.f);
+						PDI->DrawLine(ToF(VertexBuffer[i].Position), ToF(VertexBuffer[i + 2].Position), SelectedColor, SDPG_World, 1.f);
+						PDI->DrawLine(ToF(VertexBuffer[i + 1].Position), ToF(VertexBuffer[i + 2].Position), SelectedColor, SDPG_World, 1.f);
+					}
+				}
+			}
+		}
+	}
+}
+
+void FDataflowEditorViewportClient::ReleaseRenderStructures()
+{
+	RenderCollection = FManagedArrayCollection();
+
+	// MeshBuilder
+	if (bEnableDataflowClientDrawing)
+	{
+		bRenderMesh = false;
+		IndexBuffer = TArray<uint32>();
+		VertexBuffer = TArray<FDynamicMeshVertex>();
+		MeshBuilder.Reset(nullptr);
+	}
+}
+
+void FDataflowEditorViewportClient::RenderIntoStructures()
+{
+	GeometryCollection::Facades::FRenderingFacade Facade(&RenderCollection);
+	if (bEnableDataflowClientDrawing)
+	{
+		if (Facade.IsValid())
+		{
+			// @todo(flesh) : Setup buffers that will get directly copied to 
+			//                the render. This is only called during invalidation
+			//                of the graph, so any state that can be precalculated 
+			//                should be set here and just pushed to the render in
+			//                the ::Draw() method.
+			// 
+			//	Note: Draw will still be very slow for complicated geometry, 
+			//        we still need to rely on Component rendering through
+			//        a scene proxy as the primary rendering path. 
+			{
+				const TManagedArray< FIntVector >& Indices = *Facade.GetIndices();
+				const TManagedArray< FVector3f >& Vertices = *Facade.GetVertices();
+				if (Vertices.Num() && Indices.Num())
+				{
+					IndexBuffer.SetNum(Indices.Num() * 3);
+					for (int Idx = 0; Idx < Indices.Num() * 3; Idx++)
+					{
+						IndexBuffer[Idx] = Idx;
+					}
+
+					int Vdx = 0;
+					VertexBuffer.SetNum(Indices.Num() * 3);
+					for (const FIntVector& Face : Indices)
+					{
+						FIntVector Tri = FIntVector(Face[0], Face[1], Face[2]);
+						const FVector3f Pos0(Vertices[Tri[0]]), Pos1(Vertices[Tri[1]]), Pos2(Vertices[Tri[2]]);
+						const FVector3f Normal = FVector3f::CrossProduct(Pos2 - Pos0, Pos1 - Pos0).GetSafeNormal();
+						const FVector3f Tangent = ((Pos1 + Pos2) * 0.5f - Pos0).GetSafeNormal();
+
+						VertexBuffer[Vdx++] = FDynamicMeshVertex(Pos0, Tangent, Normal, FVector2f(0.f, 0.f), FColor::White);
+						VertexBuffer[Vdx++] = FDynamicMeshVertex(Pos1, Tangent, Normal, FVector2f(0.f, 1.f), FColor::White);
+						VertexBuffer[Vdx++] = FDynamicMeshVertex(Pos2, Tangent, Normal, FVector2f(1.f, 1.f), FColor::White);
+					}
+
+					bRenderMesh = true;
+				}
+			}
+		}
+	}
+}
+
 
 void FDataflowEditorViewportClient::Tick(float DeltaSeconds)
 {
@@ -99,20 +226,28 @@ void FDataflowEditorViewportClient::Tick(float DeltaSeconds)
 	{
 		if (TSharedPtr<Dataflow::FContext> Context = DataflowEditorToolkit->GetContext())
 		{
-			if (UDataflowRenderingComponent* DataflowRenderingComponent = DataflowRenderingActor->GetDataflowRenderingComponent())
+			if (const UDataflow* Dataflow = DataflowEditorToolkit->GetDataflow())
 			{
-				if (const UDataflow* Dataflow = DataflowEditorToolkit->GetDataflow())
+				if (UDataflowRenderingComponent* DataflowRenderingComponent = DataflowRenderingActor->GetDataflowRenderingComponent())
 				{
 					Dataflow::FTimestamp SystemTimestamp = LatestTimestamp(Dataflow, Context.Get());
 					if (SystemTimestamp >= LastModifiedTimestamp)
 					{
+						ReleaseRenderStructures();
+
 						if (Dataflow->GetRenderTargets().Num())
 						{
-							// @todo(dataflow) : Check the Time on the target outs first instead.
-							//					 to avoid invalidation during unrelated edits. 
+							GeometryCollection::Facades::FRenderingFacade Facade(&RenderCollection);
+							for (const UDataflowEdNode* Target : Dataflow->GetRenderTargets())
+							{
+								Target->Render(Facade, Context);
+							}
 
-							DataflowRenderingComponent->ResetRenderTargets();
+							// Client View Rendering
+							RenderIntoStructures();
 
+							// Component Object Rendering
+							DataflowRenderingComponent->ResetRenderTargets();
 							DataflowRenderingComponent->SetDataflow(Dataflow);
 							DataflowRenderingComponent->SetContext(Context);
 							for (const UDataflowEdNode* Node : Dataflow->GetRenderTargets())
@@ -124,12 +259,14 @@ void FDataflowEditorViewportClient::Tick(float DeltaSeconds)
 						{
 							DataflowRenderingComponent->ResetRenderTargets();
 						}
+
 						LastModifiedTimestamp = LatestTimestamp(Dataflow, Context.Get()).Value + 1;
 					}
 				}
 			}
 		}
 	}
+
 	// Tick the preview scene world.
 	if (!GIntraFrameDebuggingGameThread)
 	{
@@ -141,11 +278,6 @@ void FDataflowEditorViewportClient::Tick(float DeltaSeconds)
 void FDataflowEditorViewportClient::AddReferencedObjects(FReferenceCollector& Collector)
 {
 	Super::AddReferencedObjects(Collector);
-
-	if (DataflowRenderingActor)
-	{
-		Collector.AddReferencedObject(DataflowRenderingActor);
-	}
 }
 
 
diff --git a/Engine/Plugins/Experimental/Dataflow/Source/DataflowEditor/Public/Dataflow/DataflowEditorViewport.h b/Engine/Plugins/Experimental/Dataflow/Source/DataflowEditor/Public/Dataflow/DataflowEditorViewport.h
--- a/Engine/Plugins/Experimental/Dataflow/Source/DataflowEditor/Public/Dataflow/DataflowEditorViewport.h
+++ b/Engine/Plugins/Experimental/Dataflow/Source/DataflowEditor/Public/Dataflow/DataflowEditorViewport.h
@@ -1,9 +1,11 @@
 // Copyright Epic Games, Inc. All Rights Reserved.
 #pragma once
 
+#include "CanvasTypes.h"
 #include "CoreMinimal.h"
 #include "Dataflow/DataflowNodeParameters.h"
 #include "EditorViewportClient.h"
+#include "GeometryCollection/ManagedArrayCollection.h"
 #include "SAssetEditorViewport.h"
 #include "SCommonEditorViewportToolbarBase.h"
 
@@ -14,6 +16,7 @@ class FDataflowEditorViewportClient;
 class SEditorViewport;
 class ADataflowActor;
 class ADataflowRenderingActor;
+class FDynamicMeshBuilder;
 
 // ----------------------------------------------------------------------------------
 
@@ -37,6 +40,7 @@ public:
 	virtual void AddReferencedObjects(FReferenceCollector& Collector) override;
 	virtual FString GetReferencerName() const override{return TEXT("SDataflowEditorViewport");}
 
+
 protected:
 	virtual TSharedRef<FEditorViewportClient> MakeEditorViewportClient() override;
 	virtual TSharedPtr<SWidget> MakeViewportToolbar() override;
@@ -48,7 +52,6 @@ private:
 	TSharedPtr<FDataflowEditorViewportClient> EditorViewportClient;
 
 	TWeakPtr<FDataflowEditorToolkit> DataflowEditorToolkitPtr;
-
 	ADataflowActor* CustomDataflowActor = nullptr;
 	ADataflowRenderingActor* CustomDataflowRenderingActor = nullptr;
 };
@@ -65,11 +68,12 @@ public:
 		const TWeakPtr<SEditorViewport> InEditorViewportWidget = nullptr,
 		TWeakPtr<FDataflowEditorToolkit> InDataflowEditorToolkitPtr = nullptr);
 
-
-	void SetDataflowRenderingActor(ADataflowRenderingActor* InActor) { DataflowRenderingActor = InActor; }
 	Dataflow::FTimestamp LatestTimestamp(const UDataflow* Dataflow, const Dataflow::FContext* Context);
+	void SetDataflowRenderingActor(ADataflowRenderingActor* InActor) { DataflowRenderingActor = InActor; }
 
 	// FEditorViewportClient interface
+	virtual void ProcessClick(FSceneView& View, HHitProxy* HitProxy, FKey Key, EInputEvent Event, uint32 HitX, uint32 HitY) override;
+	virtual void Draw(const FSceneView* View, FPrimitiveDrawInterface* PDI) override;
 	virtual void Tick(float DeltaSeconds) override;
 	// End of FEditorViewportClient
 
@@ -77,7 +81,20 @@ public:
 	virtual void AddReferencedObjects(FReferenceCollector& Collector) override;
 	virtual FString GetReferencerName() const override { return TEXT("FDataflowEditorViewportClient"); }
 
+private:
+	void RenderIntoStructures();
+	void ReleaseRenderStructures();
+
 	TWeakPtr<FDataflowEditorToolkit> DataflowEditorToolkitPtr = nullptr;
-	ADataflowRenderingActor* DataflowRenderingActor = nullptr;
 	Dataflow::FTimestamp LastModifiedTimestamp = Dataflow::FTimestamp::Invalid;
+	ADataflowRenderingActor* DataflowRenderingActor = nullptr;
+
+	FManagedArrayCollection RenderCollection;
+
+	// Renderables
+	bool bRenderMesh = false;
+	TArray<uint32> IndexBuffer;
+	TArray<FDynamicMeshVertex> VertexBuffer;
+	TUniquePtr<FDynamicMeshBuilder> MeshBuilder = nullptr;
+
 };
diff --git a/Engine/Plugins/Experimental/Dataflow/Source/DataflowEnginePlugin/Private/Dataflow/DataflowEnginePlugin.cpp b/Engine/Plugins/Experimental/Dataflow/Source/DataflowEnginePlugin/Private/Dataflow/DataflowEnginePlugin.cpp
--- a/Engine/Plugins/Experimental/Dataflow/Source/DataflowEnginePlugin/Private/Dataflow/DataflowEnginePlugin.cpp
+++ b/Engine/Plugins/Experimental/Dataflow/Source/DataflowEnginePlugin/Private/Dataflow/DataflowEnginePlugin.cpp
@@ -5,6 +5,7 @@
 #include "CoreMinimal.h"
 #include "Modules/ModuleManager.h"
 #include "Dataflow/DataflowRenderingFactory.h"
+#include "Dataflow/DataflowEngineRendering.h"
 #include "GeometryCollection/GeometryCollectionUtility.h"
 #include "GeometryCollection/Facades/CollectionRenderingFacade.h"
 #include "GeometryCollection/ManagedArrayCollection.h"
@@ -24,71 +25,7 @@ IMPLEMENT_MODULE( FDataflowEnginePlugin, DataflowEnginePlugin)
 
 void FDataflowEnginePlugin::StartupModule()
 {
-	using namespace Dataflow;
-
-	/**
-	* DatflowNode (FGeometryCollection) Rendering
-	*
-	*		@param Type : FGeometryCollection::StaticType()
-
-	*		@param Outputs : {FManagedArrayCollection : "Collection"}
-	*/
-	FRenderingFactory::GetInstance()->RegisterOutput(FGeometryCollection::StaticType(),
-		[](GeometryCollection::Facades::FRenderingFacade& RenderCollection, const Dataflow::FGraphRenderingState& State)
-		{
-			if (State.GetRenderOutputs().Num())
-			{
-				FName PrimaryOutput = State.GetRenderOutputs()[0]; // "Collection"
-				
-				FManagedArrayCollection Default;
-				const FManagedArrayCollection& Collection = State.GetValue<FManagedArrayCollection>(PrimaryOutput, Default);
-
-				if (Collection.FindAttributeTyped<FIntVector>("Indices", FGeometryCollection::FacesGroup)
-					&& Collection.FindAttributeTyped<FVector3f>("Vertex", FGeometryCollection::VerticesGroup)
-					&& Collection.FindAttributeTyped<FTransform>(FTransformCollection::TransformAttribute, FTransformCollection::TransformGroup)
-					&& Collection.FindAttributeTyped<int32>("BoneMap", FGeometryCollection::VerticesGroup)
-					&& Collection.FindAttributeTyped<int32>(FTransformCollection::ParentAttribute, FTransformCollection::TransformGroup))
-				{
-					const TManagedArray<int32>& BoneIndex = Collection.GetAttribute<int32>("BoneMap", FGeometryCollection::VerticesGroup);
-					const TManagedArray<int32>& Parents = Collection.GetAttribute<int32>(FTransformCollection::ParentAttribute, FTransformCollection::TransformGroup);
-					const TManagedArray<FTransform>& Transforms = Collection.GetAttribute<FTransform>(FTransformCollection::TransformAttribute, FTransformCollection::TransformGroup);
-
-					TArray<FTransform> M;
-					GeometryCollectionAlgo::GlobalMatrices(Transforms, Parents, M);
-
-					auto ToD = [](FVector3f V) { return FVector3d(V.X, V.Y, V.Z); };
-					auto ToF = [](FVector3d V) { return FVector3f(V.X, V.Y, V.Z); };
-
-
-					const TManagedArray<FVector3f>& Vertex = Collection.GetAttribute<FVector3f>("Vertex", FGeometryCollection::VerticesGroup);
-					const TManagedArray<FIntVector>& Faces = Collection.GetAttribute<FIntVector>("Indices", FGeometryCollection::FacesGroup);
-
-					TArray<FVector3f> Vertices; Vertices.AddUninitialized(Vertex.Num());
-					TArray<FIntVector> Tris; Tris.AddUninitialized(Faces.Num());
-					TArray<bool> Visited; Visited.Init(false, Vertices.Num());
-
-					int32 Tdx = 0;
-					for (const FIntVector& Face : Faces)
-					{
-						FIntVector Tri = FIntVector(Face[0], Face[1], Face[2]);
-						FTransform Ms[3] = { M[BoneIndex[Tri[0]]], M[BoneIndex[Tri[1]]], M[BoneIndex[Tri[2]]] };
-
-						Tris[Tdx++] = Tri;
-						if (!Visited[Tri[0]]) Vertices[Tri[0]] = ToF(Ms[0].TransformPosition(ToD(Vertex[Tri[0]])));
-						if (!Visited[Tri[1]]) Vertices[Tri[1]] = ToF(Ms[1].TransformPosition(ToD(Vertex[Tri[1]])));
-						if (!Visited[Tri[2]]) Vertices[Tri[2]] = ToF(Ms[2].TransformPosition(ToD(Vertex[Tri[2]])));
-
-						Visited[Tri[0]] = true; Visited[Tri[1]] = true; Visited[Tri[2]] = true;
-					}
-
-					// Maybe these buffers should be shrunk, but there are unused vertices in the buffer. 
-					for (int i = 0; i < Visited.Num(); i++) if (!Visited[i]) Vertices[i] = FVector3f(0);
-
-					RenderCollection.AddSurface(MoveTemp(Vertices),MoveTemp(Tris));
-				}
-			}
-		});
-
+	Dataflow::RenderingCallbacks();
 }
 
 
diff --git a/Engine/Plugins/Experimental/Dataflow/Source/DataflowEnginePlugin/Private/Dataflow/DataflowEngineRendering.cpp b/Engine/Plugins/Experimental/Dataflow/Source/DataflowEnginePlugin/Private/Dataflow/DataflowEngineRendering.cpp
new file mode 100644
--- /dev/null
+++ b/Engine/Plugins/Experimental/Dataflow/Source/DataflowEnginePlugin/Private/Dataflow/DataflowEngineRendering.cpp
@@ -0,0 +1,84 @@
+// Copyright Epic Games, Inc. All Rights Reserved.
+
+#include "Dataflow/DataflowEngineRendering.h"
+
+#include "Dataflow/DataflowRenderingFactory.h"
+#include "GeometryCollection/GeometryCollectionUtility.h"
+#include "GeometryCollection/Facades/CollectionRenderingFacade.h"
+#include "GeometryCollection/ManagedArrayCollection.h"
+#include "GeometryCollection/GeometryCollectionAlgo.h"
+
+namespace Dataflow
+{
+
+	void RenderingCallbacks()
+	{
+		using namespace Dataflow;
+
+		/**
+		* DatflowNode (FGeometryCollection) Rendering
+		*
+		*		@param Type : FGeometryCollection::StaticType()
+
+		*		@param Outputs : {FManagedArrayCollection : "Collection"}
+		*/
+		FRenderingFactory::GetInstance()->RegisterOutput(FGeometryCollection::StaticType(),
+			[](GeometryCollection::Facades::FRenderingFacade& RenderCollection, const Dataflow::FGraphRenderingState& State)
+			{
+				if (State.GetRenderOutputs().Num())
+				{
+					FName PrimaryOutput = State.GetRenderOutputs()[0]; // "Collection"
+
+					FManagedArrayCollection Default;
+					const FManagedArrayCollection& Collection = State.GetValue<FManagedArrayCollection>(PrimaryOutput, Default);
+
+					if (Collection.FindAttributeTyped<FIntVector>("Indices", FGeometryCollection::FacesGroup)
+						&& Collection.FindAttributeTyped<FVector3f>("Vertex", FGeometryCollection::VerticesGroup)
+						&& Collection.FindAttributeTyped<FTransform>(FTransformCollection::TransformAttribute, FTransformCollection::TransformGroup)
+						&& Collection.FindAttributeTyped<int32>("BoneMap", FGeometryCollection::VerticesGroup)
+						&& Collection.FindAttributeTyped<int32>(FTransformCollection::ParentAttribute, FTransformCollection::TransformGroup))
+					{
+						const TManagedArray<int32>& BoneIndex = Collection.GetAttribute<int32>("BoneMap", FGeometryCollection::VerticesGroup);
+						const TManagedArray<int32>& Parents = Collection.GetAttribute<int32>(FTransformCollection::ParentAttribute, FTransformCollection::TransformGroup);
+						const TManagedArray<FTransform>& Transforms = Collection.GetAttribute<FTransform>(FTransformCollection::TransformAttribute, FTransformCollection::TransformGroup);
+
+						TArray<FTransform> M;
+						GeometryCollectionAlgo::GlobalMatrices(Transforms, Parents, M);
+
+						auto ToD = [](FVector3f V) { return FVector3d(V.X, V.Y, V.Z); };
+						auto ToF = [](FVector3d V) { return FVector3f(V.X, V.Y, V.Z); };
+
+
+						const TManagedArray<FVector3f>& Vertex = Collection.GetAttribute<FVector3f>("Vertex", FGeometryCollection::VerticesGroup);
+						const TManagedArray<FIntVector>& Faces = Collection.GetAttribute<FIntVector>("Indices", FGeometryCollection::FacesGroup);
+
+						TArray<FVector3f> Vertices; Vertices.AddUninitialized(Vertex.Num());
+						TArray<FIntVector> Tris; Tris.AddUninitialized(Faces.Num());
+						TArray<bool> Visited; Visited.Init(false, Vertices.Num());
+
+						int32 Tdx = 0;
+						for (const FIntVector& Face : Faces)
+						{
+							FIntVector Tri = FIntVector(Face[0], Face[1], Face[2]);
+							FTransform Ms[3] = { M[BoneIndex[Tri[0]]], M[BoneIndex[Tri[1]]], M[BoneIndex[Tri[2]]] };
+
+							Tris[Tdx++] = Tri;
+							if (!Visited[Tri[0]]) Vertices[Tri[0]] = ToF(Ms[0].TransformPosition(ToD(Vertex[Tri[0]])));
+							if (!Visited[Tri[1]]) Vertices[Tri[1]] = ToF(Ms[1].TransformPosition(ToD(Vertex[Tri[1]])));
+							if (!Visited[Tri[2]]) Vertices[Tri[2]] = ToF(Ms[2].TransformPosition(ToD(Vertex[Tri[2]])));
+
+							Visited[Tri[0]] = true; Visited[Tri[1]] = true; Visited[Tri[2]] = true;
+						}
+
+						// Maybe these buffers should be shrunk, but there are unused vertices in the buffer. 
+						for (int i = 0; i < Visited.Num(); i++) if (!Visited[i]) Vertices[i] = FVector3f(0);
+
+						RenderCollection.AddSurface(MoveTemp(Vertices), MoveTemp(Tris));
+					}
+				}
+			});
+
+	}
+
+
+}
diff --git a/Engine/Plugins/Experimental/Dataflow/Source/DataflowEnginePlugin/Public/Dataflow/DataflowEngineRendering.h b/Engine/Plugins/Experimental/Dataflow/Source/DataflowEnginePlugin/Public/Dataflow/DataflowEngineRendering.h
new file mode 100644
--- /dev/null
+++ b/Engine/Plugins/Experimental/Dataflow/Source/DataflowEnginePlugin/Public/Dataflow/DataflowEngineRendering.h
@@ -0,0 +1,15 @@
+// Copyright Epic Games, Inc. All Rights Reserved.
+
+#pragma once
+
+#include "CoreMinimal.h"
+
+
+namespace Dataflow
+{
+	/**
+	 * Rendering callbacks for the Dataflow Rendering Factory
+	 */
+	void RenderingCallbacks();
+
+}
\ No newline at end of file
diff --git a/Engine/Source/Runtime/Experimental/Chaos/Private/GeometryCollection/Facades/CollectionRenderingFacade.cpp b/Engine/Source/Runtime/Experimental/Chaos/Private/GeometryCollection/Facades/CollectionRenderingFacade.cpp
--- a/Engine/Source/Runtime/Experimental/Chaos/Private/GeometryCollection/Facades/CollectionRenderingFacade.cpp
+++ b/Engine/Source/Runtime/Experimental/Chaos/Private/GeometryCollection/Facades/CollectionRenderingFacade.cpp
@@ -10,6 +10,15 @@
 namespace GeometryCollection::Facades
 {
 
+	// Groups 
+	const FName FRenderingFacade::VerticesGroup = FGeometryCollection::VerticesGroup;
+	const FName FRenderingFacade::FacesGroup = FGeometryCollection::FacesGroup;
+
+
+	// Attributes
+	const FName FRenderingFacade::VertexAttribute = "Vertex";
+	const FName FRenderingFacade::IndicesAttribute = "Indices";
+
 	FRenderingFacade::FRenderingFacade(FManagedArrayCollection* InCollection)
 		: Self(InCollection)
 	{
@@ -34,24 +43,26 @@ namespace GeometryCollection::Facades
 			InCollection->AddGroup(FGeometryCollection::FacesGroup);
 		}
 
-		InCollection->AddAttribute<FVector3f>("Vertex", FGeometryCollection::VerticesGroup);
-		InCollection->AddAttribute<FIntVector>("Indices", FGeometryCollection::FacesGroup, VertexDependency);
+		InCollection->AddAttribute<FVector3f>(VertexAttribute, FGeometryCollection::VerticesGroup);
+		InCollection->AddAttribute<FIntVector>(IndicesAttribute, FGeometryCollection::FacesGroup, VertexDependency);
 
-		ensure(InCollection->FindAttributeTyped<FVector3f>("Vertex", FGeometryCollection::VerticesGroup) != nullptr);
-		ensure(InCollection->FindAttributeTyped<FIntVector>("Indices", FGeometryCollection::FacesGroup) != nullptr);
+		ensure(InCollection->FindAttributeTyped<FVector3f>(VertexAttribute, FGeometryCollection::VerticesGroup) != nullptr);
+		ensure(InCollection->FindAttributeTyped<FIntVector>(IndicesAttribute, FGeometryCollection::FacesGroup) != nullptr);
 	}
 
 	bool FRenderingFacade::CanRenderSurface(const FManagedArrayCollection* InCollection)
 	{
-		return InCollection->HasGroup(FGeometryCollection::VerticesGroup)
-			&& InCollection->HasGroup(FGeometryCollection::FacesGroup)
-			&& InCollection->FindAttributeTyped<FVector3f>("Vertex", FGeometryCollection::VerticesGroup)
-			&& InCollection->FindAttributeTyped<FIntVector>("Indices", FGeometryCollection::FacesGroup);
+		return  InCollection && IsValid(InCollection) &&
+			GetIndices(InCollection)->Num() &&
+			GetVertices(InCollection)->Num();
 	}
 
 	bool FRenderingFacade::IsValid(const FManagedArrayCollection* InCollection)
 	{
-		return FRenderingFacade::CanRenderSurface(InCollection);
+		return InCollection && InCollection->HasGroup(FGeometryCollection::VerticesGroup)
+			&& InCollection->HasGroup(FGeometryCollection::FacesGroup)
+			&& InCollection->FindAttributeTyped<FVector3f>(VertexAttribute, FGeometryCollection::VerticesGroup)
+			&& InCollection->FindAttributeTyped<FIntVector>(IndicesAttribute, FGeometryCollection::FacesGroup);
 	}
 
 	void FRenderingFacade::AddTriangle(FManagedArrayCollection* InCollection, const Chaos::FTriangle& InTriangle)
@@ -60,8 +71,8 @@ namespace GeometryCollection::Facades
 		{
 			auto CollectionVert = [](const Chaos::FVec3& V) { return FVector3f(float(V.X), float(V.Y), float(V.Z)); };
 
-			TManagedArray<FVector3f>& Vertices = InCollection->ModifyAttribute<FVector3f>("Vertex", FGeometryCollection::VerticesGroup);
-			TManagedArray<FIntVector>& Indices = InCollection->ModifyAttribute<FIntVector>("Indices", FGeometryCollection::FacesGroup);
+			TManagedArray<FVector3f>& Vertices = InCollection->ModifyAttribute<FVector3f>(VertexAttribute, FGeometryCollection::VerticesGroup);
+			TManagedArray<FIntVector>& Indices = InCollection->ModifyAttribute<FIntVector>(IndicesAttribute, FGeometryCollection::FacesGroup);
 
 			int32 IndicesStart = InCollection->AddElements(1, FGeometryCollection::FacesGroup);
 			int32 VertexStart = InCollection->AddElements(3, FGeometryCollection::VerticesGroup);
@@ -80,8 +91,8 @@ namespace GeometryCollection::Facades
 		{
 			auto CollectionVert = [](const Chaos::FVec3& V) { return FVector3f(float(V.X), float(V.Y), float(V.Z)); };
 
-			TManagedArray<FVector3f>& Vertices = InCollection->ModifyAttribute<FVector3f>("Vertex", FGeometryCollection::VerticesGroup);
-			TManagedArray<FIntVector>& Indices = InCollection->ModifyAttribute<FIntVector>("Indices", FGeometryCollection::FacesGroup);
+			TManagedArray<FVector3f>& Vertices = InCollection->ModifyAttribute<FVector3f>(VertexAttribute, FGeometryCollection::VerticesGroup);
+			TManagedArray<FIntVector>& Indices = InCollection->ModifyAttribute<FIntVector>(IndicesAttribute, FGeometryCollection::FacesGroup);
 
 			int32 IndicesStart = InCollection->AddElements(InIndices.Num(), FGeometryCollection::FacesGroup);
 			int32 VertexStart = InCollection->AddElements(InVertices.Num(), FGeometryCollection::VerticesGroup);
@@ -100,12 +111,12 @@ namespace GeometryCollection::Facades
 
 	const TManagedArray< FIntVector >* FRenderingFacade::GetIndices(const FManagedArrayCollection* InCollection)
 	{
-		return InCollection->FindAttribute<FIntVector>("Indices", FGeometryCollection::FacesGroup);
+		return InCollection->FindAttribute<FIntVector>(IndicesAttribute, FGeometryCollection::FacesGroup);
 	}
 
 	const TManagedArray< FVector3f >* FRenderingFacade::GetVertices(const FManagedArrayCollection* InCollection)
 	{
-		return InCollection->FindAttribute<FVector3f>("Vertex", FGeometryCollection::VerticesGroup);
+		return InCollection->FindAttribute<FVector3f>(VertexAttribute, FGeometryCollection::VerticesGroup);
 	}
 
 
diff --git a/Engine/Source/Runtime/Experimental/Chaos/Public/GeometryCollection/Facades/CollectionRenderingFacade.h b/Engine/Source/Runtime/Experimental/Chaos/Public/GeometryCollection/Facades/CollectionRenderingFacade.h
--- a/Engine/Source/Runtime/Experimental/Chaos/Public/GeometryCollection/Facades/CollectionRenderingFacade.h
+++ b/Engine/Source/Runtime/Experimental/Chaos/Public/GeometryCollection/Facades/CollectionRenderingFacade.h
@@ -14,11 +14,10 @@ namespace GeometryCollection::Facades
 	/**
 	* FRenderingFacade
 	*
-	* Defines common API for storing a vertex weights bound to a bone. This mapping is from the
-	* the vertex to the bone index. The FSelectionFacad will store the mapping from the BoneIndex
-	* to the vertex.
+	* Defines common API for storing rendering data.
+	*
 	* Usage:
-	*    FRenderingFacade::AddBoneWeights(this, FSelectionFacade(this) );
+	*    FRenderingFacade::AddTriangle(this, FTriangle(...) );
 	*
 	* Then arrays can be accessed later by:
 	*	const TManagedArray< FIntVector >* FaceIndices = FRenderingFacade::GetIndices(this);
@@ -35,6 +34,13 @@ namespace GeometryCollection::Facades
 		FManagedArrayCollection* Self;
 
 	public:
+		// Groups 
+		static const FName VerticesGroup;
+		static const FName FacesGroup;
+
+		// Attributes
+		static const FName VertexAttribute;
+		static const FName IndicesAttribute;
 
 		/**
 		* FRenderingFacade Constuctor
