Dataflow : Fix CIS
- Removed nullptr assignment to UniquePtr.

diff --git a/Engine/Plugins/Experimental/Dataflow/Source/DataflowEditor/Public/Dataflow/DataflowEditorViewport.h b/Engine/Plugins/Experimental/Dataflow/Source/DataflowEditor/Public/Dataflow/DataflowEditorViewport.h
--- a/Engine/Plugins/Experimental/Dataflow/Source/DataflowEditor/Public/Dataflow/DataflowEditorViewport.h
+++ b/Engine/Plugins/Experimental/Dataflow/Source/DataflowEditor/Public/Dataflow/DataflowEditorViewport.h
@@ -95,6 +95,6 @@ private:
 	bool bRenderMesh = false;
 	TArray<uint32> IndexBuffer;
 	TArray<FDynamicMeshVertex> VertexBuffer;
-	TUniquePtr<FDynamicMeshBuilder> MeshBuilder = nullptr;
+	TUniquePtr<FDynamicMeshBuilder> MeshBuilder;
 
 };
