Dataflow: Editor Rendering
- Removed the ProceduralMeshComponent rendering
- Added rendering via a SceneProxy
- Added hit buffer support.

diff --git a/Engine/Plugins/Experimental/Dataflow/Source/DataflowEditor/Private/Dataflow/DataflowEditorViewport.cpp b/Engine/Plugins/Experimental/Dataflow/Source/DataflowEditor/Private/Dataflow/DataflowEditorViewport.cpp
--- a/Engine/Plugins/Experimental/Dataflow/Source/DataflowEditor/Private/Dataflow/DataflowEditorViewport.cpp
+++ b/Engine/Plugins/Experimental/Dataflow/Source/DataflowEditor/Private/Dataflow/DataflowEditorViewport.cpp
@@ -4,8 +4,9 @@
 #include "AdvancedPreviewScene.h"
 #include "Dataflow/DataflowEdNode.h"
 #include "Dataflow/DataflowObject.h"
-#include "Dataflow/DataflowRenderingActor.h"
-#include "Dataflow/DataflowRenderingComponent.h"
+#include "Dataflow/DataflowActor.h"
+#include "Dataflow/DataflowComponent.h"
+#include "Dataflow/DataflowEngineSceneHitProxies.h"
 #include "DynamicMeshBuilder.h"
 #include "Dataflow/DataflowEditorToolkit.h"
 #include "EditorModes.h"
@@ -21,8 +22,6 @@ FAutoConsoleVariableRef CVarbEnableDataflowClientDrawing(
 
 SDataflowEditorViewport::SDataflowEditorViewport()
 {
-	// Temporarily allow water subsystem to be created on preview worlds because we need one here : 
-	//UWaterSubsystem::FScopedAllowWaterSubsystemOnPreviewWorld AllowWaterSubsystemOnPreviewWorldScope(true);
 	PreviewScene = MakeShareable(new FAdvancedPreviewScene(FPreviewScene::ConstructionValues()));
 	PreviewScene->SetFloorVisibility(false);
 }
@@ -36,9 +35,9 @@ void SDataflowEditorViewport::Construct(const FArguments& InArgs)
 	SEditorViewport::Construct(SEditorViewport::FArguments());
 
 	FBoxSphereBounds SphereBounds = FBoxSphereBounds(EForceInit::ForceInitToZero);
-	CustomDataflowRenderingActor = CastChecked<ADataflowRenderingActor>(PreviewScene->GetWorld()->SpawnActor(ADataflowRenderingActor::StaticClass()));
+	CustomDataflowActor = CastChecked<ADataflowActor>(PreviewScene->GetWorld()->SpawnActor(ADataflowActor::StaticClass()));
 
-	EditorViewportClient->SetDataflowRenderingActor(CustomDataflowRenderingActor);
+	EditorViewportClient->SetDataflowActor(CustomDataflowActor);
 	EditorViewportClient->FocusViewportOnBox( SphereBounds.GetBox());
 }
 
@@ -59,7 +58,7 @@ void SDataflowEditorViewport::OnFloatingButtonClicked()
 
 void SDataflowEditorViewport::AddReferencedObjects(FReferenceCollector& Collector)
 {
-	Collector.AddReferencedObject(CustomDataflowRenderingActor);
+	Collector.AddReferencedObject(CustomDataflowActor);
 }
 
 TSharedRef<FEditorViewportClient> SDataflowEditorViewport::MakeEditorViewportClient()
@@ -100,7 +99,14 @@ Dataflow::FTimestamp FDataflowEditorViewportClient::LatestTimestamp(const UDataf
 void FDataflowEditorViewportClient::ProcessClick(FSceneView& View, HHitProxy* HitProxy, FKey Key, EInputEvent Event, uint32 HitX, uint32 HitY)
 {
 	Super::ProcessClick(View, HitProxy, Key, Event, HitX, HitY);
-
+	if (HitProxy && HitProxy	->IsA(HDataflowActor::StaticGetType()))
+	{
+		if (DataflowActor && DataflowActor->DataflowComponent)
+		{
+			DataflowActor->DataflowComponent->bSelected = !DataflowActor->DataflowComponent->bSelected;
+			DataflowActor->DataflowComponent->Invalidate();
+		}
+	}
 }
 
 
@@ -222,13 +228,13 @@ void FDataflowEditorViewportClient::Tick(float DeltaSeconds)
 
 	TSharedPtr<FDataflowEditorToolkit> DataflowEditorToolkit = DataflowEditorToolkitPtr.Pin();
 
-	if (DataflowRenderingActor && DataflowEditorToolkitPtr.IsValid())
+	if (DataflowActor && DataflowEditorToolkitPtr.IsValid())
 	{
 		if (TSharedPtr<Dataflow::FContext> Context = DataflowEditorToolkit->GetContext())
 		{
 			if (const UDataflow* Dataflow = DataflowEditorToolkit->GetDataflow())
 			{
-				if (UDataflowRenderingComponent* DataflowRenderingComponent = DataflowRenderingActor->GetDataflowRenderingComponent())
+				if (UDataflowComponent* DataflowComponent = DataflowActor->GetDataflowComponent())
 				{
 					Dataflow::FTimestamp SystemTimestamp = LatestTimestamp(Dataflow, Context.Get());
 					if (SystemTimestamp >= LastModifiedTimestamp)
@@ -247,17 +253,17 @@ void FDataflowEditorViewportClient::Tick(float DeltaSeconds)
 							RenderIntoStructures();
 
 							// Component Object Rendering
-							DataflowRenderingComponent->ResetRenderTargets();
-							DataflowRenderingComponent->SetDataflow(Dataflow);
-							DataflowRenderingComponent->SetContext(Context);
+							DataflowComponent->ResetRenderTargets();
+							DataflowComponent->SetDataflow(Dataflow);
+							DataflowComponent->SetContext(Context);
 							for (const UDataflowEdNode* Node : Dataflow->GetRenderTargets())
 							{
-								DataflowRenderingComponent->AddRenderTarget(Node);
+								DataflowComponent->AddRenderTarget(Node);
 							}
 						}
 						else
 						{
-							DataflowRenderingComponent->ResetRenderTargets();
+							DataflowComponent->ResetRenderTargets();
 						}
 
 						LastModifiedTimestamp = LatestTimestamp(Dataflow, Context.Get()).Value + 1;
diff --git a/Engine/Plugins/Experimental/Dataflow/Source/DataflowEditor/Public/Dataflow/DataflowEditorViewport.h b/Engine/Plugins/Experimental/Dataflow/Source/DataflowEditor/Public/Dataflow/DataflowEditorViewport.h
--- a/Engine/Plugins/Experimental/Dataflow/Source/DataflowEditor/Public/Dataflow/DataflowEditorViewport.h
+++ b/Engine/Plugins/Experimental/Dataflow/Source/DataflowEditor/Public/Dataflow/DataflowEditorViewport.h
@@ -15,7 +15,7 @@ class FEditorViewportClient;
 class FDataflowEditorViewportClient;
 class SEditorViewport;
 class ADataflowActor;
-class ADataflowRenderingActor;
+class ADataflowActor;
 class FDynamicMeshBuilder;
 
 // ----------------------------------------------------------------------------------
@@ -53,7 +53,6 @@ private:
 
 	TWeakPtr<FDataflowEditorToolkit> DataflowEditorToolkitPtr;
 	ADataflowActor* CustomDataflowActor = nullptr;
-	ADataflowRenderingActor* CustomDataflowRenderingActor = nullptr;
 };
 
 
@@ -69,7 +68,7 @@ public:
 		TWeakPtr<FDataflowEditorToolkit> InDataflowEditorToolkitPtr = nullptr);
 
 	Dataflow::FTimestamp LatestTimestamp(const UDataflow* Dataflow, const Dataflow::FContext* Context);
-	void SetDataflowRenderingActor(ADataflowRenderingActor* InActor) { DataflowRenderingActor = InActor; }
+	void SetDataflowActor(ADataflowActor* InActor) { DataflowActor = InActor; }
 
 	// FEditorViewportClient interface
 	virtual void ProcessClick(FSceneView& View, HHitProxy* HitProxy, FKey Key, EInputEvent Event, uint32 HitX, uint32 HitY) override;
@@ -87,8 +86,9 @@ private:
 
 	TWeakPtr<FDataflowEditorToolkit> DataflowEditorToolkitPtr = nullptr;
 	Dataflow::FTimestamp LastModifiedTimestamp = Dataflow::FTimestamp::Invalid;
-	ADataflowRenderingActor* DataflowRenderingActor = nullptr;
+	ADataflowActor* DataflowActor = nullptr;
 
+	bool bIsSelected = false;
 	FManagedArrayCollection RenderCollection;
 
 	// Renderables
diff --git a/Engine/Plugins/Experimental/Dataflow/Source/DataflowEnginePlugin/DataflowEnginePlugin.Build.cs b/Engine/Plugins/Experimental/Dataflow/Source/DataflowEnginePlugin/DataflowEnginePlugin.Build.cs
--- a/Engine/Plugins/Experimental/Dataflow/Source/DataflowEnginePlugin/DataflowEnginePlugin.Build.cs
+++ b/Engine/Plugins/Experimental/Dataflow/Source/DataflowEnginePlugin/DataflowEnginePlugin.Build.cs
@@ -4,7 +4,7 @@ namespace UnrealBuildTool.Rules
 {
 	public class DataflowEnginePlugin : ModuleRules
 	{
-        public DataflowEnginePlugin(ReadOnlyTargetRules Target) : base(Target)
+		public DataflowEnginePlugin(ReadOnlyTargetRules Target) : base(Target)
 		{
 			bTreatAsEngineModule = true;
 			PublicDependencyModuleNames.AddRange(
@@ -12,11 +12,20 @@ namespace UnrealBuildTool.Rules
 				{
 					"Core",
 					"CoreUObject",
+				}
+			);
+			PrivateDependencyModuleNames.AddRange(
+				new string[]
+				{
+					// NOTE: UVEditorTools is a separate module so that it doesn't rely on the editor.
+					// So, do not add editor dependencies here.
+	
 					"Engine",
+					"RenderCore",
+					"RHI",
 					"DataflowCore",
 					"DataflowEngine",
 					"Chaos",
-					"ProceduralMeshComponent"
 				}
 			);
 		}
diff --git a/Engine/Plugins/Experimental/Dataflow/Source/DataflowEnginePlugin/Private/Dataflow/DataflowActor.cpp b/Engine/Plugins/Experimental/Dataflow/Source/DataflowEnginePlugin/Private/Dataflow/DataflowActor.cpp
new file mode 100644
--- /dev/null
+++ b/Engine/Plugins/Experimental/Dataflow/Source/DataflowEnginePlugin/Private/Dataflow/DataflowActor.cpp
@@ -0,0 +1,28 @@
+// Copyright Epic Games, Inc. All Rights Reserved.
+
+#include "Dataflow/DataflowActor.h"
+
+#include "CoreMinimal.h"
+#include "Dataflow/DataflowComponent.h"
+
+
+#include UE_INLINE_GENERATED_CPP_BY_NAME(DataflowActor)
+
+DEFINE_LOG_CATEGORY_STATIC(ADataflowLogging, Log, All);
+
+ADataflowActor::ADataflowActor(const FObjectInitializer& ObjectInitializer)
+	: Super(ObjectInitializer)
+{
+	//UE_LOG(ADataflowLogging, Verbose, TEXT("ADataflowActor::ADataflowActor()"));
+	DataflowComponent = CreateDefaultSubobject<UDataflowComponent>(TEXT("DataflowComponent0"));
+	RootComponent = DataflowComponent;
+	PrimaryActorTick.bCanEverTick = true;
+}
+
+#if WITH_EDITOR
+bool ADataflowActor::GetReferencedContentObjects(TArray<UObject*>& Objects) const
+{
+	return Super::GetReferencedContentObjects(Objects);
+}
+#endif
+
diff --git a/Engine/Plugins/Experimental/Dataflow/Source/DataflowEnginePlugin/Private/Dataflow/DataflowComponent.cpp b/Engine/Plugins/Experimental/Dataflow/Source/DataflowEnginePlugin/Private/Dataflow/DataflowComponent.cpp
new file mode 100644
--- /dev/null
+++ b/Engine/Plugins/Experimental/Dataflow/Source/DataflowEnginePlugin/Private/Dataflow/DataflowComponent.cpp
@@ -0,0 +1,136 @@
+// Copyright Epic Games, Inc. All Rights Reserved.
+
+#include "Dataflow/DataflowComponent.h"
+#include "DataflowEngineSceneProxy.h"
+
+#include "Dataflow/DataflowEdNode.h"
+#include "Dataflow/DataflowRenderingFactory.h"
+#include "GeometryCollection/Facades/CollectionRenderingFacade.h"
+#include "GeometryCollection/Facades/CollectionBoundsFacade.h"
+#include "GeometryCollection/GeometryCollection.h"
+#include "Materials/Material.h"
+
+#include UE_INLINE_GENERATED_CPP_BY_NAME(DataflowComponent)
+
+DEFINE_LOG_CATEGORY_STATIC(LogDataflowRenderComponentInternal, Log, All);
+
+UDataflowComponent::UDataflowComponent(const FObjectInitializer& ObjectInitializer)
+	: Super(ObjectInitializer)
+{
+	PrimaryComponentTick.bCanEverTick = true;
+	bTickInEditor = true;
+}
+
+UDataflowComponent::~UDataflowComponent()
+{
+}
+
+void UDataflowComponent::Invalidate()
+{
+	bUpdateRender = true;
+	bBoundsNeedsUpdate = true;
+	RenderCollection = FManagedArrayCollection();
+}
+
+
+void UDataflowComponent::ResetRenderTargets() 
+{ 
+	RenderTargets.Reset();
+	Invalidate();
+}
+
+void UDataflowComponent::AddRenderTarget(const UDataflowEdNode* InTarget) 
+{ 
+	RenderTargets.AddUnique(InTarget); 
+	Invalidate();
+}
+
+void UDataflowComponent::TickComponent(float DeltaTime, enum ELevelTick TickType, FActorComponentTickFunction *ThisTickFunction)
+{
+	Super::TickComponent(DeltaTime, TickType, ThisTickFunction);
+
+	if (bUpdateRender)
+	{
+		RenderCollection = FManagedArrayCollection();
+		GeometryCollection::Facades::FRenderingFacade Facade(RenderCollection);
+		Facade.DefineSchema();
+
+		bool bNeedsRefresh = false;
+		if (Context && Dataflow)
+		{
+			for (const UDataflowEdNode* Target : RenderTargets)
+			{
+				bNeedsRefresh |= Target->Render(Facade, Context);
+			}
+		}
+		UpdateLocalBounds();
+		bUpdateRender = false;
+
+		MarkRenderStateDirty();
+	}
+}
+
+void UDataflowComponent::UpdateLocalBounds()
+{
+	if (bBoundsNeedsUpdate)
+	{
+		GeometryCollection::Facades::FBoundsFacade(RenderCollection).UpdateBoundingBox();
+		bBoundsNeedsUpdate = false;
+	}
+}
+
+FBoxSphereBounds UDataflowComponent::CalcBounds(const FTransform& LocalToWorldIn) const
+{
+	return BoundingBox.TransformBy(GetComponentTransform());
+}
+
+void UDataflowComponent::SetRenderingCollection(FManagedArrayCollection&& InCollection)
+{
+	Invalidate();
+	RenderCollection = InCollection;
+}
+
+const FManagedArrayCollection& UDataflowComponent::GetRenderingCollection() const
+{
+	return RenderCollection;
+}
+
+
+FMaterialRelevance UDataflowComponent::GetMaterialRelevance(ERHIFeatureLevel::Type InFeatureLevel) const
+{
+	// Combine the material relevance for all materials.
+	FMaterialRelevance Result;
+	for (int32 ElementIndex = 0; ElementIndex < GetNumMaterials(); ElementIndex++)
+	{
+		UMaterialInterface const* MaterialInterface = GetMaterial(ElementIndex);
+		if (!MaterialInterface)
+		{
+			MaterialInterface = UMaterial::GetDefaultMaterial(MD_Surface);
+		}
+		Result |= MaterialInterface->GetRelevance_Concurrent(InFeatureLevel);
+	}
+
+	//if (OverlayMaterial != nullptr)
+	//{
+	//	Result |= OverlayMaterial->GetRelevance_Concurrent(InFeatureLevel);
+	//}
+
+	return Result;
+}
+
+FPrimitiveSceneProxy* UDataflowComponent::CreateSceneProxy()
+{
+	//this->bHasPerInstanceHitProxies = true; // just makes OnClick return null in the hit proxy
+	return GeometryCollection::Facades::FRenderingFacade(RenderCollection).CanRenderSurface() ? new FDataflowEngineSceneProxy(this) : nullptr;
+}
+
+UMaterialInterface* UDataflowComponent::GetMaterial(int32 Index) const
+{
+	return Dataflow->Material ? Dataflow->Material : GetDefaultMaterial();
+}
+
+UMaterialInterface* UDataflowComponent::GetDefaultMaterial() const
+{
+	const TCHAR* Path = TEXT("/Engine/EditorMaterials/Dataflow/DataflowVertexMaterial.DataflowVertexMaterial");
+	return LoadObject<UMaterialInterface>(nullptr, Path, nullptr, LOAD_None, nullptr);
+}
diff --git a/Engine/Plugins/Experimental/Dataflow/Source/DataflowEnginePlugin/Private/Dataflow/DataflowEnginePlugin.cpp b/Engine/Plugins/Experimental/Dataflow/Source/DataflowEnginePlugin/Private/Dataflow/DataflowEnginePlugin.cpp
--- a/Engine/Plugins/Experimental/Dataflow/Source/DataflowEnginePlugin/Private/Dataflow/DataflowEnginePlugin.cpp
+++ b/Engine/Plugins/Experimental/Dataflow/Source/DataflowEnginePlugin/Private/Dataflow/DataflowEnginePlugin.cpp
@@ -12,6 +12,9 @@
 #include "GeometryCollection/GeometryCollectionAlgo.h"
 
 
+FColor IDataflowEnginePlugin::SurfaceColor = FLinearColor(0.6,0.6,0.6).ToRGBE();
+FColor IDataflowEnginePlugin::PrimarySelectionColor = FLinearColor(0.8, 0.3, 0.3).ToRGBE();
+
 class FDataflowEnginePlugin : public IDataflowEnginePlugin
 {
 	/** IModuleInterface implementation */
diff --git a/Engine/Plugins/Experimental/Dataflow/Source/DataflowEnginePlugin/Private/Dataflow/DataflowEngineSceneHitProxies.cpp b/Engine/Plugins/Experimental/Dataflow/Source/DataflowEnginePlugin/Private/Dataflow/DataflowEngineSceneHitProxies.cpp
new file mode 100644
--- /dev/null
+++ b/Engine/Plugins/Experimental/Dataflow/Source/DataflowEnginePlugin/Private/Dataflow/DataflowEngineSceneHitProxies.cpp
@@ -0,0 +1,5 @@
+// Copyright Epic Games, Inc. All Rights Reserved.
+
+#include "Dataflow/DataflowEngineSceneHitProxies.h"
+
+IMPLEMENT_HIT_PROXY(HDataflowActor, HActor);
diff --git a/Engine/Plugins/Experimental/Dataflow/Source/DataflowEnginePlugin/Private/Dataflow/DataflowEngineSceneProxy.cpp b/Engine/Plugins/Experimental/Dataflow/Source/DataflowEnginePlugin/Private/Dataflow/DataflowEngineSceneProxy.cpp
new file mode 100644
--- /dev/null
+++ b/Engine/Plugins/Experimental/Dataflow/Source/DataflowEnginePlugin/Private/Dataflow/DataflowEngineSceneProxy.cpp
@@ -0,0 +1,247 @@
+// Copyright Epic Games, Inc. All Rights Reserved.
+
+#include "DataflowEngineSceneProxy.h"
+
+#include "Async/ParallelFor.h"
+#include "Dataflow/DataflowComponent.h"
+#include "Dataflow/DataflowEnginePlugin.h"
+#include "Dataflow/DataflowEngineSceneHitProxies.h"
+#include "DynamicMeshBuilder.h"
+#include "Engine/CollisionProfile.h"
+#include "GeometryCollection/Facades/CollectionRenderingFacade.h"
+#include "Materials/Material.h"
+#include "PrimitiveSceneProxy.h"
+#include "StaticMeshResources.h"
+
+FDataflowEngineSceneProxy::FDataflowEngineSceneProxy(UDataflowComponent* Component)
+	: FPrimitiveSceneProxy(Component)
+	, MaterialRelevance(Component->GetMaterialRelevance(GetScene().GetFeatureLevel()))
+	, VertexFactory(GetScene().GetFeatureLevel(), "FTriangleSetSceneProxy")
+	, DefaultHitProxy(new HDataflowActor(Component->GetOwner(), Component, 3, 7))
+	, RenderMaterial(Component->GetMaterial(0))
+	, ConstantData(Component->GetRenderingCollection().NewCopy())
+{}
+
+FDataflowEngineSceneProxy::~FDataflowEngineSceneProxy() 
+{
+	if (ConstantData) delete ConstantData;
+}
+
+
+void FDataflowEngineSceneProxy::CreateRenderThreadResources()
+{
+	check(ConstantData);
+	check(IsInRenderingThread());
+
+	GeometryCollection::Facades::FRenderingFacade Facade(*ConstantData);
+	check(Facade.CanRenderSurface());
+
+	const int32 NumTriangleVertices = Facade.NumTriangles();
+	const int32 NumTriangleIndices = Facade.NumTriangles();
+	const int32 TotalNumVertices = NumTriangleVertices * 3;
+	const int32 TotalNumIndices = NumTriangleIndices * 3;
+	constexpr int32 NumTextureCoordinates = 1;
+
+	VertexBuffers.PositionVertexBuffer.Init(TotalNumVertices);
+	VertexBuffers.StaticMeshVertexBuffer.Init(TotalNumVertices, NumTextureCoordinates);
+	VertexBuffers.ColorVertexBuffer.Init(TotalNumVertices);
+	IndexBuffer.Indices.SetNumUninitialized(TotalNumIndices);
+#if WITH_EDITOR
+	HitProxyIdBuffer.Init(TotalNumVertices);
+#endif
+
+	// Initialize points.
+	// Triangles are represented as two tris, all of whose vertices are
+	// coincident. The material then offsets them according to the signs of the
+	// vertex normals in a camera facing orientation. Size of the point is given
+	// by U0.
+	if (Facade.NumTriangles() > 0)
+	{
+		MeshBatchDatas.Emplace();
+		FDataflowTriangleSetMeshBatchData& MeshBatchData = MeshBatchDatas.Last();
+		MeshBatchData.MinVertexIndex = 0;
+		MeshBatchData.MaxVertexIndex = NumTriangleVertices - 1;
+		MeshBatchData.StartIndex = 0;
+		MeshBatchData.NumPrimitives = Facade.NumTriangles();
+		if (!RenderMaterial)
+		{
+			RenderMaterial = UMaterial::GetDefaultMaterial(MD_Surface);
+		}
+		MeshBatchData.MaterialProxy = RenderMaterial->GetRenderProxy();
+
+		const TManagedArray<FIntVector>& Indices = Facade.GetIndices();
+		const TManagedArray<FVector3f>& Vertex = Facade.GetVertices();
+
+		// The color stored in the vertices actually gets interpreted as a linear
+		// color by the material, whereas it is more convenient for the user of the
+		// TriangleSet to specify colors as sRGB. So we actually have to convert it
+		// back to linear. The ToFColor(false) call just scales back into 0-255
+		// space.
+		ParallelFor(Facade.NumTriangles(), [&](int32 i)
+		{
+		const int32 VertexBufferIndex = 3 * i;
+		const int32 IndexBufferIndex = 3 * i;
+
+		const auto& P1 = Vertex[Indices[i][0]];
+		const auto& P2 = Vertex[Indices[i][1]];
+		const auto& P3 = Vertex[Indices[i][2]];
+
+		VertexBuffers.PositionVertexBuffer.VertexPosition(VertexBufferIndex + 0) = P1;
+		VertexBuffers.PositionVertexBuffer.VertexPosition(VertexBufferIndex + 1) = P2;
+		VertexBuffers.PositionVertexBuffer.VertexPosition(VertexBufferIndex + 2) = P3;
+
+		FVector3f Tangent1 = (P2 - P1).GetSafeNormal();
+		FVector3f Tangent2 = (P3 - P2).GetSafeNormal();
+		FVector3f Normal = (Tangent2 ^ Tangent1).GetSafeNormal();
+
+		VertexBuffers.StaticMeshVertexBuffer.SetVertexTangents(VertexBufferIndex + 0, FVector3f(1, 0, 0), FVector3f(0, 1, 0), Normal);
+		VertexBuffers.StaticMeshVertexBuffer.SetVertexTangents(VertexBufferIndex + 1, FVector3f(1, 0, 0), FVector3f(0, 1, 0), Normal);
+		VertexBuffers.StaticMeshVertexBuffer.SetVertexTangents(VertexBufferIndex + 2, FVector3f(1, 0, 0), FVector3f(0, 1, 0), Normal);
+
+		VertexBuffers.StaticMeshVertexBuffer.SetVertexUV(VertexBufferIndex + 0, 0, FVector2f(0, 0));
+		VertexBuffers.StaticMeshVertexBuffer.SetVertexUV(VertexBufferIndex + 1, 0, FVector2f(0, 0));
+		VertexBuffers.StaticMeshVertexBuffer.SetVertexUV(VertexBufferIndex + 2, 0, FVector2f(0, 0));
+
+		FColor FaceColor = IsSelected() ? IDataflowEnginePlugin::PrimarySelectionColor : IDataflowEnginePlugin::SurfaceColor;
+		VertexBuffers.ColorVertexBuffer.VertexColor(VertexBufferIndex + 0) = FaceColor;
+		VertexBuffers.ColorVertexBuffer.VertexColor(VertexBufferIndex + 1) = FaceColor;
+		VertexBuffers.ColorVertexBuffer.VertexColor(VertexBufferIndex + 2) = FaceColor;
+
+		IndexBuffer.Indices[IndexBufferIndex + 0] = VertexBufferIndex + 0;
+		IndexBuffer.Indices[IndexBufferIndex + 1] = VertexBufferIndex + 1;
+		IndexBuffer.Indices[IndexBufferIndex + 2] = VertexBufferIndex + 2;
+
+#if WITH_EDITOR
+		// @todo(Dataflow) Question : Does this even matter? 
+		// FMeshBatch::BatchHitProxyId has this also and 
+		// when GetCustomHitProxyIdBuffer returns non-null
+		// the hit buffer seems to be ignored in the render.
+		HitProxyIdBuffer.VertexColor(i) = DefaultHitProxy->Id.GetColor();
+#endif
+		});
+	}
+#if WITH_EDITOR
+	SetUsedMaterialForVerification({RenderMaterial});
+#endif
+
+	VertexBuffers.PositionVertexBuffer.InitResource();
+	VertexBuffers.StaticMeshVertexBuffer.InitResource();
+	VertexBuffers.ColorVertexBuffer.InitResource();
+
+	FLocalVertexFactory::FDataType Data;
+	VertexBuffers.PositionVertexBuffer.BindPositionVertexBuffer(&VertexFactory, Data);
+	VertexBuffers.StaticMeshVertexBuffer.BindTangentVertexBuffer(&VertexFactory, Data);
+	VertexBuffers.StaticMeshVertexBuffer.BindTexCoordVertexBuffer(&VertexFactory, Data);
+	VertexBuffers.ColorVertexBuffer.BindColorVertexBuffer(&VertexFactory, Data);
+	VertexFactory.SetData(Data);
+
+	VertexFactory.InitResource();
+	IndexBuffer.InitResource();
+#if WITH_EDITOR
+	HitProxyIdBuffer.InitResource();
+#endif
+}
+
+void FDataflowEngineSceneProxy::DestroyRenderThreadResources()
+{
+	check(IsInRenderingThread());
+	VertexBuffers.PositionVertexBuffer.ReleaseResource();
+	VertexBuffers.StaticMeshVertexBuffer.ReleaseResource();
+	VertexBuffers.ColorVertexBuffer.ReleaseResource();
+	IndexBuffer.ReleaseResource();
+	VertexFactory.ReleaseResource();
+#if WITH_EDITOR
+	HitProxyIdBuffer.ReleaseResource();
+#endif
+	delete ConstantData;
+	ConstantData = nullptr;
+}
+
+
+#if WITH_EDITOR
+HHitProxy* FDataflowEngineSceneProxy::CreateHitProxies(UPrimitiveComponent* Component, TArray<TRefCountPtr<HHitProxy> >& OutHitProxies)
+{
+	check(!IsInRenderingThread());
+	OutHitProxies.Add(DefaultHitProxy);
+	return FPrimitiveSceneProxy::CreateHitProxies(Component, OutHitProxies);
+}
+
+
+const FColorVertexBuffer* FDataflowEngineSceneProxy::GetCustomHitProxyIdBuffer() const
+{
+	// returning my own HitProxyIdBuffer causes the Hit to fail. 
+	return nullptr;// &HitProxyIdBuffer;
+}
+#endif // WITH_EDITOR
+
+
+void FDataflowEngineSceneProxy::GetDynamicMeshElements( 
+	const TArray<const FSceneView*>& Views, const FSceneViewFamily& ViewFamily,
+	uint32 VisibilityMap, FMeshElementCollector& Collector) const
+{
+	QUICK_SCOPE_CYCLE_COUNTER(STAT_OverlaySceneProxy_GetDynamicMeshElements);
+	check(IsInRenderingThread());
+
+	for (int32 ViewIndex = 0; ViewIndex < Views.Num(); ViewIndex++)
+	{
+		if (VisibilityMap & (1 << ViewIndex))
+		{
+			for (const FDataflowTriangleSetMeshBatchData& MeshBatchData : MeshBatchDatas)
+			{
+				FMeshBatch& Mesh = Collector.AllocateMesh();
+				FMeshBatchElement& BatchElement = Mesh.Elements[0];
+				BatchElement.IndexBuffer = &IndexBuffer;
+				Mesh.bWireframe = false;
+				Mesh.VertexFactory = &VertexFactory;
+				Mesh.MaterialRenderProxy = MeshBatchData.MaterialProxy;
+
+				FDynamicPrimitiveUniformBuffer& DynamicPrimitiveUniformBuffer = Collector.AllocateOneFrameResource<FDynamicPrimitiveUniformBuffer>();
+				DynamicPrimitiveUniformBuffer.Set(GetLocalToWorld(), GetLocalToWorld(),GetBounds(), GetLocalBounds(), true, false, AlwaysHasVelocity());
+				BatchElement.PrimitiveUniformBufferResource = &DynamicPrimitiveUniformBuffer.UniformBuffer;
+				//BatchElement.PrimitiveUniformBuffer = GetUniformBuffer();
+
+				BatchElement.FirstIndex = MeshBatchData.StartIndex;
+				BatchElement.NumPrimitives = MeshBatchData.NumPrimitives;
+				BatchElement.MinVertexIndex = MeshBatchData.MinVertexIndex;
+				BatchElement.MaxVertexIndex = MeshBatchData.MaxVertexIndex;
+				Mesh.ReverseCulling = IsLocalToWorldDeterminantNegative();
+				Mesh.Type = PT_TriangleList;
+				Mesh.DepthPriorityGroup = SDPG_World;
+				Mesh.bCanApplyViewModeOverrides = true;
+#if WITH_EDITOR
+				Mesh.BatchHitProxyId = DefaultHitProxy->Id;
+#endif // WITH_EDITOR
+
+				Collector.AddMesh(ViewIndex, Mesh);
+			}
+		}
+	}
+}
+
+
+FPrimitiveViewRelevance FDataflowEngineSceneProxy::GetViewRelevance(const FSceneView* View) const
+{
+	FPrimitiveViewRelevance Result;
+	Result.bDrawRelevance = IsShown(View);
+	Result.bShadowRelevance = IsShadowCast(View);
+	Result.bDynamicRelevance = true;
+	Result.bRenderInMainPass = ShouldRenderInMainPass();
+	Result.bUsesLightingChannels =GetLightingChannelMask() != GetDefaultLightingChannelMask();
+	Result.bRenderCustomDepth = false; // was ShouldRenderCustomDepth();
+	Result.bTranslucentSelfShadow = bCastVolumetricTranslucentShadow;
+	MaterialRelevance.SetPrimitiveViewRelevance(Result);
+	Result.bVelocityRelevance = false;//was DrawsVelocity() && Result.bOpaque&& Result.bRenderInMainPass;
+	return Result;
+}
+
+
+bool FDataflowEngineSceneProxy::CanBeOccluded() const
+{
+	return false;//was !MaterialRelevance.bDisableDepthTest;
+}
+
+SIZE_T FDataflowEngineSceneProxy::GetTypeHash() const
+{
+	static SIZE_T UniquePointer;
+	return reinterpret_cast<SIZE_T>(&UniquePointer);
+}
diff --git a/Engine/Plugins/Experimental/Dataflow/Source/DataflowEnginePlugin/Private/Dataflow/DataflowEngineSceneProxy.h b/Engine/Plugins/Experimental/Dataflow/Source/DataflowEnginePlugin/Private/Dataflow/DataflowEngineSceneProxy.h
new file mode 100644
--- /dev/null
+++ b/Engine/Plugins/Experimental/Dataflow/Source/DataflowEnginePlugin/Private/Dataflow/DataflowEngineSceneProxy.h
@@ -0,0 +1,73 @@
+// Copyright Epic Games, Inc. All Rights Reserved.
+#pragma once
+
+#include "Async/ParallelFor.h"
+#include "Engine/CollisionProfile.h"
+#include "PrimitiveSceneProxy.h"
+#include "Materials/Material.h"
+#include "DynamicMeshBuilder.h"
+#include "StaticMeshResources.h"
+
+class UDataflowComponent;
+struct FManagedArrayCollection;
+
+struct FDataflowTriangleSetMeshBatchData
+{
+	FMaterialRenderProxy* MaterialProxy = nullptr;
+	int32 StartIndex = -1;
+	int32 NumPrimitives = -1;
+	int32 MinVertexIndex = -1;
+	int32 MaxVertexIndex = -1;
+};
+
+class FDataflowEngineSceneProxy final : public FPrimitiveSceneProxy
+{
+public:
+
+	FDataflowEngineSceneProxy(UDataflowComponent* Component);
+
+	virtual ~FDataflowEngineSceneProxy();
+
+
+	//~ FPrimitiveSceneProxy
+	virtual void GetDynamicMeshElements(const TArray<const FSceneView*>& Views, const FSceneViewFamily& ViewFamily, uint32 VisibilityMap, FMeshElementCollector& Collector) const override;
+	virtual FPrimitiveViewRelevance GetViewRelevance(const FSceneView* View) const override;
+	virtual bool CanBeOccluded() const override;
+	virtual uint32 GetMemoryFootprint() const override { return sizeof(*this) + GetAllocatedSize(); }
+	virtual SIZE_T GetTypeHash() const override;
+	virtual void CreateRenderThreadResources() override;
+	virtual void DestroyRenderThreadResources() override;
+
+#if WITH_EDITOR
+	virtual HHitProxy* CreateHitProxies(UPrimitiveComponent* Component, TArray<TRefCountPtr<HHitProxy> >& OutHitProxies) override;
+	virtual const FColorVertexBuffer* GetCustomHitProxyIdBuffer() const override;
+#endif // WITH_EDITOR
+
+
+protected:
+
+	/** Create the rendering buffer resources */
+	void InitResources();
+
+	/** Return the rendering buffer resources */
+	void ReleaseResources();
+
+private:
+	TArray<FDataflowTriangleSetMeshBatchData> MeshBatchDatas;
+	FMaterialRelevance MaterialRelevance;
+	FLocalVertexFactory VertexFactory;
+	FStaticMeshVertexBuffers VertexBuffers;
+	FDynamicMeshIndexBuffer32 IndexBuffer;
+#if WITH_EDITOR
+	FColorVertexBuffer HitProxyIdBuffer;
+#endif // WITH_EDITOR
+
+	TRefCountPtr<HHitProxy> DefaultHitProxy;
+
+	// Render thread copy of data. 
+	UMaterialInterface* RenderMaterial = nullptr;
+	FManagedArrayCollection* ConstantData = nullptr;
+};
+
+
+
diff --git a/Engine/Plugins/Experimental/Dataflow/Source/DataflowEnginePlugin/Private/Dataflow/DataflowRenderingActor.cpp b/Engine/Plugins/Experimental/Dataflow/Source/DataflowEnginePlugin/Private/Dataflow/DataflowRenderingActor.cpp
deleted file mode 100644
diff --git a/Engine/Plugins/Experimental/Dataflow/Source/DataflowEnginePlugin/Private/Dataflow/DataflowRenderingComponent.cpp b/Engine/Plugins/Experimental/Dataflow/Source/DataflowEnginePlugin/Private/Dataflow/DataflowRenderingComponent.cpp
deleted file mode 100644
diff --git a/Engine/Plugins/Experimental/Dataflow/Source/DataflowEnginePlugin/Public/Dataflow/DataflowActor.h b/Engine/Plugins/Experimental/Dataflow/Source/DataflowEnginePlugin/Public/Dataflow/DataflowActor.h
new file mode 100644
--- /dev/null
+++ b/Engine/Plugins/Experimental/Dataflow/Source/DataflowEnginePlugin/Public/Dataflow/DataflowActor.h
@@ -0,0 +1,30 @@
+// Copyright Epic Games, Inc. All Rights Reserved.
+
+#pragma once
+
+#include "CoreMinimal.h"
+#include "GameFramework/Actor.h"
+#include "UObject/ObjectMacros.h"
+#include "UObject/ObjectPtr.h"
+#include "Dataflow/DataflowComponent.h"
+
+#include "DataflowActor.generated.h"
+
+
+UCLASS()
+class DATAFLOWENGINEPLUGIN_API ADataflowActor : public AActor
+{
+	GENERATED_UCLASS_BODY()
+
+public:
+
+	/* DataflowComponent */
+	UPROPERTY(VisibleAnywhere, BlueprintReadOnly, Category = Destruction, meta = (ExposeFunctionCategories = "Components|Dataflow", AllowPrivateAccess = "true"))
+	TObjectPtr<UDataflowComponent> DataflowComponent;
+	UDataflowComponent* GetDataflowComponent() const { return DataflowComponent; }
+
+
+#if WITH_EDITOR
+	virtual bool GetReferencedContentObjects(TArray<UObject*>& Objects) const override;
+#endif
+};
diff --git a/Engine/Plugins/Experimental/Dataflow/Source/DataflowEnginePlugin/Public/Dataflow/DataflowComponent.h b/Engine/Plugins/Experimental/Dataflow/Source/DataflowEnginePlugin/Public/Dataflow/DataflowComponent.h
new file mode 100644
--- /dev/null
+++ b/Engine/Plugins/Experimental/Dataflow/Source/DataflowEnginePlugin/Public/Dataflow/DataflowComponent.h
@@ -0,0 +1,73 @@
+// Copyright Epic Games, Inc. All Rights Reserved.
+#pragma once
+
+#include "GeometryCollection/ManagedArrayCollection.h"
+#include "Components/MeshComponent.h"
+#include "Dataflow/DataflowObject.h"
+#include "UObject/ObjectMacros.h"
+#include "UObject/UObjectGlobals.h"
+
+#include "DataflowComponent.generated.h"
+
+/**
+*	UDataflowComponent
+*/
+UCLASS(meta = (BlueprintSpawnableComponent))
+class DATAFLOWENGINEPLUGIN_API UDataflowComponent : public UPrimitiveComponent
+{
+	GENERATED_UCLASS_BODY()
+
+public:
+
+	virtual ~UDataflowComponent();
+
+	void Invalidate();
+	void UpdateLocalBounds();
+
+
+	//~ USceneComponent Interface.
+	virtual void TickComponent(float DeltaTime, enum ELevelTick TickType, FActorComponentTickFunction* ThisTickFunction) override;
+	virtual FBoxSphereBounds CalcBounds(const FTransform& LocalToWorld) const override;
+
+	//~ UPrimitiveComponent Interface.
+	virtual FPrimitiveSceneProxy* CreateSceneProxy() override;
+	virtual UMaterialInterface* GetMaterial(int32 Index) const override;
+
+	/* Rendering Support*/
+	virtual FMaterialRelevance GetMaterialRelevance(ERHIFeatureLevel::Type InFeatureLevel) const;
+	UMaterialInterface* GetDefaultMaterial() const;
+
+	// temp code for selection...
+	bool bSelected = false;
+	virtual bool ShouldRenderSelected() const override { return Super::ShouldRenderSelected() || bSelected; }
+
+	/** Render Targets*/
+	void ResetRenderTargets(); 
+	void AddRenderTarget(const UDataflowEdNode* InTarget); 
+	const TArray<const UDataflowEdNode*>& GetRenderTargets() const {return RenderTargets;}
+
+	/** Context */
+	void SetContext(TSharedPtr<Dataflow::FContext> InContext) { Context = InContext; }
+
+	/** RenderCollection */
+	void SetRenderingCollection(FManagedArrayCollection&& InCollection);
+	const FManagedArrayCollection& GetRenderingCollection() const;
+
+
+	/** Dataflow */
+	void SetDataflow(const UDataflow* InDataflow) { Dataflow = InDataflow; }
+	const UDataflow* GetDataflow() const { return Dataflow; }
+
+
+
+private:
+	TSharedPtr<Dataflow::FContext> Context;
+	TArray<const UDataflowEdNode*> RenderTargets;
+	TObjectPtr< const UDataflow> Dataflow;
+	FManagedArrayCollection RenderCollection;
+
+	bool bUpdateRender = true;
+	bool bBoundsNeedsUpdate = true;
+	FBoxSphereBounds BoundingBox = FBoxSphereBounds(ForceInitToZero);
+};
+
diff --git a/Engine/Plugins/Experimental/Dataflow/Source/DataflowEnginePlugin/Public/Dataflow/DataflowEnginePlugin.h b/Engine/Plugins/Experimental/Dataflow/Source/DataflowEnginePlugin/Public/Dataflow/DataflowEnginePlugin.h
--- a/Engine/Plugins/Experimental/Dataflow/Source/DataflowEnginePlugin/Public/Dataflow/DataflowEnginePlugin.h
+++ b/Engine/Plugins/Experimental/Dataflow/Source/DataflowEnginePlugin/Public/Dataflow/DataflowEnginePlugin.h
@@ -15,6 +15,10 @@ class IDataflowEnginePlugin : public IModuleInterface
 
 public:
 
+	static FColor SurfaceColor;
+	static FColor PrimarySelectionColor;
+
+
 	/**
 	 * Singleton-like access to this module's interface.  This is just for convenience!
 	 * Beware of calling this during the shutdown phase, though.  Your module might have been unloaded already.
diff --git a/Engine/Plugins/Experimental/Dataflow/Source/DataflowEnginePlugin/Public/Dataflow/DataflowEngineSceneHitProxies.h b/Engine/Plugins/Experimental/Dataflow/Source/DataflowEnginePlugin/Public/Dataflow/DataflowEngineSceneHitProxies.h
new file mode 100644
--- /dev/null
+++ b/Engine/Plugins/Experimental/Dataflow/Source/DataflowEnginePlugin/Public/Dataflow/DataflowEngineSceneHitProxies.h
@@ -0,0 +1,30 @@
+// Copyright Epic Games, Inc. All Rights Reserved.
+#pragma once
+
+#include "CoreMinimal.h"
+#include "UObject/UObjectGlobals.h"
+#include "Components/PrimitiveComponent.h"
+
+#include "EngineUtils.h"
+
+class UDataflowComponent;
+class UPrimitiveComponent;
+
+/** HitProxy with for dataflow actor.
+ */
+struct HDataflowActor : public HActor
+{
+	DECLARE_HIT_PROXY(DATAFLOWENGINEPLUGIN_API)
+
+		HDataflowActor(AActor* InActor, const UPrimitiveComponent* InPrimitiveComponent, int32 InSectionIndex, int32 InMaterialIndex)
+		: HActor(InActor, InPrimitiveComponent)
+	{
+		SectionIndex= InSectionIndex;
+		MaterialIndex = InMaterialIndex;
+	}
+
+	virtual EMouseCursor::Type GetMouseCursor() override
+	{
+		return EMouseCursor::GrabHandClosed;
+	}
+};
\ No newline at end of file
diff --git a/Engine/Plugins/Experimental/Dataflow/Source/DataflowEnginePlugin/Public/Dataflow/DataflowRenderingActor.h b/Engine/Plugins/Experimental/Dataflow/Source/DataflowEnginePlugin/Public/Dataflow/DataflowRenderingActor.h
deleted file mode 100644
diff --git a/Engine/Plugins/Experimental/Dataflow/Source/DataflowEnginePlugin/Public/Dataflow/DataflowRenderingComponent.h b/Engine/Plugins/Experimental/Dataflow/Source/DataflowEnginePlugin/Public/Dataflow/DataflowRenderingComponent.h
deleted file mode 100644
diff --git a/Engine/Source/Runtime/Experimental/Chaos/Private/GeometryCollection/Facades/CollectionRenderingFacade.cpp b/Engine/Source/Runtime/Experimental/Chaos/Private/GeometryCollection/Facades/CollectionRenderingFacade.cpp
--- a/Engine/Source/Runtime/Experimental/Chaos/Private/GeometryCollection/Facades/CollectionRenderingFacade.cpp
+++ b/Engine/Source/Runtime/Experimental/Chaos/Private/GeometryCollection/Facades/CollectionRenderingFacade.cpp
@@ -52,6 +52,17 @@ namespace GeometryCollection::Facades
 			MaterialIDAttribute.IsValid() && TriangleSectionAttribute.IsValid();
 	}
 
+	int32 FRenderingFacade::NumTriangles() const
+	{
+		if (IsValid())
+		{
+			return GetIndices().Num();
+		}
+			 
+		return 0;
+	}
+
+
 	void FRenderingFacade::AddTriangle(const Chaos::FTriangle& InTriangle)
 	{
 		check(!IsConst());
diff --git a/Engine/Source/Runtime/Experimental/Chaos/Public/GeometryCollection/Facades/CollectionRenderingFacade.h b/Engine/Source/Runtime/Experimental/Chaos/Public/GeometryCollection/Facades/CollectionRenderingFacade.h
--- a/Engine/Source/Runtime/Experimental/Chaos/Public/GeometryCollection/Facades/CollectionRenderingFacade.h
+++ b/Engine/Source/Runtime/Experimental/Chaos/Public/GeometryCollection/Facades/CollectionRenderingFacade.h
@@ -42,6 +42,13 @@ namespace GeometryCollection::Facades
 		/**Does it support rendering surfaces.*/
 		bool CanRenderSurface() const;
 
+		//
+		// Facade API
+		//
+
+		/**Add a triangle to the rendering view.*/
+		int32 NumTriangles() const;
+
 		/**Add a triangle to the rendering view.*/
 		void AddTriangle(const Chaos::FTriangle& InTriangle);
 
diff --git a/Engine/Source/Runtime/Experimental/Dataflow/Engine/Public/Dataflow/DataflowObject.h b/Engine/Source/Runtime/Experimental/Dataflow/Engine/Public/Dataflow/DataflowObject.h
--- a/Engine/Source/Runtime/Experimental/Dataflow/Engine/Public/Dataflow/DataflowObject.h
+++ b/Engine/Source/Runtime/Experimental/Dataflow/Engine/Public/Dataflow/DataflowObject.h
@@ -18,6 +18,7 @@ class FArchive;
 class UDataflow;
 class UObject;
 class UDataflowEdNode;
+class UMaterial;
 namespace Dataflow { class FGraph; }
 
 
@@ -79,6 +80,10 @@ public:
 	UPROPERTY(EditAnywhere, Category = "Evaluation", AdvancedDisplay )
 	TArray<TObjectPtr<UObject>> Targets;
 
+	UPROPERTY(EditAnywhere, Category = "Render")
+	TObjectPtr<UMaterial> Material = nullptr;
+
+
 public:
 	/** UObject Interface */
 #if WITH_EDITOR
