Dataflow : Viewport client object selection toolbar

diff --git a/Engine/Plugins/Experimental/Dataflow/Source/DataflowEditor/Private/Dataflow/DataflowEditorActions.cpp b/Engine/Plugins/Experimental/Dataflow/Source/DataflowEditor/Private/Dataflow/DataflowEditorActions.cpp
deleted file mode 100644
diff --git a/Engine/Plugins/Experimental/Dataflow/Source/DataflowEditor/Private/Dataflow/DataflowEditorCommands.cpp b/Engine/Plugins/Experimental/Dataflow/Source/DataflowEditor/Private/Dataflow/DataflowEditorCommands.cpp
new file mode 100644
--- /dev/null
+++ b/Engine/Plugins/Experimental/Dataflow/Source/DataflowEditor/Private/Dataflow/DataflowEditorCommands.cpp
@@ -0,0 +1,322 @@
+// Copyright Epic Games, Inc. All Rights Reserved.
+
+#include "Dataflow/DataflowEditorCommands.h"
+
+#include "Dataflow/DataflowEdNode.h"
+#include "Dataflow/DataflowNodeFactory.h"
+#include "Dataflow/DataflowObject.h"
+#include "EdGraph/EdGraphNode.h"
+#include "IStructureDetailsView.h"
+#include "EdGraphNode_Comment.h"
+
+#define LOCTEXT_NAMESPACE "DataflowEditorCommands"
+
+void FDataflowEditorCommandsImpl::RegisterCommands()
+{
+	UI_COMMAND(EvaluateNode, "Evaluate", "Trigger an evaluation of the selected node.", EUserInterfaceActionType::Button, FInputChord());
+	UI_COMMAND(CreateComment, "CreateComment", "Create a Comment node.", EUserInterfaceActionType::None, FInputChord());
+	UI_COMMAND(ToggleEnabledState, "ToggleEnabledState", "Toggle node between Enabled/Disabled state.", EUserInterfaceActionType::Button, FInputChord());
+	UI_COMMAND(ToggleObjectSelection, "ToggleObjectSelection", "Enable object selection in editor.", EUserInterfaceActionType::ToggleButton, FInputChord());
+
+	if (Dataflow::FNodeFactory* Factory = Dataflow::FNodeFactory::GetInstance())
+	{
+		for (Dataflow::FFactoryParameters Parameters : Factory->RegisteredParameters())
+		{
+			TSharedPtr< FUICommandInfo > AddNode;
+			FUICommandInfo::MakeCommandInfo(
+				this->AsShared(),
+				AddNode,
+				Parameters.TypeName,
+				LOCTEXT("DataflowButton", "New Dataflow Node"),
+				LOCTEXT("NewDataflowNodeTooltip", "New Dataflow Node Tooltip"),
+				FSlateIcon(),
+				EUserInterfaceActionType::Button,
+				FInputChord()
+			);
+			CreateNodesMap.Add(Parameters.TypeName, AddNode);
+		}
+	}
+
+		
+}
+
+const FDataflowEditorCommandsImpl& FDataflowEditorCommands::Get()
+{
+	return FDataflowEditorCommandsImpl::Get();
+}
+
+void FDataflowEditorCommands::Register()
+{
+	return FDataflowEditorCommandsImpl::Register();
+}
+
+void FDataflowEditorCommands::Unregister()
+{
+	return FDataflowEditorCommandsImpl::Unregister();
+}
+
+
+
+void FDataflowEditorCommands::EvaluateSelectedNodes(const FGraphPanelSelectionSet& SelectedNodes, FDataflowEditorCommands::FGraphEvaluationCallback Evaluate)
+{
+	for (UObject* Ode : SelectedNodes)
+	{
+		if (UDataflowEdNode* EdNode = dynamic_cast<UDataflowEdNode*>(Ode))
+		{
+			if (const TSharedPtr<Dataflow::FGraph> DataflowGraph = EdNode->GetDataflowGraph())
+			{
+				if (const TSharedPtr<FDataflowNode> DataflowNode = DataflowGraph->FindBaseNode(EdNode->GetDataflowNodeGuid()))
+				{
+					if (DataflowNode->bActive)
+					{
+						if (DataflowNode->GetOutputs().Num())
+						{
+							for (FDataflowConnection* NodeOutput : DataflowNode->GetOutputs())
+							{
+								Evaluate(DataflowNode.Get(), (FDataflowOutput*)NodeOutput);
+							}
+						}
+						else
+						{
+							Evaluate(DataflowNode.Get(), nullptr);
+						}
+					}
+				}
+			}
+		}
+	}
+}
+
+void FDataflowEditorCommands::EvaluateNode(Dataflow::FContext& Context, Dataflow::FTimestamp& OutLastNodeTimestamp,
+	const UDataflow* Dataflow, const FDataflowNode* InNode, const FDataflowOutput* Output, FString NodeName)
+{
+	if (Dataflow)
+	{
+		const FDataflowNode* Node = InNode;
+		if (Node == nullptr)
+		{
+			if (const TSharedPtr<Dataflow::FGraph> Graph = Dataflow->GetDataflow())
+			{
+				if (TSharedPtr<const FDataflowNode> GraphNode = Graph->FindBaseNode(FName(NodeName)))
+				{
+					Node = GraphNode.Get();
+				}
+			}
+		}
+
+		if (Node != nullptr)
+		{
+			if (Output == nullptr)
+			{
+				if (Node->GetTimestamp() >= OutLastNodeTimestamp)
+				{
+					Context.Evaluate(Node, nullptr);
+					OutLastNodeTimestamp = Context.GetTimestamp();
+				}
+			}
+			else // Output != nullptr
+			{
+				if (!Context.HasData(Output->CacheKey(), Context.GetTimestamp()))
+				{
+					Context.Evaluate(Node, Output);
+				}
+			}
+		}
+	}
+}
+
+void FDataflowEditorCommands::EvaluateTerminalNode(Dataflow::FContext& Context, Dataflow::FTimestamp& OutLastNodeTimestamp,
+	const UDataflow* Dataflow, const FDataflowNode* InNode, const FDataflowOutput* Output, UObject* InAsset, FString NodeName)
+{
+	if (Dataflow)
+	{
+		const FDataflowNode* Node = InNode;
+		if (Node == nullptr)
+		{
+			if (const TSharedPtr<Dataflow::FGraph> Graph = Dataflow->GetDataflow())
+			{
+				if (TSharedPtr<const FDataflowNode> GraphNode = Graph->FindBaseNode(FName(NodeName)))
+				{
+					Node = GraphNode.Get();
+				}
+			}
+		}
+
+		if (Node != nullptr)
+		{
+			if (Output == nullptr)
+			{
+				if (Node->GetTimestamp() >= OutLastNodeTimestamp)
+				{
+					Context.Evaluate(Node, nullptr);
+					OutLastNodeTimestamp = Context.GetTimestamp();
+
+					if (const FDataflowTerminalNode* TerminalNode = Node->AsType<const FDataflowTerminalNode>())
+					{
+						if (InAsset)
+						{
+							TerminalNode->SetAssetValue(InAsset, Context);
+						}
+					}
+				}
+			}
+			else // Output != nullptr
+			{
+				if (!Context.HasData(Output->CacheKey(), Context.GetTimestamp()))
+				{
+					Context.Evaluate(Node, Output);
+				}
+			}
+		}
+	}
+}
+
+
+bool FDataflowEditorCommands::OnNodeVerifyTitleCommit(const FText& NewText, UEdGraphNode* GraphNode, FText& OutErrorMessage)
+{
+	if (GraphNode)
+	{
+		if (UDataflowEdNode* DataflowNode = Cast<UDataflowEdNode>(GraphNode))
+		{
+			if (TSharedPtr<Dataflow::FGraph> Graph = DataflowNode->GetDataflowGraph())
+			{
+				if( Graph->FindBaseNode(FName(NewText.ToString())).Get()==nullptr )
+				{
+					return true;
+				}
+			}
+		}
+		else if( Cast<UEdGraphNode_Comment>(GraphNode))
+		{
+			return true;
+		}
+	}
+	OutErrorMessage = FText::FromString(FString::Printf(TEXT("Non-unique name for graph node (%s)"), *NewText.ToString()));
+	return false;
+}
+
+
+void FDataflowEditorCommands::OnNodeTitleCommitted(const FText& InNewText, ETextCommit::Type InCommitType, UEdGraphNode* GraphNode)
+{
+	if (InCommitType == ETextCommit::OnCleared)
+	{
+		return;
+	}
+
+	if (GraphNode)
+	{
+		if (UDataflowEdNode* DataflowNode = Cast<UDataflowEdNode>(GraphNode))
+		{
+			if (TSharedPtr<Dataflow::FGraph> Graph = DataflowNode->GetDataflowGraph())
+			{
+				if (TSharedPtr<FDataflowNode> Node = Graph->FindBaseNode(DataflowNode->GetDataflowNodeGuid()))
+				{
+					GraphNode->Rename(*InNewText.ToString());
+					Node->SetName(FName(InNewText.ToString()));
+				}
+			}
+		}
+		else if (UEdGraphNode_Comment* CommentNode = Cast<UEdGraphNode_Comment>(GraphNode))
+		{
+			GraphNode->NodeComment = InNewText.ToString();
+		}
+	}
+}
+
+
+void FDataflowEditorCommands::OnPropertyValueChanged(UDataflow* OutDataflow, TSharedPtr<Dataflow::FEngineContext>& Context, Dataflow::FTimestamp& OutLastNodeTimestamp, const FPropertyChangedEvent& InPropertyChangedEvent)
+{
+	if (InPropertyChangedEvent.ChangeType == EPropertyChangeType::ValueSet)
+	{
+		TSharedPtr<const FDataflowNode> UpdatedNode = nullptr;
+		if (OutDataflow && InPropertyChangedEvent.Property && InPropertyChangedEvent.Property->GetOwnerUObject())
+		{
+			OutDataflow->MarkPackageDirty();
+
+			FString Name = InPropertyChangedEvent.Property->GetOwnerUObject()->GetName();
+			Name = Name.Left(Name.Len() - FString("DataflowNode").Len());
+			{
+				FDataflowAssetEdit Edit = OutDataflow->EditDataflow();
+
+				Dataflow::FGraph* Graph = Edit.GetGraph();
+				for (auto Node : Graph->GetNodes())
+				{
+					FString NodeName = Node->GetName().ToString();
+					if (NodeName.StartsWith(Name))
+					{
+						UpdatedNode = Node;
+						Node->Invalidate();
+					}
+				}
+			}
+		}
+
+		if (!UpdatedNode && Context)
+		{
+			// Some base properties dont link back to the parent, so just clobber the cache for now. 
+			Context.Reset();
+		}
+		OutLastNodeTimestamp = Dataflow::FTimestamp::Invalid;
+	}
+}
+
+
+void FDataflowEditorCommands::DeleteNodes(UDataflow* Graph, const FGraphPanelSelectionSet& SelectedNodes)
+{
+	if (ensureMsgf(Graph != nullptr, TEXT("Warning : Failed to find valid graph.")))
+	{
+		for (UObject* Ode : SelectedNodes)
+		{
+			if (UDataflowEdNode* EdNode = dynamic_cast<UDataflowEdNode*>(Ode))
+			{
+				if (const TSharedPtr<Dataflow::FGraph> DataflowGraph = EdNode->GetDataflowGraph())
+				{
+					Graph->RemoveNode(EdNode);
+					if (TSharedPtr<FDataflowNode> DataflowNode = DataflowGraph->FindBaseNode(EdNode->GetDataflowNodeGuid()))
+					{
+						DataflowGraph->RemoveNode(DataflowNode);
+					}
+				}
+			}
+			else if (UEdGraphNode_Comment* CommentNode = dynamic_cast<UEdGraphNode_Comment*>(Ode))
+			{
+				Graph->RemoveNode(CommentNode);
+			}
+		}
+	}
+}
+
+void FDataflowEditorCommands::OnSelectedNodesChanged(TSharedPtr<IStructureDetailsView> PropertiesEditor, UObject* Asset, UDataflow* Graph, const TSet<UObject*>& NewSelection)
+{
+	PropertiesEditor->SetStructureData(nullptr);
+
+	if (Graph && PropertiesEditor)
+	{
+		if (const TSharedPtr<Dataflow::FGraph> DataflowGraph = Graph->GetDataflow())
+		{
+			FGraphPanelSelectionSet SelectedNodes = NewSelection;
+			if (SelectedNodes.Num())
+			{
+				TArray<UObject*> Objects;
+				for (UObject* SelectedObject : SelectedNodes)
+				{
+					if (UDataflowEdNode* EdNode = Cast<UDataflowEdNode>(SelectedObject))
+					{
+						if (TSharedPtr<FDataflowNode> DataflowNode = DataflowGraph->FindBaseNode(EdNode->GetDataflowNodeGuid()))
+						{
+							TSharedPtr<FStructOnScope> Struct(DataflowNode->NewStructOnScope());
+							PropertiesEditor->SetStructureData(Struct);
+						}
+					}
+				}
+				
+			}
+		}
+	}
+}
+
+void FDataflowEditorCommands::ToggleEnabledState(UDataflow* Graph)
+{
+}
+
+#undef LOCTEXT_NAMESPACE
diff --git a/Engine/Plugins/Experimental/Dataflow/Source/DataflowEditor/Private/Dataflow/DataflowEditorStyle.h b/Engine/Plugins/Experimental/Dataflow/Source/DataflowEditor/Private/Dataflow/DataflowEditorStyle.h
--- a/Engine/Plugins/Experimental/Dataflow/Source/DataflowEditor/Private/Dataflow/DataflowEditorStyle.h
+++ b/Engine/Plugins/Experimental/Dataflow/Source/DataflowEditor/Private/Dataflow/DataflowEditorStyle.h
@@ -16,6 +16,7 @@ public:
 		const FVector2D Icon64x64(64.f, 64.f);
 		const FVector2D Icon56x28(56.f, 28.f);
 		const FVector2D Icon28x14(28.f, 14.f);
+		const FVector2D Icon40x40(28.f, 14.f);
 
 		SetContentRoot(IPluginManager::Get().FindPlugin("Dataflow")->GetBaseDir() / TEXT("Resources"));
 		Set("ClassIcon.Dataflow", new FSlateVectorImageBrush(RootToContentDir(TEXT("DataflowAsset_16.svg")), Icon16x16));
@@ -28,7 +29,6 @@ public:
 		Set("Dataflow.Cached.False", new FSlateImageBrush(RootToContentDir(TEXT("Slate/status_grey.png")), Icon16x16));
 		Set("Dataflow.Cached.True", new FSlateImageBrush(RootToContentDir(TEXT("Slate/status_green.png")), Icon16x16));
 
-
 		FSlateStyleRegistry::RegisterSlateStyle(*this);
 	}
 
diff --git a/Engine/Plugins/Experimental/Dataflow/Source/DataflowEditor/Private/Dataflow/DataflowEditorToolkit.cpp b/Engine/Plugins/Experimental/Dataflow/Source/DataflowEditor/Private/Dataflow/DataflowEditorToolkit.cpp
--- a/Engine/Plugins/Experimental/Dataflow/Source/DataflowEditor/Private/Dataflow/DataflowEditorToolkit.cpp
+++ b/Engine/Plugins/Experimental/Dataflow/Source/DataflowEditor/Private/Dataflow/DataflowEditorToolkit.cpp
@@ -4,7 +4,7 @@
 
 #include "Animation/Skeleton.h"
 #include "Dataflow/DataflowCore.h"
-#include "Dataflow/DataflowEditorActions.h"
+#include "Dataflow/DataflowEditorCommands.h"
 #include "Dataflow/DataflowEditorViewport.h"
 #include "Dataflow/DataflowEdNode.h"
 #include "Dataflow/DataflowGraphEditor.h"
diff --git a/Engine/Plugins/Experimental/Dataflow/Source/DataflowEditor/Private/Dataflow/DataflowEditorViewport.cpp b/Engine/Plugins/Experimental/Dataflow/Source/DataflowEditor/Private/Dataflow/DataflowEditorViewport.cpp
--- a/Engine/Plugins/Experimental/Dataflow/Source/DataflowEditor/Private/Dataflow/DataflowEditorViewport.cpp
+++ b/Engine/Plugins/Experimental/Dataflow/Source/DataflowEditor/Private/Dataflow/DataflowEditorViewport.cpp
@@ -2,15 +2,20 @@
 #include "Dataflow/DataflowEditorViewport.h"
 
 #include "AdvancedPreviewScene.h"
+#include "Dataflow/DataflowEditorCommands.h"
 #include "Dataflow/DataflowEdNode.h"
 #include "Dataflow/DataflowObject.h"
 #include "Dataflow/DataflowActor.h"
 #include "Dataflow/DataflowComponent.h"
+#include "Dataflow/DataflowEditorViewportToolbar.h"
 #include "Dataflow/DataflowEngineSceneHitProxies.h"
 #include "DynamicMeshBuilder.h"
 #include "Dataflow/DataflowEditorToolkit.h"
 #include "EditorModes.h"
 #include "EditorViewportClient.h"
+#include "EditorViewportCommands.h"
+#include "Framework/Application/SlateApplication.h"
+#include "Framework/Commands/UIAction.h"
 #include "Materials/Material.h"
 #include "SCommonEditorViewportToolbarBase.h"
 #include "Settings/EditorStyleSettings.h"
@@ -37,8 +42,8 @@ void SDataflowEditorViewport::Construct(const FArguments& InArgs)
 	FBoxSphereBounds SphereBounds = FBoxSphereBounds(EForceInit::ForceInitToZero);
 	CustomDataflowActor = CastChecked<ADataflowActor>(PreviewScene->GetWorld()->SpawnActor(ADataflowActor::StaticClass()));
 
-	EditorViewportClient->SetDataflowActor(CustomDataflowActor);
-	EditorViewportClient->FocusViewportOnBox( SphereBounds.GetBox());
+	ViewportClient->SetDataflowActor(CustomDataflowActor);
+	ViewportClient->FocusViewportOnBox( SphereBounds.GetBox());
 }
 
 TSharedRef<SEditorViewport> SDataflowEditorViewport::GetViewportWidget()
@@ -46,6 +51,7 @@ TSharedRef<SEditorViewport> SDataflowEditorViewport::GetViewportWidget()
 	return SharedThis(this);
 }
 
+
 TSharedPtr<FExtender> SDataflowEditorViewport::GetExtenders() const
 {
 	TSharedPtr<FExtender> Result(MakeShareable(new FExtender));
@@ -63,13 +69,32 @@ void SDataflowEditorViewport::AddReferencedObjects(FReferenceCollector& Collecto
 
 TSharedRef<FEditorViewportClient> SDataflowEditorViewport::MakeEditorViewportClient()
 {
-	EditorViewportClient = MakeShareable(new FDataflowEditorViewportClient(PreviewScene.Get(), SharedThis(this), DataflowEditorToolkitPtr));
-	return EditorViewportClient.ToSharedRef();
+	ViewportClient = MakeShareable(new FDataflowEditorViewportClient(PreviewScene.Get(), SharedThis(this), DataflowEditorToolkitPtr));
+	return ViewportClient.ToSharedRef();
 }
 
 TSharedPtr<SWidget> SDataflowEditorViewport::MakeViewportToolbar()
 {
-	return SNew(SCommonEditorViewportToolbarBase, SharedThis(this));
+	return
+		SNew(SDataflowViewportSelectionToolBar)
+		.EditorViewport(SharedThis(this))
+		.IsEnabled(FSlateApplication::Get().GetNormalExecutionAttribute());
+}
+
+void SDataflowEditorViewport::BindCommands()
+{
+	SEditorViewport::BindCommands();
+	{
+		const FDataflowEditorCommandsImpl& Commands = FDataflowEditorCommands::Get();
+		TSharedRef<FDataflowEditorViewportClient> ClientRef = ViewportClient.ToSharedRef();
+
+		CommandList->MapAction(
+			Commands.ToggleObjectSelection,
+			FExecuteAction::CreateSP(ClientRef, &FDataflowEditorViewportClient::SetSelectionMode, FDataflowSelectionState::EMode::DSS_Dataflow_Object)
+			,FCanExecuteAction::CreateSP(ClientRef, &FDataflowEditorViewportClient::CanSetSelectionMode, FDataflowSelectionState::EMode::DSS_Dataflow_Object)
+			,FIsActionChecked::CreateSP(ClientRef, &FDataflowEditorViewportClient::IsSelectionModeActive, FDataflowSelectionState::EMode::DSS_Dataflow_Object)
+		);
+	}
 }
 
 // ----------------------------------------------------------------------------------
@@ -87,6 +112,50 @@ FDataflowEditorViewportClient::FDataflowEditorViewportClient(FPreviewScene* InPr
 	EngineShowFlags.Grid = false;
 }
 
+void FDataflowEditorViewportClient::SetSelectionMode(FDataflowSelectionState::EMode InState)
+{
+	if (SelectionMode == InState)
+	{
+		SelectionMode = FDataflowSelectionState::EMode::DSS_Dataflow_None;
+	}
+	else
+	{
+		SelectionMode = InState;
+	}
+
+	if (SelectionMode == FDataflowSelectionState::EMode::DSS_Dataflow_None && !SelectionState.IsEmpty())
+	{
+		SelectionState = FDataflowSelectionState();
+		SelectionState.UpdateSelection(DataflowActor->DataflowComponent);
+	}
+}
+bool FDataflowEditorViewportClient::CanSetSelectionMode(FDataflowSelectionState::EMode InState)
+{
+	TSharedPtr<FDataflowEditorToolkit> DataflowEditorToolkit = DataflowEditorToolkitPtr.Pin();
+	if (DataflowEditorToolkitPtr.IsValid())
+	{
+		if (const UDataflow* Dataflow = DataflowEditorToolkit->GetDataflow())
+		{
+			if (Dataflow->GetRenderTargets().Num())
+			{
+				return true;
+			}
+		}
+	}
+
+	if (!SelectionState.IsEmpty())
+	{
+		SelectionState = FDataflowSelectionState();
+		SelectionState.UpdateSelection(DataflowActor->DataflowComponent);
+	}
+
+	return false;
+}
+bool FDataflowEditorViewportClient::IsSelectionModeActive(FDataflowSelectionState::EMode InState) 
+{ 
+	return SelectionMode == InState;
+}
+
 Dataflow::FTimestamp FDataflowEditorViewportClient::LatestTimestamp(const UDataflow* Dataflow, const Dataflow::FContext* Context)
 {
 	if (Dataflow && Context)
@@ -99,12 +168,25 @@ Dataflow::FTimestamp FDataflowEditorViewportClient::LatestTimestamp(const UDataf
 void FDataflowEditorViewportClient::ProcessClick(FSceneView& View, HHitProxy* HitProxy, FKey Key, EInputEvent Event, uint32 HitX, uint32 HitY)
 {
 	Super::ProcessClick(View, HitProxy, Key, Event, HitX, HitY);
-	if (HitProxy && HitProxy	->IsA(HDataflowActor::StaticGetType()))
+	if (DataflowActor && DataflowActor->DataflowComponent)
 	{
-		if (DataflowActor && DataflowActor->DataflowComponent)
+		FDataflowSelectionState PreState = SelectionState;
+
+		if (SelectionMode == FDataflowSelectionState::EMode::DSS_Dataflow_Object)
+		{
+			if (HitProxy && HitProxy->IsA(HDataflowActor::StaticGetType()))
+			{
+					SelectionState.bComponentSelected = !SelectionState.bComponentSelected;
+			}
+			else
+			{
+				SelectionState.bComponentSelected = false;
+			}
+		}
+
+		if (PreState != SelectionState)
 		{
-			DataflowActor->DataflowComponent->bSelected = !DataflowActor->DataflowComponent->bSelected;
-			DataflowActor->DataflowComponent->Invalidate();
+			SelectionState.UpdateSelection(DataflowActor->DataflowComponent);
 		}
 	}
 }
@@ -113,6 +195,7 @@ void FDataflowEditorViewportClient::ProcessClick(FSceneView& View, HHitProxy* Hi
 void FDataflowEditorViewportClient::Draw(const FSceneView* View, FPrimitiveDrawInterface* PDI)
 {
 	Super::Draw(View, PDI);
+	
 	if (bEnableDataflowClientDrawing)
 	{
 		GeometryCollection::Facades::FRenderingFacade Facade(RenderCollection);
diff --git a/Engine/Plugins/Experimental/Dataflow/Source/DataflowEditor/Private/Dataflow/DataflowEditorViewportToolbar.cpp b/Engine/Plugins/Experimental/Dataflow/Source/DataflowEditor/Private/Dataflow/DataflowEditorViewportToolbar.cpp
new file mode 100644
--- /dev/null
+++ b/Engine/Plugins/Experimental/Dataflow/Source/DataflowEditor/Private/Dataflow/DataflowEditorViewportToolbar.cpp
@@ -0,0 +1,65 @@
+// Copyright Epic Games, Inc. All Rights Reserved.
+#include "Dataflow/DataflowEditorViewportToolbar.h"
+
+#include "Dataflow/DataflowEditorCommands.h"
+#include "Dataflow/DataflowEditorViewport.h"
+#include "EditorViewportCommands.h"
+#include "Widgets/Layout/SBorder.h"
+
+
+void SDataflowViewportSelectionToolBar::Construct(const FArguments& InArgs)
+{
+	EditorViewport = InArgs._EditorViewport;
+	static const FName DefaultForegroundName("DefaultForeground");
+
+	this->ChildSlot
+		[
+			SNew(SBorder)
+			.BorderImage(FAppStyle::GetBrush("NoBorder"))
+			.ForegroundColor(FAppStyle::GetSlateColor(DefaultForegroundName))
+			[
+			SNew(SHorizontalBox)
+			+ SHorizontalBox::Slot()
+				.AutoWidth()
+				.Padding(2.0f, 2.0f)
+				.HAlign(HAlign_Right)
+				[
+					MakeSelectionModeToolBar()
+				]
+			]
+		];
+
+	SViewportToolBar::Construct(SViewportToolBar::FArguments());
+}
+
+TSharedRef<SWidget> SDataflowViewportSelectionToolBar::MakeSelectionModeToolBar()
+{
+	FSlimHorizontalToolBarBuilder ToolbarBuilder(EditorViewport.Pin()->GetCommandList(), FMultiBoxCustomization::None);
+
+	const FName ToolBarStyle = TEXT("EditorViewportToolBar");
+	ToolbarBuilder.SetStyle(&FAppStyle::Get(), ToolBarStyle);
+	ToolbarBuilder.SetLabelVisibility(EVisibility::Collapsed);
+
+	ToolbarBuilder.SetIsFocusable(false);
+
+	ToolbarBuilder.BeginSection("Selection");
+	{
+		ToolbarBuilder.BeginBlockGroup();
+
+		ToolbarBuilder.AddToolBarButton(FDataflowEditorCommands::Get().ToggleObjectSelection,
+			NAME_None,
+			TAttribute<FText>(),
+			TAttribute<FText>(),
+			FSlateIcon(FAppStyle::GetAppStyleSetName(), TEXT("Icons.SelectInViewport")),
+			FName(TEXT("SelectObjectMode"))
+		);
+
+
+		ToolbarBuilder.EndBlockGroup();
+	}
+
+	ToolbarBuilder.EndSection();
+	ToolbarBuilder.AddSeparator();
+
+	return ToolbarBuilder.MakeWidget();
+}
\ No newline at end of file
diff --git a/Engine/Plugins/Experimental/Dataflow/Source/DataflowEditor/Private/Dataflow/DataflowSchema.cpp b/Engine/Plugins/Experimental/Dataflow/Source/DataflowEditor/Private/Dataflow/DataflowSchema.cpp
--- a/Engine/Plugins/Experimental/Dataflow/Source/DataflowEditor/Private/Dataflow/DataflowSchema.cpp
+++ b/Engine/Plugins/Experimental/Dataflow/Source/DataflowEditor/Private/Dataflow/DataflowSchema.cpp
@@ -3,7 +3,7 @@
 #include "Dataflow/DataflowSchema.h"
 
 #include "Dataflow/DataflowEdNode.h"
-#include "Dataflow/DataflowEditorActions.h"
+#include "Dataflow/DataflowEditorCommands.h"
 #include "Dataflow/DataflowEdNode.h"
 #include "Dataflow/DataflowSNode.h"
 #include "Dataflow/DataflowNodeFactory.h"
diff --git a/Engine/Plugins/Experimental/Dataflow/Source/DataflowEditor/Public/Dataflow/DataflowEditorActions.h b/Engine/Plugins/Experimental/Dataflow/Source/DataflowEditor/Public/Dataflow/DataflowEditorActions.h
deleted file mode 100644
diff --git a/Engine/Plugins/Experimental/Dataflow/Source/DataflowEditor/Public/Dataflow/DataflowEditorCommands.h b/Engine/Plugins/Experimental/Dataflow/Source/DataflowEditor/Public/Dataflow/DataflowEditorCommands.h
new file mode 100644
--- /dev/null
+++ b/Engine/Plugins/Experimental/Dataflow/Source/DataflowEditor/Public/Dataflow/DataflowEditorCommands.h
@@ -0,0 +1,107 @@
+// Copyright Epic Games, Inc. All Rights Reserved.
+
+#pragma once
+
+#include "CoreMinimal.h"
+#include "Dataflow/DataflowObjectInterface.h"
+#include "Framework/Commands/Commands.h"
+#include "Styling/AppStyle.h"
+
+class FDragDropEvent;
+struct FDataflowOutput;
+struct FGeometry;
+class IStructureDetailsView;
+class UDataflow;
+class UDataflowEdNode;
+struct FDataflowNode;
+class UEdGraphNode;
+
+typedef TSet<class UObject*> FGraphPanelSelectionSet;
+
+/*
+* FDataflowEditorCommandsImpl
+* 
+*/
+class FDataflowEditorCommandsImpl : public TCommands<FDataflowEditorCommandsImpl>
+{
+public:
+
+	FDataflowEditorCommandsImpl()
+		: TCommands<FDataflowEditorCommandsImpl>( TEXT("DataflowEditor"), NSLOCTEXT("Contexts", "DataflowEditor", "Scene Graph Editor"), NAME_None, FAppStyle::GetAppStyleSetName() )
+	{
+	}	
+
+	virtual ~FDataflowEditorCommandsImpl()
+	{
+	}
+
+	DATAFLOWEDITOR_API virtual void RegisterCommands() override;
+
+	TSharedPtr< FUICommandInfo > EvaluateNode;
+	TSharedPtr< FUICommandInfo > CreateComment;
+	TSharedPtr< FUICommandInfo > ToggleEnabledState;
+	TSharedPtr<FUICommandInfo> ToggleObjectSelection;
+
+	TMap< FName, TSharedPtr<FUICommandInfo> > CreateNodesMap;
+};
+
+class DATAFLOWEDITOR_API FDataflowEditorCommands
+{
+public:
+	typedef TFunction<void(FDataflowNode*, FDataflowOutput*)> FGraphEvaluationCallback;
+	typedef TFunction<void(const FGeometry& MyGeometry, const FDragDropEvent& DragDropEvent)> FOnDragDropEventCallback;
+
+	static void Register();
+	static void Unregister();
+
+	static const FDataflowEditorCommandsImpl& Get();
+
+	/*
+	*  EvaluateSelectedNodes
+	*/
+	static void EvaluateSelectedNodes(const FGraphPanelSelectionSet& SelectedNodes, FGraphEvaluationCallback);
+
+	/*
+	* EvaluateGraph
+	*/
+	static void EvaluateNode(Dataflow::FContext& Context, Dataflow::FTimestamp& OutLastNodeTimestamp,
+		const UDataflow* Dataflow, const FDataflowNode* Node = nullptr, const FDataflowOutput* Out = nullptr, 
+		FString NodeName = FString()); // @todo(Dataflow) deprecate  
+
+	static void EvaluateTerminalNode(Dataflow::FContext& Context, Dataflow::FTimestamp& OutLastNodeTimestamp,
+		const UDataflow* Dataflow, const FDataflowNode* Node = nullptr, const FDataflowOutput* Out = nullptr,
+		UObject* InAsset = nullptr, FString NodeName = FString());
+
+	/*
+	*  DeleteNodes
+	*/
+	static void DeleteNodes(UDataflow* Graph, const FGraphPanelSelectionSet& SelectedNodes);
+
+	/*
+	* OnNodeVerifyTitleCommit
+	*/
+	static bool OnNodeVerifyTitleCommit(const FText& NewText, UEdGraphNode* GraphNode, FText& OutErrorMessage);
+
+	/*
+	* OnNodeTitleCommitted
+	*/
+	static void OnNodeTitleCommitted(const FText& InNewText, ETextCommit::Type InCommitType, UEdGraphNode* GraphNode);
+
+	/*
+	*  OnPropertyValueChanged
+	*/
+	static void OnPropertyValueChanged(UDataflow* Graph, TSharedPtr<Dataflow::FEngineContext>& Context, Dataflow::FTimestamp& OutLastNodeTimestamp, const FPropertyChangedEvent& PropertyChangedEvent);
+
+	/*
+	*  OnSelectedNodesChanged
+	*/
+	static void OnSelectedNodesChanged(TSharedPtr<IStructureDetailsView> PropertiesEditor, UObject* Asset, UDataflow* Graph, const TSet<UObject*>& NewSelection);
+
+	/*
+	*  ToggleEnabledState
+	*/
+	static void ToggleEnabledState(UDataflow* Graph);
+
+
+
+};
diff --git a/Engine/Plugins/Experimental/Dataflow/Source/DataflowEditor/Public/Dataflow/DataflowEditorViewport.h b/Engine/Plugins/Experimental/Dataflow/Source/DataflowEditor/Public/Dataflow/DataflowEditorViewport.h
--- a/Engine/Plugins/Experimental/Dataflow/Source/DataflowEditor/Public/Dataflow/DataflowEditorViewport.h
+++ b/Engine/Plugins/Experimental/Dataflow/Source/DataflowEditor/Public/Dataflow/DataflowEditorViewport.h
@@ -3,6 +3,7 @@
 
 #include "CanvasTypes.h"
 #include "CoreMinimal.h"
+#include "Dataflow/DataflowComponent.h"
 #include "Dataflow/DataflowNodeParameters.h"
 #include "EditorViewportClient.h"
 #include "GeometryCollection/ManagedArrayCollection.h"
@@ -15,11 +16,47 @@ class FEditorViewportClient;
 class FDataflowEditorViewportClient;
 class SEditorViewport;
 class ADataflowActor;
-class ADataflowActor;
 class FDynamicMeshBuilder;
 
 // ----------------------------------------------------------------------------------
 
+class FDataflowSelectionState
+{
+public:
+	FDataflowSelectionState() {}
+	FDataflowSelectionState(const UDataflowComponent* DataflowComponent) 
+	:bComponentSelected(DataflowComponent->bSelected)
+	{}
+
+	enum EMode
+	{
+		DSS_Dataflow_None,
+		DSS_Dataflow_Object,
+		DSS_Dataflow_Max
+	};
+
+	void UpdateSelection(UDataflowComponent* DataflowComponent) const
+	{
+		DataflowComponent->bSelected = bComponentSelected;
+		DataflowComponent->Invalidate();
+	}
+
+	bool IsEmpty() const
+	{
+		return !bComponentSelected;
+	}
+
+	bool operator==(const FDataflowSelectionState& A) {
+		return A.bComponentSelected == bComponentSelected;
+	}
+	bool operator!=(const FDataflowSelectionState& A) {
+		return !this->operator==(A);
+	}
+
+	bool bComponentSelected = false;
+};
+
+
 class SDataflowEditorViewport : public SAssetEditorViewport, public ICommonEditorViewportToolbarInfoProvider, public FGCObject
 {
 public:
@@ -45,11 +82,15 @@ protected:
 	virtual TSharedRef<FEditorViewportClient> MakeEditorViewportClient() override;
 	virtual TSharedPtr<SWidget> MakeViewportToolbar() override;
 private:
+
+	virtual void BindCommands() override;
+
+
 	/// The scene for this viewport. 
 	TSharedPtr<FAdvancedPreviewScene> PreviewScene;
 
 	/// Editor viewport client 
-	TSharedPtr<FDataflowEditorViewportClient> EditorViewportClient;
+	TSharedPtr<FDataflowEditorViewportClient> ViewportClient;
 
 	TWeakPtr<FDataflowEditorToolkit> DataflowEditorToolkitPtr;
 	ADataflowActor* CustomDataflowActor = nullptr;
@@ -70,6 +111,12 @@ public:
 	Dataflow::FTimestamp LatestTimestamp(const UDataflow* Dataflow, const Dataflow::FContext* Context);
 	void SetDataflowActor(ADataflowActor* InActor) { DataflowActor = InActor; }
 
+	void SetSelectionMode(FDataflowSelectionState::EMode InState);
+	bool CanSetSelectionMode(FDataflowSelectionState::EMode InState);
+	bool IsSelectionModeActive(FDataflowSelectionState::EMode InState);
+	FDataflowSelectionState::EMode GetSelectionMode() const { return SelectionMode; }
+
+
 	// FEditorViewportClient interface
 	virtual void ProcessClick(FSceneView& View, HHitProxy* HitProxy, FKey Key, EInputEvent Event, uint32 HitX, uint32 HitY) override;
 	virtual void Draw(const FSceneView* View, FPrimitiveDrawInterface* PDI) override;
@@ -97,4 +144,8 @@ private:
 	TArray<FDynamicMeshVertex> VertexBuffer;
 	TUniquePtr<FDynamicMeshBuilder> MeshBuilder;
 
+	// State
+	FDataflowSelectionState SelectionState;
+	FDataflowSelectionState::EMode SelectionMode;
+
 };
diff --git a/Engine/Plugins/Experimental/Dataflow/Source/DataflowEditor/Public/Dataflow/DataflowEditorViewportToolbar.h b/Engine/Plugins/Experimental/Dataflow/Source/DataflowEditor/Public/Dataflow/DataflowEditorViewportToolbar.h
new file mode 100644
--- /dev/null
+++ b/Engine/Plugins/Experimental/Dataflow/Source/DataflowEditor/Public/Dataflow/DataflowEditorViewportToolbar.h
@@ -0,0 +1,26 @@
+// Copyright Epic Games, Inc. All Rights Reserved.
+#pragma once
+
+#include "CoreMinimal.h"
+#include "SEditorViewportToolBarMenu.h"
+
+// ----------------------------------------------------------------------------------
+class SDataflowEditorViewport;
+
+class SDataflowViewportSelectionToolBar : public SViewportToolBar
+{
+public:
+	SLATE_BEGIN_ARGS(SDataflowViewportSelectionToolBar) {}
+		SLATE_ARGUMENT(TWeakPtr<SDataflowEditorViewport>, EditorViewport)
+	SLATE_END_ARGS()
+
+		/** Constructs this widget with the given parameters */
+		void Construct(const FArguments& InArgs);
+
+		TSharedRef<SWidget> MakeSelectionModeToolBar();
+
+
+private:
+	/** Reference to the parent viewport */
+	TWeakPtr<SDataflowEditorViewport> EditorViewport;
+};
\ No newline at end of file
diff --git a/Engine/Plugins/Experimental/Dataflow/Source/DataflowEditor/Public/Dataflow/DataflowGraphEditor.h b/Engine/Plugins/Experimental/Dataflow/Source/DataflowEditor/Public/Dataflow/DataflowGraphEditor.h
--- a/Engine/Plugins/Experimental/Dataflow/Source/DataflowEditor/Public/Dataflow/DataflowGraphEditor.h
+++ b/Engine/Plugins/Experimental/Dataflow/Source/DataflowEditor/Public/Dataflow/DataflowGraphEditor.h
@@ -4,7 +4,7 @@
 
 #include "CoreMinimal.h"
 
-#include "Dataflow/DataflowEditorActions.h"
+#include "Dataflow/DataflowEditorCommands.h"
 #include "EdGraphUtilities.h"
 #include "Framework/Commands/GenericCommands.h"
 #include "Framework/Commands/UICommandList.h"
