Dataflow: add a multicast "on selection changed" delegate which gets re-broadcast from OnSelectedNodesChanged. This allows multiple handlers to be notified when the node selection in the Dataflow graph editor changes.

(Also make the timestamp comparison operator const.)

diff --git a/Engine/Plugins/Experimental/Dataflow/Source/DataflowEditor/Private/Dataflow/DataflowGraphEditor.cpp b/Engine/Plugins/Experimental/Dataflow/Source/DataflowEditor/Private/Dataflow/DataflowGraphEditor.cpp
--- a/Engine/Plugins/Experimental/Dataflow/Source/DataflowEditor/Private/Dataflow/DataflowGraphEditor.cpp
+++ b/Engine/Plugins/Experimental/Dataflow/Source/DataflowEditor/Private/Dataflow/DataflowGraphEditor.cpp
@@ -142,6 +142,8 @@ void SDataflowGraphEditor::OnSelectedNodesChanged(const TSet<UObject*>& NewSelec
 	{
 		FDataflowEditorCommands::OnSelectedNodesChanged(DetailsView, AssetOwner.Get(), DataflowAsset.Get(), NewSelection);
 	}
+
+	OnSelectionChangedMulticast.Broadcast(NewSelection);
 }
 
 FReply SDataflowGraphEditor::OnDragOver(const FGeometry& MyGeometry, const FDragDropEvent& DragDropEvent)
diff --git a/Engine/Plugins/Experimental/Dataflow/Source/DataflowEditor/Public/Dataflow/DataflowGraphEditor.h b/Engine/Plugins/Experimental/Dataflow/Source/DataflowEditor/Public/Dataflow/DataflowGraphEditor.h
--- a/Engine/Plugins/Experimental/Dataflow/Source/DataflowEditor/Public/Dataflow/DataflowGraphEditor.h
+++ b/Engine/Plugins/Experimental/Dataflow/Source/DataflowEditor/Public/Dataflow/DataflowGraphEditor.h
@@ -49,6 +49,11 @@ public:
 	SLATE_ARGUMENT(FDataflowEditorCommands::FOnDragDropEventCallback, OnDragDropEvent)
 	SLATE_END_ARGS()
 
+	// This delegate exists in SGraphEditor but it is not multicast, and we are going to bind it to OnSelectedNodesChanged().
+	// This new multicast delegate will be broadcast from the OnSelectedNodesChanged handler in case another class wants to be notified.
+	DECLARE_MULTICAST_DELEGATE_OneParam(FOnSelectionChangedMulticast, const FGraphPanelSelectionSet&)
+	FOnSelectionChangedMulticast OnSelectionChangedMulticast;
+
 	// SWidget overrides
 	virtual FReply OnDragOver(const FGeometry& MyGeometry, const FDragDropEvent& DragDropEvent) override;
 	virtual FReply OnDrop(const FGeometry& MyGeometry, const FDragDropEvent& DragDropEvent) override;
diff --git a/Engine/Source/Runtime/Experimental/Dataflow/Core/Public/Dataflow/DataflowNodeParameters.h b/Engine/Source/Runtime/Experimental/Dataflow/Core/Public/Dataflow/DataflowNodeParameters.h
--- a/Engine/Source/Runtime/Experimental/Dataflow/Core/Public/Dataflow/DataflowNodeParameters.h
+++ b/Engine/Source/Runtime/Experimental/Dataflow/Core/Public/Dataflow/DataflowNodeParameters.h
@@ -20,7 +20,7 @@ namespace Dataflow
 		Type Value = Type(0);
 
 		FTimestamp(Type InValue) : Value(InValue) {}
-		bool operator>=(const FTimestamp& InTimestamp) { return Value >= InTimestamp.Value; }
+		bool operator>=(const FTimestamp& InTimestamp) const { return Value >= InTimestamp.Value; }
 		bool IsInvalid() { return Value == Invalid; }
 
 		static Type Current();
