Dataflow : Node viewport selection

diff --git a/Engine/Plugins/Experimental/Dataflow/Source/DataflowEditor/Private/Dataflow/DataflowEditorViewport.cpp b/Engine/Plugins/Experimental/Dataflow/Source/DataflowEditor/Private/Dataflow/DataflowEditorViewport.cpp
--- a/Engine/Plugins/Experimental/Dataflow/Source/DataflowEditor/Private/Dataflow/DataflowEditorViewport.cpp
+++ b/Engine/Plugins/Experimental/Dataflow/Source/DataflowEditor/Private/Dataflow/DataflowEditorViewport.cpp
@@ -118,10 +118,12 @@ void FDataflowEditorViewportClient::SetSelectionMode(FDataflowSelectionState::EM
 		SelectionMode = InState;
 	}
 
-	if (SelectionMode == FDataflowSelectionState::EMode::DSS_Dataflow_None && !SelectionState.IsEmpty())
+	if (SelectionMode == FDataflowSelectionState::EMode::DSS_Dataflow_None)
 	{
-		SelectionState = FDataflowSelectionState();
-		SelectionState.UpdateSelection(DataflowActor->DataflowComponent);
+		if (!DataflowActor->DataflowComponent->GetSelectionState().IsEmpty())
+		{
+			DataflowActor->DataflowComponent->SetSelectionState(FDataflowSelectionState());
+		}
 	}
 }
 bool FDataflowEditorViewportClient::CanSetSelectionMode(FDataflowSelectionState::EMode InState)
@@ -138,10 +140,9 @@ bool FDataflowEditorViewportClient::CanSetSelectionMode(FDataflowSelectionState:
 		}
 	}
 
-	if (!SelectionState.IsEmpty())
+	if (!DataflowActor->DataflowComponent->GetSelectionState().IsEmpty())
 	{
-		SelectionState = FDataflowSelectionState();
-		SelectionState.UpdateSelection(DataflowActor->DataflowComponent);
+		DataflowActor->DataflowComponent->SetSelectionState(FDataflowSelectionState());
 	}
 
 	return false;
@@ -165,23 +166,35 @@ void FDataflowEditorViewportClient::ProcessClick(FSceneView& View, HHitProxy* Hi
 	Super::ProcessClick(View, HitProxy, Key, Event, HitX, HitY);
 	if (DataflowActor && DataflowActor->DataflowComponent)
 	{
+		const bool bIsShiftKeyDown = Viewport->KeyState(EKeys::LeftShift) || Viewport->KeyState(EKeys::RightShift);
+
+		FDataflowSelectionState SelectionState = DataflowActor->DataflowComponent->GetSelectionState();
 		FDataflowSelectionState PreState = SelectionState;
 
 		if (SelectionMode == FDataflowSelectionState::EMode::DSS_Dataflow_Object)
 		{
-			if (HitProxy && HitProxy->IsA(HDataflowActor::StaticGetType()))
+			if (HitProxy && HitProxy->IsA(HDataflowNode::StaticGetType()))
 			{
-				SelectionState.bComponentSelected = !SelectionState.bComponentSelected;
+				HDataflowNode* DataflowNode = (HDataflowNode*)(HitProxy);
+				FDataflowSelectionState::ObjectID ID(DataflowNode->NodeName, DataflowNode->GeometryIndex);
+				if (SelectionState.Nodes.Contains(ID))
+				{
+					SelectionState.Nodes.Remove(ID);
+				}
+				else
+				{
+					SelectionState.Nodes.AddUnique(ID);
+				}
 			}
 			else
 			{
-				SelectionState.bComponentSelected = false;
+				SelectionState.Nodes.Empty();
 			}
 		}
 
 		if (PreState != SelectionState)
 		{
-			SelectionState.UpdateSelection(DataflowActor->DataflowComponent);
+			DataflowActor->DataflowComponent->SetSelectionState(SelectionState);
 		}
 	}
 }
diff --git a/Engine/Plugins/Experimental/Dataflow/Source/DataflowEditor/Public/Dataflow/DataflowEditorViewport.h b/Engine/Plugins/Experimental/Dataflow/Source/DataflowEditor/Public/Dataflow/DataflowEditorViewport.h
--- a/Engine/Plugins/Experimental/Dataflow/Source/DataflowEditor/Public/Dataflow/DataflowEditorViewport.h
+++ b/Engine/Plugins/Experimental/Dataflow/Source/DataflowEditor/Public/Dataflow/DataflowEditorViewport.h
@@ -6,6 +6,7 @@
 #include "Dataflow/DataflowComponent.h"
 #include "Dataflow/DataflowNodeParameters.h"
 #include "EditorViewportClient.h"
+#include "GeometryCollection/Facades/CollectionRenderingFacade.h"
 #include "GeometryCollection/ManagedArrayCollection.h"
 #include "SAssetEditorViewport.h"
 #include "SCommonEditorViewportToolbarBase.h"
@@ -20,43 +21,6 @@ class FDynamicMeshBuilder;
 
 // ----------------------------------------------------------------------------------
 
-class FDataflowSelectionState
-{
-public:
-	FDataflowSelectionState() {}
-	FDataflowSelectionState(const UDataflowComponent* DataflowComponent) 
-	:bComponentSelected(DataflowComponent->bSelected)
-	{}
-
-	enum EMode
-	{
-		DSS_Dataflow_None,
-		DSS_Dataflow_Object,
-		DSS_Dataflow_Max
-	};
-
-	void UpdateSelection(UDataflowComponent* DataflowComponent) const
-	{
-		DataflowComponent->bSelected = bComponentSelected;
-		DataflowComponent->Invalidate();
-	}
-
-	bool IsEmpty() const
-	{
-		return !bComponentSelected;
-	}
-
-	bool operator==(const FDataflowSelectionState& A) {
-		return A.bComponentSelected == bComponentSelected;
-	}
-	bool operator!=(const FDataflowSelectionState& A) {
-		return !this->operator==(A);
-	}
-
-	bool bComponentSelected = false;
-};
-
-
 class SDataflowEditorViewport : public SAssetEditorViewport, public ICommonEditorViewportToolbarInfoProvider, public FGCObject
 {
 public:
@@ -135,7 +99,7 @@ private:
 	//
 	// Selection
 	//
-	FDataflowSelectionState SelectionState;
+
 	FDataflowSelectionState::EMode SelectionMode;
 
 };
diff --git a/Engine/Plugins/Experimental/Dataflow/Source/DataflowEnginePlugin/Private/Dataflow/DataflowComponent.cpp b/Engine/Plugins/Experimental/Dataflow/Source/DataflowEnginePlugin/Private/Dataflow/DataflowComponent.cpp
--- a/Engine/Plugins/Experimental/Dataflow/Source/DataflowEnginePlugin/Private/Dataflow/DataflowComponent.cpp
+++ b/Engine/Plugins/Experimental/Dataflow/Source/DataflowEnginePlugin/Private/Dataflow/DataflowComponent.cpp
@@ -28,6 +28,7 @@ UDataflowComponent::~UDataflowComponent()
 void UDataflowComponent::Invalidate()
 {
 	bUpdateRender = true;
+	bUpdateSelection = true;
 	bBoundsNeedsUpdate = true;
 	RenderCollection = FManagedArrayCollection();
 }
@@ -47,6 +48,8 @@ void UDataflowComponent::AddRenderTarget(const UDataflowEdNode* InTarget)
 
 void UDataflowComponent::TickComponent(float DeltaTime, enum ELevelTick TickType, FActorComponentTickFunction *ThisTickFunction)
 {
+	bool bNeedsSceneProxyUpdate = false;
+
 	Super::TickComponent(DeltaTime, TickType, ThisTickFunction);
 
 	if (bUpdateRender)
@@ -63,9 +66,21 @@ void UDataflowComponent::TickComponent(float DeltaTime, enum ELevelTick TickType
 				bNeedsRefresh |= Target->Render(Facade, Context);
 			}
 		}
+
 		UpdateLocalBounds();
 		bUpdateRender = false;
+		bNeedsSceneProxyUpdate = true;
+	}
+
+	if (bUpdateSelection)
+	{
+		SelectionState.UpdateSelection(this);
+		bUpdateSelection = false;
+		bNeedsSceneProxyUpdate = true;
+	}
 
+	if (bNeedsSceneProxyUpdate)
+	{
 		MarkRenderStateDirty();
 	}
 }
@@ -95,6 +110,11 @@ const FManagedArrayCollection& UDataflowComponent::GetRenderingCollection() cons
 	return RenderCollection;
 }
 
+FManagedArrayCollection& UDataflowComponent::ModifyRenderingCollection()
+{
+	return RenderCollection;
+}
+
 
 FMaterialRelevance UDataflowComponent::GetMaterialRelevance(ERHIFeatureLevel::Type InFeatureLevel) const
 {
diff --git a/Engine/Plugins/Experimental/Dataflow/Source/DataflowEnginePlugin/Private/Dataflow/DataflowComponentSelectionState.cpp b/Engine/Plugins/Experimental/Dataflow/Source/DataflowEnginePlugin/Private/Dataflow/DataflowComponentSelectionState.cpp
new file mode 100644
--- /dev/null
+++ b/Engine/Plugins/Experimental/Dataflow/Source/DataflowEnginePlugin/Private/Dataflow/DataflowComponentSelectionState.cpp
@@ -0,0 +1,49 @@
+// Copyright Epic Games, Inc. All Rights Reserved.
+
+#include "Dataflow/DataflowComponentSelectionState.h"
+
+#include "Dataflow/DataflowComponent.h"
+#include "GeometryCollection/Facades/CollectionRenderingFacade.h"
+
+
+void FDataflowSelectionState::UpdateSelection(UDataflowComponent* DataflowComponent)
+{
+	using namespace GeometryCollection::Facades;
+	FRenderingFacade Facade(DataflowComponent->ModifyRenderingCollection());
+	if (Facade.IsValid())
+	{
+		TManagedArray<int32>& SelectionArray = Facade.ModifySelectionState();
+		if (Nodes.Num())
+		{
+			const TManagedArray<FString>& GeomNameArray = Facade.GetGeometryNameAttribute();
+
+			TArray<bool> bVisited;
+			bVisited.Init(false, Nodes.Num());
+
+			for (int i = 0; i < SelectionArray.Num(); i++)
+			{
+				ObjectID ID(GeomNameArray[i], i);
+				SelectionArray[i] = false;
+				int32 IndexOf = Nodes.IndexOfByKey(ID);
+				if (IndexOf != INDEX_NONE)
+				{
+					SelectionArray[i] = true;
+					bVisited[IndexOf] = true;
+				}
+			}
+
+			// remove unknown selections
+			for (int32 Ndx = bVisited.Num() - 1; 0 <= Ndx; Ndx--)
+			{
+				if (!bVisited[Ndx])
+				{
+					Nodes.RemoveAt(Ndx);
+				}
+			}
+		}
+		else
+		{
+			SelectionArray.Fill(0);
+		}
+	}
+}
diff --git a/Engine/Plugins/Experimental/Dataflow/Source/DataflowEnginePlugin/Private/Dataflow/DataflowEnginePlugin.cpp b/Engine/Plugins/Experimental/Dataflow/Source/DataflowEnginePlugin/Private/Dataflow/DataflowEnginePlugin.cpp
--- a/Engine/Plugins/Experimental/Dataflow/Source/DataflowEnginePlugin/Private/Dataflow/DataflowEnginePlugin.cpp
+++ b/Engine/Plugins/Experimental/Dataflow/Source/DataflowEnginePlugin/Private/Dataflow/DataflowEnginePlugin.cpp
@@ -13,7 +13,7 @@
 
 
 FColor IDataflowEnginePlugin::SurfaceColor = FLinearColor(0.6,0.6,0.6).ToRGBE();
-FColor IDataflowEnginePlugin::PrimarySelectionColor = FLinearColor(0.8, 0.3, 0.3).ToRGBE();
+FColor IDataflowEnginePlugin::PrimarySelectionColor = FLinearColor(0.8, 0.4, 0.0).ToRGBE();
 
 class FDataflowEnginePlugin : public IDataflowEnginePlugin
 {
diff --git a/Engine/Plugins/Experimental/Dataflow/Source/DataflowEnginePlugin/Private/Dataflow/DataflowEngineSceneHitProxies.cpp b/Engine/Plugins/Experimental/Dataflow/Source/DataflowEnginePlugin/Private/Dataflow/DataflowEngineSceneHitProxies.cpp
--- a/Engine/Plugins/Experimental/Dataflow/Source/DataflowEnginePlugin/Private/Dataflow/DataflowEngineSceneHitProxies.cpp
+++ b/Engine/Plugins/Experimental/Dataflow/Source/DataflowEnginePlugin/Private/Dataflow/DataflowEngineSceneHitProxies.cpp
@@ -2,4 +2,4 @@
 
 #include "Dataflow/DataflowEngineSceneHitProxies.h"
 
-IMPLEMENT_HIT_PROXY(HDataflowActor, HActor);
+IMPLEMENT_HIT_PROXY(HDataflowNode, HActor);
diff --git a/Engine/Plugins/Experimental/Dataflow/Source/DataflowEnginePlugin/Private/Dataflow/DataflowEngineSceneProxy.cpp b/Engine/Plugins/Experimental/Dataflow/Source/DataflowEnginePlugin/Private/Dataflow/DataflowEngineSceneProxy.cpp
--- a/Engine/Plugins/Experimental/Dataflow/Source/DataflowEnginePlugin/Private/Dataflow/DataflowEngineSceneProxy.cpp
+++ b/Engine/Plugins/Experimental/Dataflow/Source/DataflowEnginePlugin/Private/Dataflow/DataflowEngineSceneProxy.cpp
@@ -17,7 +17,6 @@ FDataflowEngineSceneProxy::FDataflowEngineSceneProxy(UDataflowComponent* Compone
 	: FPrimitiveSceneProxy(Component)
 	, MaterialRelevance(Component->GetMaterialRelevance(GetScene().GetFeatureLevel()))
 	, VertexFactory(GetScene().GetFeatureLevel(), "FTriangleSetSceneProxy")
-	, DefaultHitProxy(new HDataflowActor(Component->GetOwner(), Component, 3, 7))
 	, RenderMaterial(Component->GetMaterial(0))
 	, ConstantData(Component->GetRenderingCollection().NewCopy())
 {}
@@ -48,9 +47,6 @@ void FDataflowEngineSceneProxy::CreateRenderThreadResources()
 	VertexBuffers.StaticMeshVertexBuffer.Init(TotalNumVertices, NumTextureCoordinates);
 	VertexBuffers.ColorVertexBuffer.Init(TotalNumVertices);
 	IndexBuffer.Indices.SetNumUninitialized(TotalNumIndices);
-#if WITH_EDITOR
-	HitProxyIdBuffer.Init(TotalNumVertices);
-#endif
 
 	// Initialize
 	// Triangles are represented as two tris, all of whose vertices are
@@ -81,6 +77,8 @@ void FDataflowEngineSceneProxy::CreateRenderThreadResources()
 
 		const TManagedArray<FIntVector>& Indices = Facade.GetIndices();
 		const TManagedArray<FVector3f>& Vertex = Facade.GetVertices();
+		const TManagedArray<int32>& SelectionArray = Facade.GetSelectionState();
+		const TManagedArray<int32>& GeomIndex = Facade.GetVertexToGeometryIndexAttribute();
 
 		// The color stored in the vertices actually gets interpreted as a linear
 		// color by the material, whereas it is more convenient for the user of the
@@ -112,7 +110,9 @@ void FDataflowEngineSceneProxy::CreateRenderThreadResources()
 		VertexBuffers.StaticMeshVertexBuffer.SetVertexUV(VertexBufferIndex + 1, 0, FVector2f(0, 0));
 		VertexBuffers.StaticMeshVertexBuffer.SetVertexUV(VertexBufferIndex + 2, 0, FVector2f(0, 0));
 
-		FColor FaceColor = IsSelected() ? IDataflowEnginePlugin::PrimarySelectionColor : IDataflowEnginePlugin::SurfaceColor;
+
+
+		FColor FaceColor = SelectionArray[GeomIndex[Indices[i][0]]] ? IDataflowEnginePlugin::PrimarySelectionColor : IDataflowEnginePlugin::SurfaceColor;
 		VertexBuffers.ColorVertexBuffer.VertexColor(VertexBufferIndex + 0) = FaceColor;
 		VertexBuffers.ColorVertexBuffer.VertexColor(VertexBufferIndex + 1) = FaceColor;
 		VertexBuffers.ColorVertexBuffer.VertexColor(VertexBufferIndex + 2) = FaceColor;
@@ -121,13 +121,6 @@ void FDataflowEngineSceneProxy::CreateRenderThreadResources()
 		IndexBuffer.Indices[IndexBufferIndex + 1] = VertexBufferIndex + 1;
 		IndexBuffer.Indices[IndexBufferIndex + 2] = VertexBufferIndex + 2;
 
-#if WITH_EDITOR
-		// @todo(Dataflow) Question : Does this even matter? 
-		// FMeshBatch::BatchHitProxyId has this also and 
-		// when GetCustomHitProxyIdBuffer returns non-null
-		// the hit buffer seems to be ignored in the render.
-		HitProxyIdBuffer.VertexColor(i) = DefaultHitProxy->Id.GetColor();
-#endif
 		});
 
 	}
@@ -148,9 +141,6 @@ void FDataflowEngineSceneProxy::CreateRenderThreadResources()
 
 	VertexFactory.InitResource();
 	IndexBuffer.InitResource();
-#if WITH_EDITOR
-	HitProxyIdBuffer.InitResource();
-#endif
 }
 
 void FDataflowEngineSceneProxy::DestroyRenderThreadResources()
@@ -161,9 +151,6 @@ void FDataflowEngineSceneProxy::DestroyRenderThreadResources()
 	VertexBuffers.ColorVertexBuffer.ReleaseResource();
 	IndexBuffer.ReleaseResource();
 	VertexFactory.ReleaseResource();
-#if WITH_EDITOR
-	HitProxyIdBuffer.ReleaseResource();
-#endif
 	delete ConstantData;
 	ConstantData = nullptr;
 }
@@ -178,9 +165,10 @@ HHitProxy* FDataflowEngineSceneProxy::CreateHitProxies(UPrimitiveComponent* Comp
 	if (Facade.NumTriangles() > 0)
 	{
 		int32 NumGeom = Facade.NumGeometry();
+		const TManagedArray<FString>& NodeName = Facade.GetGeometryNameAttribute();
 		for (int i = 0; i < NumGeom; i++)
 		{
-			HDataflowActor * NodeProxy = new HDataflowActor(Component->GetOwner(), Component, i, INDEX_NONE);
+			HDataflowNode * NodeProxy = new HDataflowNode(Component->GetOwner(), Component, NodeName[i], i);
 			LocalHitProxies.Add(NodeProxy);
 			Facade.ModifyHitProxyIndexAttribute()[i] = LocalHitProxies.Num() - 1;
 			OutHitProxies.Add(NodeProxy);
@@ -189,12 +177,6 @@ HHitProxy* FDataflowEngineSceneProxy::CreateHitProxies(UPrimitiveComponent* Comp
 	return FPrimitiveSceneProxy::CreateHitProxies(Component, OutHitProxies);
 }
 
-
-const FColorVertexBuffer* FDataflowEngineSceneProxy::GetCustomHitProxyIdBuffer() const
-{
-	// returning my own HitProxyIdBuffer causes the Hit to fail. 
-	return nullptr;// &HitProxyIdBuffer;
-}
 #endif // WITH_EDITOR
 
 
@@ -234,8 +216,7 @@ void FDataflowEngineSceneProxy::GetDynamicMeshElements(
 				Mesh.bCanApplyViewModeOverrides = true;
 #if WITH_EDITOR
 				Mesh.BatchHitProxyId = LocalHitProxies[Facade.GetHitProxyIndexAttribute()[MeshBatchData.GeomIndex]]->Id;
-#endif // WITH_EDITOR
-
+#endif
 				Collector.AddMesh(ViewIndex, Mesh);
 			}
 		}
diff --git a/Engine/Plugins/Experimental/Dataflow/Source/DataflowEnginePlugin/Private/Dataflow/DataflowEngineSceneProxy.h b/Engine/Plugins/Experimental/Dataflow/Source/DataflowEnginePlugin/Private/Dataflow/DataflowEngineSceneProxy.h
--- a/Engine/Plugins/Experimental/Dataflow/Source/DataflowEnginePlugin/Private/Dataflow/DataflowEngineSceneProxy.h
+++ b/Engine/Plugins/Experimental/Dataflow/Source/DataflowEnginePlugin/Private/Dataflow/DataflowEngineSceneProxy.h
@@ -41,7 +41,6 @@ public:
 
 #if WITH_EDITOR
 	virtual HHitProxy* CreateHitProxies(UPrimitiveComponent* Component, TArray<TRefCountPtr<HHitProxy> >& OutHitProxies) override;
-	virtual const FColorVertexBuffer* GetCustomHitProxyIdBuffer() const override;
 #endif // WITH_EDITOR
 
 
@@ -64,8 +63,6 @@ private:
 	TArray<TRefCountPtr<HHitProxy> > LocalHitProxies;
 #endif // WITH_EDITOR
 
-	TRefCountPtr<HHitProxy> DefaultHitProxy;
-
 	// Render thread copy of data. 
 	UMaterialInterface* RenderMaterial = nullptr;
 	FManagedArrayCollection* ConstantData = nullptr;
diff --git a/Engine/Plugins/Experimental/Dataflow/Source/DataflowEnginePlugin/Public/Dataflow/DataflowComponent.h b/Engine/Plugins/Experimental/Dataflow/Source/DataflowEnginePlugin/Public/Dataflow/DataflowComponent.h
--- a/Engine/Plugins/Experimental/Dataflow/Source/DataflowEnginePlugin/Public/Dataflow/DataflowComponent.h
+++ b/Engine/Plugins/Experimental/Dataflow/Source/DataflowEnginePlugin/Public/Dataflow/DataflowComponent.h
@@ -3,12 +3,14 @@
 
 #include "GeometryCollection/ManagedArrayCollection.h"
 #include "Components/MeshComponent.h"
+#include "Dataflow/DataflowComponentSelectionState.h"
 #include "Dataflow/DataflowObject.h"
 #include "UObject/ObjectMacros.h"
 #include "UObject/UObjectGlobals.h"
 
 #include "DataflowComponent.generated.h"
 
+
 /**
 *	UDataflowComponent
 */
@@ -37,10 +39,6 @@ public:
 	virtual FMaterialRelevance GetMaterialRelevance(ERHIFeatureLevel::Type InFeatureLevel) const;
 	UMaterialInterface* GetDefaultMaterial() const;
 
-	// temp code for selection...
-	bool bSelected = false;
-	virtual bool ShouldRenderSelected() const override { return Super::ShouldRenderSelected() || bSelected; }
-
 	/** Render Targets*/
 	void ResetRenderTargets(); 
 	void AddRenderTarget(const UDataflowEdNode* InTarget); 
@@ -52,12 +50,20 @@ public:
 	/** RenderCollection */
 	void SetRenderingCollection(FManagedArrayCollection&& InCollection);
 	const FManagedArrayCollection& GetRenderingCollection() const;
+	      FManagedArrayCollection& ModifyRenderingCollection();
 
 
 	/** Dataflow */
 	void SetDataflow(const UDataflow* InDataflow) { Dataflow = InDataflow; }
 	const UDataflow* GetDataflow() const { return Dataflow; }
 
+	/* Selection */
+	FDataflowSelectionState GetSelectionState() const { return SelectionState; }
+	void SetSelectionState(const FDataflowSelectionState& InState) 
+	{
+		bUpdateSelection = true;
+		SelectionState = InState; 
+	}
 
 
 private:
@@ -67,7 +73,9 @@ private:
 	FManagedArrayCollection RenderCollection;
 
 	bool bUpdateRender = true;
+	bool bUpdateSelection = true;
 	bool bBoundsNeedsUpdate = true;
 	FBoxSphereBounds BoundingBox = FBoxSphereBounds(ForceInitToZero);
+	FDataflowSelectionState SelectionState;
 };
 
diff --git a/Engine/Plugins/Experimental/Dataflow/Source/DataflowEnginePlugin/Public/Dataflow/DataflowComponentSelectionState.h b/Engine/Plugins/Experimental/Dataflow/Source/DataflowEnginePlugin/Public/Dataflow/DataflowComponentSelectionState.h
new file mode 100644
--- /dev/null
+++ b/Engine/Plugins/Experimental/Dataflow/Source/DataflowEnginePlugin/Public/Dataflow/DataflowComponentSelectionState.h
@@ -0,0 +1,52 @@
+// Copyright Epic Games, Inc. All Rights Reserved.
+#pragma once
+
+#include "CoreMinimal.h"
+
+class UDataflowComponent;
+
+
+class DATAFLOWENGINEPLUGIN_API FDataflowSelectionState
+{
+public:
+	struct ObjectID
+	{
+		ObjectID(FString InName, int32 InID) : Name(InName), ID(InID) {}
+
+		bool operator==(const ObjectID& A) const {
+			return A.Name.Equals(Name);
+		}
+
+		FString Name;
+		int32 ID;
+	};
+
+	FDataflowSelectionState() {}
+	FDataflowSelectionState(const UDataflowComponent* DataflowComponent)
+	{
+		check(DataflowComponent);
+	}
+
+	enum EMode
+	{
+		DSS_Dataflow_None,
+		DSS_Dataflow_Object,
+		DSS_Dataflow_Max
+	};
+
+	void UpdateSelection(UDataflowComponent* DataflowComponent);
+
+	bool IsEmpty() const
+	{
+		return Nodes.Num() == 0;
+	}
+
+	bool operator==(const FDataflowSelectionState& A) const {
+		return A.Nodes == Nodes;
+	}
+	bool operator!=(const FDataflowSelectionState& A) const {
+		return !this->operator==(A);
+	}
+
+	TArray<ObjectID> Nodes;
+};
\ No newline at end of file
diff --git a/Engine/Plugins/Experimental/Dataflow/Source/DataflowEnginePlugin/Public/Dataflow/DataflowEngineSceneHitProxies.h b/Engine/Plugins/Experimental/Dataflow/Source/DataflowEnginePlugin/Public/Dataflow/DataflowEngineSceneHitProxies.h
--- a/Engine/Plugins/Experimental/Dataflow/Source/DataflowEnginePlugin/Public/Dataflow/DataflowEngineSceneHitProxies.h
+++ b/Engine/Plugins/Experimental/Dataflow/Source/DataflowEnginePlugin/Public/Dataflow/DataflowEngineSceneHitProxies.h
@@ -12,15 +12,18 @@ class UPrimitiveComponent;
 
 /** HitProxy with for dataflow actor.
  */
-struct HDataflowActor : public HActor
+struct HDataflowNode : public HActor
 {
 	DECLARE_HIT_PROXY(DATAFLOWENGINEPLUGIN_API)
+	int32 GeometryIndex = INDEX_NONE;
+	FString NodeName = FString("");
 
-		HDataflowActor(AActor* InActor, const UPrimitiveComponent* InPrimitiveComponent, int32 InSectionIndex, int32 InMaterialIndex)
+	HDataflowNode(AActor* InActor, const UPrimitiveComponent* InPrimitiveComponent, FString InNodeName, int32 InGeometryIndex)
 		: HActor(InActor, InPrimitiveComponent)
 	{
-		SectionIndex= InSectionIndex;
-		MaterialIndex = InMaterialIndex;
+		NodeName = InNodeName;
+		GeometryIndex = InGeometryIndex;
+		SectionIndex = GeometryIndex;
 	}
 
 	virtual EMouseCursor::Type GetMouseCursor() override
diff --git a/Engine/Source/Runtime/Experimental/Chaos/Private/GeometryCollection/Facades/CollectionRenderingFacade.cpp b/Engine/Source/Runtime/Experimental/Chaos/Private/GeometryCollection/Facades/CollectionRenderingFacade.cpp
--- a/Engine/Source/Runtime/Experimental/Chaos/Private/GeometryCollection/Facades/CollectionRenderingFacade.cpp
+++ b/Engine/Source/Runtime/Experimental/Chaos/Private/GeometryCollection/Facades/CollectionRenderingFacade.cpp
@@ -24,6 +24,7 @@ namespace GeometryCollection::Facades
 		, VertixCountAttribute(InCollection, "VertexCount", FGeometryCollection::GeometryGroup)
 		, IndicesStartAttribute(InCollection, "IndicesStart", FGeometryCollection::GeometryGroup)
 		, IndicesCountAttribute(InCollection, "IndicesCount", FGeometryCollection::GeometryGroup)
+		, SelectionState(InCollection, "SelectionState", FGeometryCollection::GeometryGroup)
 
 	{}
 
@@ -41,6 +42,7 @@ namespace GeometryCollection::Facades
 		, VertixCountAttribute(InCollection, "VertexCount", FGeometryCollection::GeometryGroup)
 		, IndicesStartAttribute(InCollection, "IndicesStart", FGeometryCollection::GeometryGroup, FGeometryCollection::FacesGroup)
 		, IndicesCountAttribute(InCollection, "IndicesCount", FGeometryCollection::GeometryGroup)
+		, SelectionState(InCollection, "SelectionState", FGeometryCollection::GeometryGroup)
 	{}
 
 	//
@@ -61,6 +63,7 @@ namespace GeometryCollection::Facades
 		VertixCountAttribute.Add();
 		IndicesStartAttribute.Add();
 		IndicesCountAttribute.Add();
+		SelectionState.Add();
 	}
 
 	bool FRenderingFacade::CanRenderSurface( ) const
@@ -74,8 +77,9 @@ namespace GeometryCollection::Facades
 			IndicesAttribute.IsValid() &&
 			MaterialIDAttribute.IsValid() && TriangleSectionAttribute.IsValid() &&
 			GeometryNameAttribute.IsValid() && HitProxyIndexAttribute.IsValid() &&
-			VertixStartAttribute.IsValid() && VertixCountAttribute.IsValid() && 
-			IndicesStartAttribute.IsValid() && IndicesCountAttribute.IsValid();
+			VertixStartAttribute.IsValid() && VertixCountAttribute.IsValid() &&
+			IndicesStartAttribute.IsValid() && IndicesCountAttribute.IsValid() &&
+			SelectionState.IsValid();
 	}
 
 	int32 FRenderingFacade::NumTriangles() const
@@ -161,6 +165,7 @@ namespace GeometryCollection::Facades
 			VertixCountAttribute.Modify()[GeomIndex] = 0;
 			IndicesStartAttribute.Modify()[GeomIndex] = IndicesAttribute.Num();
 			IndicesCountAttribute.Modify()[GeomIndex] = 0;
+			SelectionState.Modify()[GeomIndex] = 0;
 		}
 		return GeomIndex;
 	}
@@ -198,6 +203,17 @@ namespace GeometryCollection::Facades
 		}
 	}
 
+	FRenderingFacade::FStringIntMap FRenderingFacade::GetGeometryNameToIndexMap() const
+	{
+		FStringIntMap Map;
+		for (int32 i = 0; i < GeometryNameAttribute.Num(); i++)
+		{
+			Map.Add(GetGeometryNameAttribute()[i], i);
+		}
+		return Map;
+	}
+
+
 }; // GeometryCollection::Facades
 
 
diff --git a/Engine/Source/Runtime/Experimental/Chaos/Public/GeometryCollection/Facades/CollectionRenderingFacade.h b/Engine/Source/Runtime/Experimental/Chaos/Public/GeometryCollection/Facades/CollectionRenderingFacade.h
--- a/Engine/Source/Runtime/Experimental/Chaos/Public/GeometryCollection/Facades/CollectionRenderingFacade.h
+++ b/Engine/Source/Runtime/Experimental/Chaos/Public/GeometryCollection/Facades/CollectionRenderingFacade.h
@@ -22,7 +22,7 @@ namespace GeometryCollection::Facades
 	{
 	public:
 		typedef FGeometryCollectionSection FTriangleSection;
-
+		typedef TMap<FString, int32> FStringIntMap;
 		/**
 		* FRenderingFacade Constuctor
 		* @param VertixDependencyGroup : GroupName the index attribute is dependent on.
@@ -73,6 +73,11 @@ namespace GeometryCollection::Facades
 		/** BuildMeshSections */
 		TArray<FTriangleSection> BuildMeshSections(const TArray<FIntVector>& Indices, TArray<int32> BaseMeshOriginalIndicesIndex, TArray<FIntVector>& RetIndices) const;
 		
+
+		//
+		// Geometry Group Attributes
+		//
+
 		/** Geometry Group Start : */
 		int32 StartGeometryGroup(FString InName);
 
@@ -85,8 +90,8 @@ namespace GeometryCollection::Facades
 		const TManagedArray< FString >& GetGeometryNameAttribute() const { return GeometryNameAttribute.Get(); }
 
 		/** HitProxyIDAttribute */
-		const TManagedArray< int32>& GetHitProxyIndexAttribute() const { return HitProxyIndexAttribute.Get(); }
-		TManagedArray< int32 >& ModifyHitProxyIndexAttribute() {check(!IsConst());return HitProxyIndexAttribute.Modify(); }
+		const TManagedArray< int32 >& GetHitProxyIndexAttribute() const { return HitProxyIndexAttribute.Get(); }
+		      TManagedArray< int32 >& ModifyHitProxyIndexAttribute() {check(!IsConst());return HitProxyIndexAttribute.Modify(); }
 
 		/** VertixStartAttribute */
 		const TManagedArray< int32 >& GetVertixStartAttribute() const { return VertixStartAttribute.Get(); }
@@ -100,6 +105,13 @@ namespace GeometryCollection::Facades
 		/** IndicesCountAttribute */
 		const TManagedArray< int32 >& GetIndicesCountAttribute() const { return IndicesCountAttribute.Get(); }
 
+		/** SelectionState */
+		const TManagedArray< int32 >& GetSelectionState() const { return SelectionState.Get(); }
+		      TManagedArray< int32 >& ModifySelectionState() { check(!IsConst()); return SelectionState.Modify(); }
+
+		/** GetGeometryNameToIndexMap */
+		FStringIntMap GetGeometryNameToIndexMap() const;
+
 	private : 
 		const FManagedArrayCollection& ConstCollection;
 		FManagedArrayCollection* Collection = nullptr;
@@ -116,6 +128,7 @@ namespace GeometryCollection::Facades
 		TManagedArrayAccessor<int32> VertixCountAttribute;
 		TManagedArrayAccessor<int32> IndicesStartAttribute;
 		TManagedArrayAccessor<int32> IndicesCountAttribute;
+		TManagedArrayAccessor<int32> SelectionState;
 
 
 	};
