Dataflow : Selection operation

diff --git a/Engine/Plugins/Experimental/Dataflow/Source/DataflowEditor/Private/Dataflow/DataflowEditorViewport.cpp b/Engine/Plugins/Experimental/Dataflow/Source/DataflowEditor/Private/Dataflow/DataflowEditorViewport.cpp
--- a/Engine/Plugins/Experimental/Dataflow/Source/DataflowEditor/Private/Dataflow/DataflowEditorViewport.cpp
+++ b/Engine/Plugins/Experimental/Dataflow/Source/DataflowEditor/Private/Dataflow/DataflowEditorViewport.cpp
@@ -167,6 +167,7 @@ void FDataflowEditorViewportClient::ProcessClick(FSceneView& View, HHitProxy* Hi
 	if (DataflowActor && DataflowActor->DataflowComponent)
 	{
 		const bool bIsShiftKeyDown = Viewport->KeyState(EKeys::LeftShift) || Viewport->KeyState(EKeys::RightShift);
+		const bool bIsCtrltKeyDown = Viewport->KeyState(EKeys::LeftControl) || Viewport->KeyState(EKeys::RightControl);
 
 		FDataflowSelectionState SelectionState = DataflowActor->DataflowComponent->GetSelectionState();
 		FDataflowSelectionState PreState = SelectionState;
@@ -177,16 +178,28 @@ void FDataflowEditorViewportClient::ProcessClick(FSceneView& View, HHitProxy* Hi
 			{
 				HDataflowNode* DataflowNode = (HDataflowNode*)(HitProxy);
 				FDataflowSelectionState::ObjectID ID(DataflowNode->NodeName, DataflowNode->GeometryIndex);
-				if (SelectionState.Nodes.Contains(ID))
+
+				if (bIsShiftKeyDown)
 				{
-					SelectionState.Nodes.Remove(ID);
+					if (!SelectionState.Nodes.Contains(ID))
+					{
+						SelectionState.Nodes.AddUnique(ID);
+					}
+				}
+				else if (bIsCtrltKeyDown)
+				{
+					if (SelectionState.Nodes.Contains(ID))
+					{
+						SelectionState.Nodes.Remove(ID);
+					}
 				}
 				else
 				{
+					SelectionState.Nodes.Empty();
 					SelectionState.Nodes.AddUnique(ID);
 				}
 			}
-			else
+			else if (!bIsShiftKeyDown && !bIsCtrltKeyDown)
 			{
 				SelectionState.Nodes.Empty();
 			}
diff --git a/Engine/Plugins/Experimental/Dataflow/Source/DataflowEnginePlugin/Public/Dataflow/DataflowEngineSceneHitProxies.h b/Engine/Plugins/Experimental/Dataflow/Source/DataflowEnginePlugin/Public/Dataflow/DataflowEngineSceneHitProxies.h
--- a/Engine/Plugins/Experimental/Dataflow/Source/DataflowEnginePlugin/Public/Dataflow/DataflowEngineSceneHitProxies.h
+++ b/Engine/Plugins/Experimental/Dataflow/Source/DataflowEnginePlugin/Public/Dataflow/DataflowEngineSceneHitProxies.h
@@ -28,6 +28,6 @@ struct HDataflowNode : public HActor
 
 	virtual EMouseCursor::Type GetMouseCursor() override
 	{
-		return EMouseCursor::GrabHandClosed;
+		return EMouseCursor::Default;
 	}
 };
\ No newline at end of file
