Dataflow : Point selection

diff --git a/Engine/Plugins/Experimental/Dataflow/Source/DataflowEditor/Private/Dataflow/DataflowEditorCommands.cpp b/Engine/Plugins/Experimental/Dataflow/Source/DataflowEditor/Private/Dataflow/DataflowEditorCommands.cpp
--- a/Engine/Plugins/Experimental/Dataflow/Source/DataflowEditor/Private/Dataflow/DataflowEditorCommands.cpp
+++ b/Engine/Plugins/Experimental/Dataflow/Source/DataflowEditor/Private/Dataflow/DataflowEditorCommands.cpp
@@ -17,6 +17,8 @@ void FDataflowEditorCommandsImpl::RegisterCommands()
 	UI_COMMAND(CreateComment, "CreateComment", "Create a Comment node.", EUserInterfaceActionType::None, FInputChord());
 	UI_COMMAND(ToggleEnabledState, "ToggleEnabledState", "Toggle node between Enabled/Disabled state.", EUserInterfaceActionType::Button, FInputChord());
 	UI_COMMAND(ToggleObjectSelection, "ToggleObjectSelection", "Enable object selection in editor.", EUserInterfaceActionType::ToggleButton, FInputChord());
+	UI_COMMAND(ToggleFaceSelection, "ToggleFaceSelection", "Enable face selection in editor.", EUserInterfaceActionType::ToggleButton, FInputChord());
+	UI_COMMAND(ToggleVertexSelection, "ToggleVertexSelection", "Enable vertex selection in editor.", EUserInterfaceActionType::ToggleButton, FInputChord());
 
 	if (Dataflow::FNodeFactory* Factory = Dataflow::FNodeFactory::GetInstance())
 	{
diff --git a/Engine/Plugins/Experimental/Dataflow/Source/DataflowEditor/Private/Dataflow/DataflowEditorStyle.h b/Engine/Plugins/Experimental/Dataflow/Source/DataflowEditor/Private/Dataflow/DataflowEditorStyle.h
--- a/Engine/Plugins/Experimental/Dataflow/Source/DataflowEditor/Private/Dataflow/DataflowEditorStyle.h
+++ b/Engine/Plugins/Experimental/Dataflow/Source/DataflowEditor/Private/Dataflow/DataflowEditorStyle.h
@@ -13,10 +13,10 @@ public:
 	FDataflowEditorStyle() : FSlateStyleSet("DataflowEditorStyle")
 	{
 		const FVector2D Icon16x16(16.f, 16.f);
-		const FVector2D Icon64x64(64.f, 64.f);
-		const FVector2D Icon56x28(56.f, 28.f);
 		const FVector2D Icon28x14(28.f, 14.f);
-		const FVector2D Icon40x40(28.f, 14.f);
+		const FVector2D Icon20x20(20.f, 20.f);
+		const FVector2D Icon40x40(40.f, 40.f);
+		const FVector2D Icon64x64(64.f, 64.f);
 
 		SetContentRoot(IPluginManager::Get().FindPlugin("Dataflow")->GetBaseDir() / TEXT("Resources"));
 		Set("ClassIcon.Dataflow", new FSlateVectorImageBrush(RootToContentDir(TEXT("DataflowAsset_16.svg")), Icon16x16));
@@ -29,6 +29,11 @@ public:
 		Set("Dataflow.Cached.False", new FSlateImageBrush(RootToContentDir(TEXT("Slate/status_grey.png")), Icon16x16));
 		Set("Dataflow.Cached.True", new FSlateImageBrush(RootToContentDir(TEXT("Slate/status_green.png")), Icon16x16));
 
+		Set("Dataflow.SelectObject", new FSlateImageBrush(RootToContentDir(TEXT("Slate/Dataflow_SelectObject_40x.png")), Icon40x40));
+		Set("Dataflow.SelectFace", new FSlateImageBrush(RootToContentDir(TEXT("Slate/Dataflow_SelectFace_40x.png")), Icon40x40));
+		Set("Dataflow.SelectVertex", new FSlateImageBrush(RootToContentDir(TEXT("Slate/Dataflow_SelectVertex_40x.png")), Icon40x40));
+
+
 		FSlateStyleRegistry::RegisterSlateStyle(*this);
 	}
 
diff --git a/Engine/Plugins/Experimental/Dataflow/Source/DataflowEditor/Private/Dataflow/DataflowEditorViewport.cpp b/Engine/Plugins/Experimental/Dataflow/Source/DataflowEditor/Private/Dataflow/DataflowEditorViewport.cpp
--- a/Engine/Plugins/Experimental/Dataflow/Source/DataflowEditor/Private/Dataflow/DataflowEditorViewport.cpp
+++ b/Engine/Plugins/Experimental/Dataflow/Source/DataflowEditor/Private/Dataflow/DataflowEditorViewport.cpp
@@ -89,6 +89,20 @@ void SDataflowEditorViewport::BindCommands()
 			,FCanExecuteAction::CreateSP(ClientRef, &FDataflowEditorViewportClient::CanSetSelectionMode, FDataflowSelectionState::EMode::DSS_Dataflow_Object)
 			,FIsActionChecked::CreateSP(ClientRef, &FDataflowEditorViewportClient::IsSelectionModeActive, FDataflowSelectionState::EMode::DSS_Dataflow_Object)
 		);
+
+		CommandList->MapAction(
+			Commands.ToggleFaceSelection,
+			FExecuteAction::CreateSP(ClientRef, &FDataflowEditorViewportClient::SetSelectionMode, FDataflowSelectionState::EMode::DSS_Dataflow_Face)
+			, FCanExecuteAction::CreateSP(ClientRef, &FDataflowEditorViewportClient::CanSetSelectionMode, FDataflowSelectionState::EMode::DSS_Dataflow_Face)
+			, FIsActionChecked::CreateSP(ClientRef, &FDataflowEditorViewportClient::IsSelectionModeActive, FDataflowSelectionState::EMode::DSS_Dataflow_Face)
+		);
+
+		CommandList->MapAction(
+			Commands.ToggleVertexSelection,
+			FExecuteAction::CreateSP(ClientRef, &FDataflowEditorViewportClient::SetSelectionMode, FDataflowSelectionState::EMode::DSS_Dataflow_Vertex)
+			, FCanExecuteAction::CreateSP(ClientRef, &FDataflowEditorViewportClient::CanSetSelectionMode, FDataflowSelectionState::EMode::DSS_Dataflow_Vertex)
+			, FIsActionChecked::CreateSP(ClientRef, &FDataflowEditorViewportClient::IsSelectionModeActive, FDataflowSelectionState::EMode::DSS_Dataflow_Vertex)
+		);
 	}
 }
 
@@ -109,6 +123,8 @@ FDataflowEditorViewportClient::FDataflowEditorViewportClient(FPreviewScene* InPr
 
 void FDataflowEditorViewportClient::SetSelectionMode(FDataflowSelectionState::EMode InState)
 {
+	FDataflowSelectionState State = DataflowActor->DataflowComponent->GetSelectionState();
+		
 	if (SelectionMode == InState)
 	{
 		SelectionMode = FDataflowSelectionState::EMode::DSS_Dataflow_None;
@@ -118,11 +134,14 @@ void FDataflowEditorViewportClient::SetSelectionMode(FDataflowSelectionState::EM
 		SelectionMode = InState;
 	}
 
+	State.Mode = SelectionMode;
+	DataflowActor->DataflowComponent->SetSelectionState(State);
+
 	if (SelectionMode == FDataflowSelectionState::EMode::DSS_Dataflow_None)
 	{
 		if (!DataflowActor->DataflowComponent->GetSelectionState().IsEmpty())
 		{
-			DataflowActor->DataflowComponent->SetSelectionState(FDataflowSelectionState());
+			DataflowActor->DataflowComponent->SetSelectionState(FDataflowSelectionState(SelectionMode));
 		}
 	}
 }
@@ -135,16 +154,20 @@ bool FDataflowEditorViewportClient::CanSetSelectionMode(FDataflowSelectionState:
 		{
 			if (Dataflow->GetRenderTargets().Num())
 			{
-				return true;
+				if (InState == FDataflowSelectionState::EMode::DSS_Dataflow_Object)
+				{
+					return true;
+				}
+
+				if (InState == FDataflowSelectionState::EMode::DSS_Dataflow_Vertex
+					&& !DataflowActor->DataflowComponent->GetSelectionState().Nodes.IsEmpty())
+				{
+					return true;
+				}
 			}
 		}
 	}
 
-	if (!DataflowActor->DataflowComponent->GetSelectionState().IsEmpty())
-	{
-		DataflowActor->DataflowComponent->SetSelectionState(FDataflowSelectionState());
-	}
-
 	return false;
 }
 bool FDataflowEditorViewportClient::IsSelectionModeActive(FDataflowSelectionState::EMode InState) 
@@ -161,6 +184,7 @@ Dataflow::FTimestamp FDataflowEditorViewportClient::LatestTimestamp(const UDataf
 	return Dataflow::FTimestamp::Invalid;
 }
 
+
 void FDataflowEditorViewportClient::ProcessClick(FSceneView& View, HHitProxy* HitProxy, FKey Key, EInputEvent Event, uint32 HitX, uint32 HitY)
 {
 	Super::ProcessClick(View, HitProxy, Key, Event, HitX, HitY);
@@ -205,6 +229,38 @@ void FDataflowEditorViewportClient::ProcessClick(FSceneView& View, HHitProxy* Hi
 			}
 		}
 
+		if (SelectionMode == FDataflowSelectionState::EMode::DSS_Dataflow_Vertex)
+		{
+			if (HitProxy && HitProxy->IsA(HDataflowVertex::StaticGetType()))
+			{
+				HDataflowVertex* DataflowVertex = (HDataflowVertex*)(HitProxy);
+				int32 ID = DataflowVertex->SectionIndex;
+				if (bIsShiftKeyDown)
+				{
+					if (!SelectionState.Vertices.Contains(ID))
+					{
+						SelectionState.Vertices.AddUnique(ID);
+					}
+				}
+				else if (bIsCtrltKeyDown)
+				{
+					if (SelectionState.Vertices.Contains(ID))
+					{
+						SelectionState.Vertices.Remove(ID);
+					}
+				}
+				else
+				{
+					SelectionState.Vertices.Empty();
+					SelectionState.Vertices.AddUnique(ID);
+				}
+			}
+			else if (!bIsShiftKeyDown && !bIsCtrltKeyDown)
+			{
+				SelectionState.Vertices.Empty();
+			}
+		}
+
 		if (PreState != SelectionState)
 		{
 			DataflowActor->DataflowComponent->SetSelectionState(SelectionState);
diff --git a/Engine/Plugins/Experimental/Dataflow/Source/DataflowEditor/Private/Dataflow/DataflowEditorViewportToolbar.cpp b/Engine/Plugins/Experimental/Dataflow/Source/DataflowEditor/Private/Dataflow/DataflowEditorViewportToolbar.cpp
--- a/Engine/Plugins/Experimental/Dataflow/Source/DataflowEditor/Private/Dataflow/DataflowEditorViewportToolbar.cpp
+++ b/Engine/Plugins/Experimental/Dataflow/Source/DataflowEditor/Private/Dataflow/DataflowEditorViewportToolbar.cpp
@@ -2,11 +2,11 @@
 #include "Dataflow/DataflowEditorViewportToolbar.h"
 
 #include "Dataflow/DataflowEditorCommands.h"
+#include "Dataflow/DataflowEditorStyle.h"
 #include "Dataflow/DataflowEditorViewport.h"
 #include "EditorViewportCommands.h"
 #include "Widgets/Layout/SBorder.h"
 
-
 void SDataflowViewportSelectionToolBar::Construct(const FArguments& InArgs)
 {
 	EditorViewport = InArgs._EditorViewport;
@@ -50,10 +50,26 @@ TSharedRef<SWidget> SDataflowViewportSelectionToolBar::MakeSelectionModeToolBar(
 			NAME_None,
 			TAttribute<FText>(),
 			TAttribute<FText>(),
-			FSlateIcon(FAppStyle::GetAppStyleSetName(), TEXT("Icons.SelectInViewport")),
+			FSlateIcon(FDataflowEditorStyle::Get().GetStyleSetName(), TEXT("Dataflow.SelectObject")),
 			FName(TEXT("SelectObjectMode"))
 		);
 
+		ToolbarBuilder.AddToolBarButton(FDataflowEditorCommands::Get().ToggleFaceSelection,
+										NAME_None,
+										TAttribute<FText>(),
+										TAttribute<FText>(),
+										FSlateIcon(FDataflowEditorStyle::Get().GetStyleSetName(), TEXT("Dataflow.SelectFace")),
+										FName(TEXT("SelectFaceMode"))
+		);
+
+		ToolbarBuilder.AddToolBarButton(FDataflowEditorCommands::Get().ToggleVertexSelection,
+										NAME_None,
+										TAttribute<FText>(),
+										TAttribute<FText>(),
+										FSlateIcon(FDataflowEditorStyle::Get().GetStyleSetName(), TEXT("Dataflow.SelectVertex")),
+										FName(TEXT("SelectVertexMode"))
+		);
+
 
 		ToolbarBuilder.EndBlockGroup();
 	}
diff --git a/Engine/Plugins/Experimental/Dataflow/Source/DataflowEditor/Public/Dataflow/DataflowEditorCommands.h b/Engine/Plugins/Experimental/Dataflow/Source/DataflowEditor/Public/Dataflow/DataflowEditorCommands.h
--- a/Engine/Plugins/Experimental/Dataflow/Source/DataflowEditor/Public/Dataflow/DataflowEditorCommands.h
+++ b/Engine/Plugins/Experimental/Dataflow/Source/DataflowEditor/Public/Dataflow/DataflowEditorCommands.h
@@ -41,6 +41,8 @@ public:
 	TSharedPtr< FUICommandInfo > CreateComment;
 	TSharedPtr< FUICommandInfo > ToggleEnabledState;
 	TSharedPtr<FUICommandInfo> ToggleObjectSelection;
+	TSharedPtr<FUICommandInfo> ToggleFaceSelection;
+	TSharedPtr<FUICommandInfo> ToggleVertexSelection;
 
 	TMap< FName, TSharedPtr<FUICommandInfo> > CreateNodesMap;
 };
diff --git a/Engine/Plugins/Experimental/Dataflow/Source/DataflowEditor/Public/Dataflow/DataflowEditorViewport.h b/Engine/Plugins/Experimental/Dataflow/Source/DataflowEditor/Public/Dataflow/DataflowEditorViewport.h
--- a/Engine/Plugins/Experimental/Dataflow/Source/DataflowEditor/Public/Dataflow/DataflowEditorViewport.h
+++ b/Engine/Plugins/Experimental/Dataflow/Source/DataflowEditor/Public/Dataflow/DataflowEditorViewport.h
@@ -100,6 +100,6 @@ private:
 	// Selection
 	//
 
-	FDataflowSelectionState::EMode SelectionMode;
+	FDataflowSelectionState::EMode SelectionMode = FDataflowSelectionState::EMode::DSS_Dataflow_None;
 
 };
diff --git a/Engine/Plugins/Experimental/Dataflow/Source/DataflowEnginePlugin/Private/Dataflow/DataflowComponent.cpp b/Engine/Plugins/Experimental/Dataflow/Source/DataflowEnginePlugin/Private/Dataflow/DataflowComponent.cpp
--- a/Engine/Plugins/Experimental/Dataflow/Source/DataflowEnginePlugin/Private/Dataflow/DataflowComponent.cpp
+++ b/Engine/Plugins/Experimental/Dataflow/Source/DataflowEnginePlugin/Private/Dataflow/DataflowComponent.cpp
@@ -30,6 +30,7 @@ void UDataflowComponent::Invalidate()
 	bUpdateRender = true;
 	bUpdateSelection = true;
 	bBoundsNeedsUpdate = true;
+	SelectionState.Vertices.Empty();
 	RenderCollection = FManagedArrayCollection();
 }
 
@@ -54,6 +55,7 @@ void UDataflowComponent::TickComponent(float DeltaTime, enum ELevelTick TickType
 
 	if (bUpdateRender)
 	{
+		SelectionState.Vertices.Empty();
 		RenderCollection = FManagedArrayCollection();
 		GeometryCollection::Facades::FRenderingFacade Facade(RenderCollection);
 		Facade.DefineSchema();
diff --git a/Engine/Plugins/Experimental/Dataflow/Source/DataflowEnginePlugin/Private/Dataflow/DataflowComponentSelectionState.cpp b/Engine/Plugins/Experimental/Dataflow/Source/DataflowEnginePlugin/Private/Dataflow/DataflowComponentSelectionState.cpp
--- a/Engine/Plugins/Experimental/Dataflow/Source/DataflowEnginePlugin/Private/Dataflow/DataflowComponentSelectionState.cpp
+++ b/Engine/Plugins/Experimental/Dataflow/Source/DataflowEnginePlugin/Private/Dataflow/DataflowComponentSelectionState.cpp
@@ -12,38 +12,62 @@ void FDataflowSelectionState::UpdateSelection(UDataflowComponent* DataflowCompon
 	FRenderingFacade Facade(DataflowComponent->ModifyRenderingCollection());
 	if (Facade.IsValid())
 	{
-		TManagedArray<int32>& SelectionArray = Facade.ModifySelectionState();
-		if (Nodes.Num())
+		TManagedArray<int32>& ObjectSelectionArray = Facade.ModifySelectionState();
+		TManagedArray<int32>& VertexSelectionArray = Facade.ModifyVertexSelection();
+
+		if (Mode == EMode::DSS_Dataflow_Object)
 		{
-			const TManagedArray<FString>& GeomNameArray = Facade.GetGeometryNameAttribute();
+			if (Nodes.Num())
+			{
+				const TManagedArray<FString>& GeomNameArray = Facade.GetGeometryName();
 
-			TArray<bool> bVisited;
-			bVisited.Init(false, Nodes.Num());
+				TArray<bool> bVisited;
+				bVisited.Init(false, Nodes.Num());
 
-			for (int i = 0; i < SelectionArray.Num(); i++)
-			{
-				ObjectID ID(GeomNameArray[i], i);
-				SelectionArray[i] = false;
-				int32 IndexOf = Nodes.IndexOfByKey(ID);
-				if (IndexOf != INDEX_NONE)
+				for (int i = 0; i < ObjectSelectionArray.Num(); i++)
+				{
+					ObjectID ID(GeomNameArray[i], i);
+					ObjectSelectionArray[i] = false;
+					int32 IndexOf = Nodes.IndexOfByKey(ID);
+					if (IndexOf != INDEX_NONE)
+					{
+						ObjectSelectionArray[i] = true;
+						bVisited[IndexOf] = true;
+					}
+				}
+
+				// remove unknown selections
+				for (int32 Ndx = bVisited.Num() - 1; 0 <= Ndx; Ndx--)
 				{
-					SelectionArray[i] = true;
-					bVisited[IndexOf] = true;
+					if (!bVisited[Ndx])
+					{
+						Nodes.RemoveAt(Ndx);
+					}
 				}
 			}
+			else
+			{
+				ObjectSelectionArray.Fill(0);
+			}
+		}
+		else if (Mode == EMode::DSS_Dataflow_Vertex)
+		{
+			VertexSelectionArray.Fill(0);
 
-			// remove unknown selections
-			for (int32 Ndx = bVisited.Num() - 1; 0 <= Ndx; Ndx--)
+			int32 NumVertices = VertexSelectionArray.Num();
+			for (int32 Index : Vertices)
 			{
-				if (!bVisited[Ndx])
+				if (Index < NumVertices)
 				{
-					Nodes.RemoveAt(Ndx);
+					VertexSelectionArray[Index] = 1;
 				}
 			}
 		}
 		else
 		{
-			SelectionArray.Fill(0);
+			VertexSelectionArray.Fill(0);
+			ObjectSelectionArray.Fill(0);
+
 		}
 	}
 }
diff --git a/Engine/Plugins/Experimental/Dataflow/Source/DataflowEnginePlugin/Private/Dataflow/DataflowEnginePlugin.cpp b/Engine/Plugins/Experimental/Dataflow/Source/DataflowEnginePlugin/Private/Dataflow/DataflowEnginePlugin.cpp
--- a/Engine/Plugins/Experimental/Dataflow/Source/DataflowEnginePlugin/Private/Dataflow/DataflowEnginePlugin.cpp
+++ b/Engine/Plugins/Experimental/Dataflow/Source/DataflowEnginePlugin/Private/Dataflow/DataflowEnginePlugin.cpp
@@ -12,8 +12,10 @@
 #include "GeometryCollection/GeometryCollectionAlgo.h"
 
 
-FColor IDataflowEnginePlugin::SurfaceColor = FLinearColor(0.6,0.6,0.6).ToRGBE();
-FColor IDataflowEnginePlugin::PrimarySelectionColor = FLinearColor(0.8, 0.4, 0.0).ToRGBE();
+FColor IDataflowEnginePlugin::SurfaceColor = FLinearColor(0.6, 0.6, 0.6).ToRGBE();
+FColor IDataflowEnginePlugin::VertexColor = FLinearColor(0.0,0.0,0.0).ToRGBE();
+FColor IDataflowEnginePlugin::SelectionPrimaryColor = FLinearColor(0.8, 0.4, 0.0).ToRGBE();
+FColor IDataflowEnginePlugin::SelectionLockedPrimaryColor = FLinearColor(0.8, 0.4, 0.2).ToRGBE();
 
 class FDataflowEnginePlugin : public IDataflowEnginePlugin
 {
diff --git a/Engine/Plugins/Experimental/Dataflow/Source/DataflowEnginePlugin/Private/Dataflow/DataflowEngineSceneHitProxies.cpp b/Engine/Plugins/Experimental/Dataflow/Source/DataflowEnginePlugin/Private/Dataflow/DataflowEngineSceneHitProxies.cpp
--- a/Engine/Plugins/Experimental/Dataflow/Source/DataflowEnginePlugin/Private/Dataflow/DataflowEngineSceneHitProxies.cpp
+++ b/Engine/Plugins/Experimental/Dataflow/Source/DataflowEnginePlugin/Private/Dataflow/DataflowEngineSceneHitProxies.cpp
@@ -2,4 +2,6 @@
 
 #include "Dataflow/DataflowEngineSceneHitProxies.h"
 
+IMPLEMENT_HIT_PROXY(HDataflowDefault, HActor);
 IMPLEMENT_HIT_PROXY(HDataflowNode, HActor);
+IMPLEMENT_HIT_PROXY(HDataflowVertex, HActor);
diff --git a/Engine/Plugins/Experimental/Dataflow/Source/DataflowEnginePlugin/Private/Dataflow/DataflowEngineSceneProxy.cpp b/Engine/Plugins/Experimental/Dataflow/Source/DataflowEnginePlugin/Private/Dataflow/DataflowEngineSceneProxy.cpp
--- a/Engine/Plugins/Experimental/Dataflow/Source/DataflowEnginePlugin/Private/Dataflow/DataflowEngineSceneProxy.cpp
+++ b/Engine/Plugins/Experimental/Dataflow/Source/DataflowEnginePlugin/Private/Dataflow/DataflowEngineSceneProxy.cpp
@@ -13,28 +13,136 @@
 #include "PrimitiveSceneProxy.h"
 #include "StaticMeshResources.h"
 
+const TArray<FVector3f> BoxVertexArray = {
+	{-1.0, -1.0, -1.0}, {1.0, -1.0, -1.0}, {1.0, -1.0, 1.0}, {-1.0, -1.0, 1.0},
+	{-1.0, 1.0, -1.0}, {1.0, 1.0, -1.0}, {1.0, 1.0, 1.0}, {-1.0, 1.0, 1.0},
+};
+
+const TArray<FIntVector> BoxIndicesArray = {
+	{1, 5, 4}, {2, 6, 5}, {3, 7, 6}, {0, 4, 7}, {2, 1, 0}, {5, 6, 7},
+	{7, 4, 5}, {0, 3, 2}, {7, 3, 0}, {6, 2, 3}, {5, 1, 2}, {4, 0, 1}
+};
+
+
 FDataflowEngineSceneProxy::FDataflowEngineSceneProxy(UDataflowComponent* Component)
 	: FPrimitiveSceneProxy(Component)
-	, MaterialRelevance(Component->GetMaterialRelevance(GetScene().GetFeatureLevel()))
-	, VertexFactory(GetScene().GetFeatureLevel(), "FTriangleSetSceneProxy")
-	, RenderMaterial(Component->GetMaterial(0))
+	, DataflowComponent(Component)
 	, ConstantData(Component->GetRenderingCollection().NewCopy())
-{}
+	, RenderMaterial(UMaterial::GetDefaultMaterial(MD_Surface))
+	, MaterialRelevance(Component->GetMaterialRelevance(GetScene().GetFeatureLevel()))
+	, VertexFactory(GetScene().GetFeatureLevel(), "FDataflowEngineMeshVertexFactory")
+	, BoxVertexFactory(GetScene().GetFeatureLevel(), "FDataflowEngineBoxInstanceVertexFactory")
+{
+	RenderMaterial = GetRenderMaterial();
+}
 
 FDataflowEngineSceneProxy::~FDataflowEngineSceneProxy() 
+{}
+
+UMaterialInterface* FDataflowEngineSceneProxy::GetRenderMaterial() const
 {
-	if (ConstantData) delete ConstantData;
+	UMaterialInterface* RetRenderMaterial = DataflowComponent->GetMaterial(0);
+	if (RetRenderMaterial == nullptr)
+	{
+		RetRenderMaterial = UMaterial::GetDefaultMaterial(MD_Surface);
+	}
+	return RetRenderMaterial;
 }
 
-
 void FDataflowEngineSceneProxy::CreateRenderThreadResources()
 {
-	static int foobar = INDEX_NONE;
-
 	check(ConstantData);
+	check(RenderMaterial);
+	check(IsInRenderingThread());
+
+#if WITH_EDITOR
+	SetUsedMaterialForVerification({RenderMaterial});
+#endif
+
+	CreateInstancedVertexRenderThreadResources();
+	CreateMeshRenderThreadResources();
+}
+
+
+void FDataflowEngineSceneProxy::DestroyRenderThreadResources()
+{
+	DestroyInstancedVertexRenderThreadResources();
+	DestroyMeshRenderThreadResources();
+
+	delete ConstantData;
+	ConstantData = nullptr;
+}
+
+
+#if WITH_EDITOR
+HHitProxy* FDataflowEngineSceneProxy::CreateHitProxies(UPrimitiveComponent* Component, TArray<TRefCountPtr<HHitProxy> >& OutHitProxies)
+{
+	check(!IsInRenderingThread());
+
+	GeometryCollection::Facades::FRenderingFacade Facade(*ConstantData);
+	if (Facade.NumTriangles() > 0)
+	{
+		OutHitProxies.Add(new HDataflowDefault(Component->GetOwner(), Component));
+		LocalHitProxies.Add(OutHitProxies.Last());
+
+		const FDataflowSelectionState& State = DataflowComponent->GetSelectionState();
+		if (State.Mode == FDataflowSelectionState::EMode::DSS_Dataflow_Object)
+		{
+			const TManagedArray<FString>& NodeName = Facade.GetGeometryName();
+			TManagedArray<int32>& GeometryHitProxyIndex = Facade.ModifyGeometryHitProxyIndex();
+
+			int32 NumGeom = Facade.NumGeometry();
+			for (int i = 0; i < NumGeom; i++)
+			{
+				OutHitProxies.Add(new HDataflowNode(Component->GetOwner(), Component, NodeName[i], i));
+				LocalHitProxies.Add(OutHitProxies.Last());
+				GeometryHitProxyIndex[i] = LocalHitProxies.Num() - 1;
+			}
+		}
+		else if (State.Mode == FDataflowSelectionState::EMode::DSS_Dataflow_Vertex)
+		{
+			TManagedArray<int32>& VertexHitProxyIndex = Facade.ModifyVertexHitProxyIndex();
+
+			int32 NumVerts = Facade.NumVertices();
+			for (int i = 0; i < NumVerts; i++)
+			{
+				OutHitProxies.Add(new HDataflowVertex(Component->GetOwner(), Component, i));
+				LocalHitProxies.Add(OutHitProxies.Last());
+				VertexHitProxyIndex[i] = LocalHitProxies.Num() - 1;
+			}
+		}
+	}
+	return FPrimitiveSceneProxy::CreateHitProxies(Component, OutHitProxies);
+}
+
+#endif // WITH_EDITOR
+
+void FDataflowEngineSceneProxy::GetDynamicMeshElements(const TArray<const FSceneView*>& Views, const FSceneViewFamily& ViewFamily, uint32 VisibilityMap,
+													   FMeshElementCollector& Collector) const
+{
+	QUICK_SCOPE_CYCLE_COUNTER(STAT_OverlaySceneProxy_GetDynamicMeshElements);
+
+	check(RenderMaterial);
 	check(IsInRenderingThread());
 
+	for (int32 ViewIndex = 0; ViewIndex < Views.Num(); ViewIndex++)
+	{
+		if (VisibilityMap & (1 << ViewIndex))
+		{
+			GetMeshDynamicMeshElements(ViewIndex, Collector);
+			GetDynamicInstancedVertexMeshElements(ViewIndex, Collector);
+		}
+	}
+}
+
+//
+// Mesh Selection Rendering
+//
+
+void FDataflowEngineSceneProxy::CreateMeshRenderThreadResources()
+{
 	GeometryCollection::Facades::FRenderingFacade Facade(*ConstantData);
+	const FDataflowSelectionState& State = DataflowComponent->GetSelectionState();
 	check(Facade.CanRenderSurface());
 
 	const int32 NumTriangleVertices = Facade.NumTriangles();
@@ -56,29 +164,26 @@ void FDataflowEngineSceneProxy::CreateRenderThreadResources()
 	if (Facade.NumTriangles() > 0)
 	{
 		int32 NumGeom = Facade.NumGeometry();
-		for (int i = 0; i<NumGeom; i++)
+		for (int i = 0; i < NumGeom; i++)
 		{
-			MeshBatchDatas.Emplace();
-			FDataflowTriangleSetMeshBatchData& MeshBatchData = MeshBatchDatas.Last();
-
-			MeshBatchData.MinVertexIndex = Facade.GetIndicesStartAttribute()[i] * 3;
-			MeshBatchData.MaxVertexIndex = MeshBatchData.MinVertexIndex + Facade.GetIndicesCountAttribute()[i] * 3 - 1;
-			MeshBatchData.FirstTriangleIndex = Facade.GetIndicesStartAttribute()[i] * 3;
-			MeshBatchData.NumTriangles = Facade.GetIndicesCountAttribute()[i];
-			MeshBatchData.GeomIndex = i;
-
-			if (!RenderMaterial)
+			if (Facade.GetIndicesCount()[i] != 0)
 			{
-				RenderMaterial = UMaterial::GetDefaultMaterial(MD_Surface);
+				MeshBatchDatas.Emplace();
+				FDataflowTriangleSetMeshBatchData& MeshBatchData = MeshBatchDatas.Last();
+
+				MeshBatchData.MinVertexIndex = Facade.GetIndicesStart()[i] * 3;
+				MeshBatchData.MaxVertexIndex = MeshBatchData.MinVertexIndex + Facade.GetIndicesCount()[i] * 3 - 1;
+				MeshBatchData.FirstTriangleIndex = Facade.GetIndicesStart()[i] * 3;
+				MeshBatchData.NumTriangles = Facade.GetIndicesCount()[i];
+				MeshBatchData.GeomIndex = i;
 			}
-			MeshBatchData.MaterialProxy = RenderMaterial->GetRenderProxy();
 		}
 
 
 		const TManagedArray<FIntVector>& Indices = Facade.GetIndices();
 		const TManagedArray<FVector3f>& Vertex = Facade.GetVertices();
 		const TManagedArray<int32>& SelectionArray = Facade.GetSelectionState();
-		const TManagedArray<int32>& GeomIndex = Facade.GetVertexToGeometryIndexAttribute();
+		const TManagedArray<int32>& GeomIndex = Facade.GetVertexToGeometryIndex();
 
 		// The color stored in the vertices actually gets interpreted as a linear
 		// color by the material, whereas it is more convenient for the user of the
@@ -86,8 +191,8 @@ void FDataflowEngineSceneProxy::CreateRenderThreadResources()
 		// back to linear. The ToFColor(false) call just scales back into 0-255
 		// space.
 		ParallelFor(Facade.NumTriangles(), [&](int32 i)
-		{
-		const int32 VertexBufferIndex = 3 * i;
+					{
+						const int32 VertexBufferIndex = 3 * i;
 		const int32 IndexBufferIndex = 3 * i;
 
 		const auto& P1 = Vertex[Indices[i][0]];
@@ -110,9 +215,8 @@ void FDataflowEngineSceneProxy::CreateRenderThreadResources()
 		VertexBuffers.StaticMeshVertexBuffer.SetVertexUV(VertexBufferIndex + 1, 0, FVector2f(0, 0));
 		VertexBuffers.StaticMeshVertexBuffer.SetVertexUV(VertexBufferIndex + 2, 0, FVector2f(0, 0));
 
-
-
-		FColor FaceColor = SelectionArray[GeomIndex[Indices[i][0]]] ? IDataflowEnginePlugin::PrimarySelectionColor : IDataflowEnginePlugin::SurfaceColor;
+		FColor SelectionColor = (State.Mode == FDataflowSelectionState::EMode::DSS_Dataflow_Object) ? IDataflowEnginePlugin::SelectionPrimaryColor : IDataflowEnginePlugin::SelectionLockedPrimaryColor;
+		FColor FaceColor = SelectionArray[GeomIndex[Indices[i][0]]] ? SelectionColor : IDataflowEnginePlugin::SurfaceColor;
 		VertexBuffers.ColorVertexBuffer.VertexColor(VertexBufferIndex + 0) = FaceColor;
 		VertexBuffers.ColorVertexBuffer.VertexColor(VertexBufferIndex + 1) = FaceColor;
 		VertexBuffers.ColorVertexBuffer.VertexColor(VertexBufferIndex + 2) = FaceColor;
@@ -121,12 +225,10 @@ void FDataflowEngineSceneProxy::CreateRenderThreadResources()
 		IndexBuffer.Indices[IndexBufferIndex + 1] = VertexBufferIndex + 1;
 		IndexBuffer.Indices[IndexBufferIndex + 2] = VertexBufferIndex + 2;
 
-		});
+					});
 
 	}
-#if WITH_EDITOR
-	SetUsedMaterialForVerification({RenderMaterial});
-#endif
+
 
 	VertexBuffers.PositionVertexBuffer.InitResource();
 	VertexBuffers.StaticMeshVertexBuffer.InitResource();
@@ -143,7 +245,9 @@ void FDataflowEngineSceneProxy::CreateRenderThreadResources()
 	IndexBuffer.InitResource();
 }
 
-void FDataflowEngineSceneProxy::DestroyRenderThreadResources()
+
+
+void FDataflowEngineSceneProxy::DestroyMeshRenderThreadResources()
 {
 	check(IsInRenderingThread());
 	VertexBuffers.PositionVertexBuffer.ReleaseResource();
@@ -151,71 +255,199 @@ void FDataflowEngineSceneProxy::DestroyRenderThreadResources()
 	VertexBuffers.ColorVertexBuffer.ReleaseResource();
 	IndexBuffer.ReleaseResource();
 	VertexFactory.ReleaseResource();
-	delete ConstantData;
-	ConstantData = nullptr;
-}
 
+}
 
-#if WITH_EDITOR
-HHitProxy* FDataflowEngineSceneProxy::CreateHitProxies(UPrimitiveComponent* Component, TArray<TRefCountPtr<HHitProxy> >& OutHitProxies)
+void FDataflowEngineSceneProxy::GetMeshDynamicMeshElements(int32 ViewIndex, FMeshElementCollector& Collector) const
 {
-	check(!IsInRenderingThread());
+	const GeometryCollection::Facades::FRenderingFacade Facade(*ConstantData);
+	const TManagedArray<int32>& GeometryHitProxyIndex = Facade.GetGeometryHitProxyIndex();
+	const FDataflowSelectionState& State = DataflowComponent->GetSelectionState();
 
-	GeometryCollection::Facades::FRenderingFacade Facade(*ConstantData);
-	if (Facade.NumTriangles() > 0)
+	for (const FDataflowTriangleSetMeshBatchData& MeshBatchData : MeshBatchDatas)
 	{
-		int32 NumGeom = Facade.NumGeometry();
-		const TManagedArray<FString>& NodeName = Facade.GetGeometryNameAttribute();
-		for (int i = 0; i < NumGeom; i++)
+		FDynamicPrimitiveUniformBuffer& DynamicPrimitiveUniformBuffer = Collector.AllocateOneFrameResource<FDynamicPrimitiveUniformBuffer>();
+		DynamicPrimitiveUniformBuffer.Set(GetLocalToWorld(), GetLocalToWorld(), GetBounds(), GetLocalBounds(), true, false, AlwaysHasVelocity());
+
+		FMeshBatch& Mesh = Collector.AllocateMesh();
+		Mesh.bWireframe = false;
+		Mesh.VertexFactory = &VertexFactory;
+		Mesh.MaterialRenderProxy = RenderMaterial->GetRenderProxy();
+		Mesh.ReverseCulling = IsLocalToWorldDeterminantNegative();
+		Mesh.Type = PT_TriangleList;
+		Mesh.DepthPriorityGroup = SDPG_World;
+		Mesh.bCanApplyViewModeOverrides = true;
+
+		FMeshBatchElement& BatchElement = Mesh.Elements[0];
+		BatchElement.IndexBuffer = &IndexBuffer;
+		BatchElement.PrimitiveUniformBufferResource = &DynamicPrimitiveUniformBuffer.UniformBuffer;
+		BatchElement.FirstIndex = MeshBatchData.FirstTriangleIndex;
+		BatchElement.NumPrimitives = MeshBatchData.NumTriangles;
+		BatchElement.MinVertexIndex = MeshBatchData.MinVertexIndex;
+		BatchElement.MaxVertexIndex = MeshBatchData.MaxVertexIndex;
+
+#if WITH_EDITOR
+		if (State.Mode == FDataflowSelectionState::EMode::DSS_Dataflow_Object)
 		{
-			HDataflowNode * NodeProxy = new HDataflowNode(Component->GetOwner(), Component, NodeName[i], i);
-			LocalHitProxies.Add(NodeProxy);
-			Facade.ModifyHitProxyIndexAttribute()[i] = LocalHitProxies.Num() - 1;
-			OutHitProxies.Add(NodeProxy);
+			Mesh.BatchHitProxyId = LocalHitProxies[GeometryHitProxyIndex[MeshBatchData.GeomIndex]]->Id;
 		}
+		else
+		{
+			Mesh.BatchHitProxyId = LocalHitProxies[0]->Id;
+		}
+#endif
+		Collector.AddMesh(ViewIndex, Mesh);
 	}
-	return FPrimitiveSceneProxy::CreateHitProxies(Component, OutHitProxies);
 }
 
-#endif // WITH_EDITOR
 
+//
+// Vertex Selection Rendering
+//
 
-void FDataflowEngineSceneProxy::GetDynamicMeshElements( 
-	const TArray<const FSceneView*>& Views, const FSceneViewFamily& ViewFamily,
-	uint32 VisibilityMap, FMeshElementCollector& Collector) const
+void FDataflowEngineSceneProxy::CreateInstancedVertexRenderThreadResources()
+{
+	const FDataflowSelectionState& State = DataflowComponent->GetSelectionState();
+	if (State.Mode == FDataflowSelectionState::EMode::DSS_Dataflow_Vertex)
+	{
+		const bool bCanUseGPUScene = UseGPUScene(GetScene().GetShaderPlatform(), GetScene().GetFeatureLevel());
+		if (bCanUseGPUScene)
+		{
+			GeometryCollection::Facades::FRenderingFacade Facade(*ConstantData);
+			const TManagedArray<FVector3f>& Vertex = Facade.GetVertices();
+			const TManagedArray<int32>& VertexSelection = Facade.GetVertexSelection();
+			const TManagedArray<int32>& GeometrySelection = Facade.GetSelectionState();
+			const TManagedArray<int32>& VertexToGeometryMap = Facade.GetVertexToGeometryIndex();
+
+			NumRenderedVerts = Facade.NumVerticesOnSelectedGeometry();
+			const int32 NumElements = BoxIndicesArray.Num() * 3;
+			const int32 TotalNumElements = NumElements * NumRenderedVerts;
+			constexpr int32 NumTextureCoordinates = 1;
+
+			if (NumRenderedVerts)
+			{
+				BoxVertexBuffers.PositionVertexBuffer.Init(TotalNumElements);
+				BoxVertexBuffers.StaticMeshVertexBuffer.Init(TotalNumElements, NumTextureCoordinates);
+				BoxVertexBuffers.ColorVertexBuffer.Init(TotalNumElements);
+				BoxIndexBuffer.Indices.SetNumUninitialized(TotalNumElements);
+
+				VertexBatchDatas.SetNum(NumRenderedVerts);
+				for (int32 VertIndex = 0, Vdx = 0; VertIndex < Facade.NumVertices(); VertIndex++)
+				{
+					if (GeometrySelection[VertexToGeometryMap[VertIndex]])
+					{
+						VertexBatchDatas[Vdx].FirstElementIndex = Vdx * NumElements;
+						VertexBatchDatas[Vdx].NumElements = BoxIndicesArray.Num();
+						VertexBatchDatas[Vdx].MinVertexIndex = Vdx * NumElements;
+						VertexBatchDatas[Vdx].MaxVertexIndex = Vdx * NumElements + NumElements - 1;
+						VertexBatchDatas[Vdx].VertexIndex = VertIndex;
+
+						for (int32 Idx = 0; Idx < BoxIndicesArray.Num(); Idx++)
+						{
+							const int32 BufferIndex = (Vdx * NumElements) + Idx * 3;
+
+							float Scale = 2.0;
+							const auto& P1 = BoxVertexArray[BoxIndicesArray[Idx][0]] * Scale + Vertex[VertIndex];
+							const auto& P2 = BoxVertexArray[BoxIndicesArray[Idx][1]] * Scale + Vertex[VertIndex];
+							const auto& P3 = BoxVertexArray[BoxIndicesArray[Idx][2]] * Scale + Vertex[VertIndex];
+
+							BoxVertexBuffers.PositionVertexBuffer.VertexPosition(BufferIndex + 0) = P1;
+							BoxVertexBuffers.PositionVertexBuffer.VertexPosition(BufferIndex + 1) = P2;
+							BoxVertexBuffers.PositionVertexBuffer.VertexPosition(BufferIndex + 2) = P3;
+
+							FVector3f Tangent1 = (P2 - P1).GetSafeNormal();
+							FVector3f Tangent2 = (P3 - P2).GetSafeNormal();
+							FVector3f Normal = (Tangent2 ^ Tangent1).GetSafeNormal();
+
+							BoxVertexBuffers.StaticMeshVertexBuffer.SetVertexTangents(BufferIndex + 0, FVector3f(1, 0, 0), FVector3f(0, 1, 0), Normal);
+							BoxVertexBuffers.StaticMeshVertexBuffer.SetVertexTangents(BufferIndex + 1, FVector3f(1, 0, 0), FVector3f(0, 1, 0), Normal);
+							BoxVertexBuffers.StaticMeshVertexBuffer.SetVertexTangents(BufferIndex + 2, FVector3f(1, 0, 0), FVector3f(0, 1, 0), Normal);
+
+							BoxVertexBuffers.StaticMeshVertexBuffer.SetVertexUV(BufferIndex + 0, 0, FVector2f(0, 0));
+							BoxVertexBuffers.StaticMeshVertexBuffer.SetVertexUV(BufferIndex + 1, 0, FVector2f(0, 0));
+							BoxVertexBuffers.StaticMeshVertexBuffer.SetVertexUV(BufferIndex + 2, 0, FVector2f(0, 0));
+
+							FColor FaceColor = VertexSelection[VertIndex] ? IDataflowEnginePlugin::SelectionPrimaryColor : IDataflowEnginePlugin::VertexColor;
+							BoxVertexBuffers.ColorVertexBuffer.VertexColor(BufferIndex + 0) = FaceColor;
+							BoxVertexBuffers.ColorVertexBuffer.VertexColor(BufferIndex + 1) = FaceColor;
+							BoxVertexBuffers.ColorVertexBuffer.VertexColor(BufferIndex + 2) = FaceColor;
+
+							BoxIndexBuffer.Indices[BufferIndex + 0] = BufferIndex + 0;
+							BoxIndexBuffer.Indices[BufferIndex + 1] = BufferIndex + 1;
+							BoxIndexBuffer.Indices[BufferIndex + 2] = BufferIndex + 2;
+						}
+						Vdx++;
+					}
+				}
+
+				BoxVertexBuffers.PositionVertexBuffer.InitResource();
+				BoxVertexBuffers.StaticMeshVertexBuffer.InitResource();
+				BoxVertexBuffers.ColorVertexBuffer.InitResource();
+
+				FLocalVertexFactory::FDataType Data;
+				BoxVertexBuffers.PositionVertexBuffer.BindPositionVertexBuffer(&VertexFactory, Data);
+				BoxVertexBuffers.StaticMeshVertexBuffer.BindTangentVertexBuffer(&VertexFactory, Data);
+				BoxVertexBuffers.StaticMeshVertexBuffer.BindTexCoordVertexBuffer(&VertexFactory, Data);
+				BoxVertexBuffers.ColorVertexBuffer.BindColorVertexBuffer(&VertexFactory, Data);
+				BoxVertexFactory.SetData(Data);
+
+				BoxVertexFactory.InitResource();
+				BoxIndexBuffer.InitResource();
+			}
+		}
+	}
+}
+
+void FDataflowEngineSceneProxy::DestroyInstancedVertexRenderThreadResources()
 {
-	QUICK_SCOPE_CYCLE_COUNTER(STAT_OverlaySceneProxy_GetDynamicMeshElements);
 	check(IsInRenderingThread());
-	const GeometryCollection::Facades::FRenderingFacade Facade(*ConstantData);
+	if (NumRenderedVerts)
+	{
+		BoxVertexBuffers.PositionVertexBuffer.ReleaseResource();
+		BoxVertexBuffers.StaticMeshVertexBuffer.ReleaseResource();
+		BoxVertexBuffers.ColorVertexBuffer.ReleaseResource();
+		BoxIndexBuffer.ReleaseResource();
+		BoxVertexFactory.ReleaseResource();
+	}
+}
 
-	for (int32 ViewIndex = 0; ViewIndex < Views.Num(); ViewIndex++)
+void FDataflowEngineSceneProxy::GetDynamicInstancedVertexMeshElements(int32 ViewIndex, FMeshElementCollector& Collector) const
+{
+	if (NumRenderedVerts)
 	{
-		if (VisibilityMap & (1 << ViewIndex))
+		if (DataflowComponent->GetSelectionState().Mode == FDataflowSelectionState::EMode::DSS_Dataflow_Vertex)
 		{
-			for (const FDataflowTriangleSetMeshBatchData& MeshBatchData : MeshBatchDatas)
-			{
-				FMeshBatch& Mesh = Collector.AllocateMesh();
-				FMeshBatchElement& BatchElement = Mesh.Elements[0];
-				BatchElement.IndexBuffer = &IndexBuffer;
-				Mesh.bWireframe = false;
-				Mesh.VertexFactory = &VertexFactory;
-				Mesh.MaterialRenderProxy = MeshBatchData.MaterialProxy;
+			GeometryCollection::Facades::FRenderingFacade Facade(*ConstantData);
+			const TManagedArray<FVector3f>& Vertex = Facade.GetVertices();
+			const TManagedArray<int32>& VertexHitProxyIndex = Facade.GetVertexHitProxyIndex();
+			const int32 NumVertices = Facade.NumVertices();
+			const int32 NumElements = BoxIndicesArray.Num() * 3;
 
+			for (const FDataflowVertexBatchData& VertexBatchData : VertexBatchDatas)
+			{
 				FDynamicPrimitiveUniformBuffer& DynamicPrimitiveUniformBuffer = Collector.AllocateOneFrameResource<FDynamicPrimitiveUniformBuffer>();
-				DynamicPrimitiveUniformBuffer.Set(GetLocalToWorld(), GetLocalToWorld(),GetBounds(), GetLocalBounds(), true, false, AlwaysHasVelocity());
-				BatchElement.PrimitiveUniformBufferResource = &DynamicPrimitiveUniformBuffer.UniformBuffer;
-				//BatchElement.PrimitiveUniformBuffer = GetUniformBuffer();
+				DynamicPrimitiveUniformBuffer.Set(GetLocalToWorld(), GetLocalToWorld(), GetBounds(), GetLocalBounds(), true, false, AlwaysHasVelocity());
 
-				BatchElement.FirstIndex = MeshBatchData.FirstTriangleIndex;
-				BatchElement.NumPrimitives = MeshBatchData.NumTriangles;
-				BatchElement.MinVertexIndex = MeshBatchData.MinVertexIndex;
-				BatchElement.MaxVertexIndex = MeshBatchData.MaxVertexIndex;
+				FMeshBatch& Mesh = Collector.AllocateMesh();
+				Mesh.CastShadow = false;
+				Mesh.bUseAsOccluder = false;
+				Mesh.VertexFactory = &BoxVertexFactory;
+				Mesh.MaterialRenderProxy = RenderMaterial->GetRenderProxy();
 				Mesh.ReverseCulling = IsLocalToWorldDeterminantNegative();
 				Mesh.Type = PT_TriangleList;
 				Mesh.DepthPriorityGroup = SDPG_World;
 				Mesh.bCanApplyViewModeOverrides = true;
+
+				FMeshBatchElement& BatchElement = Mesh.Elements[0];
+				BatchElement.IndexBuffer = &BoxIndexBuffer;
+				BatchElement.PrimitiveUniformBufferResource = &DynamicPrimitiveUniformBuffer.UniformBuffer;
+				BatchElement.FirstIndex = VertexBatchData.FirstElementIndex;
+				BatchElement.NumPrimitives = VertexBatchData.NumElements;
+				BatchElement.MinVertexIndex = VertexBatchData.MinVertexIndex;
+				BatchElement.MaxVertexIndex = VertexBatchData.MaxVertexIndex;
+
 #if WITH_EDITOR
-				Mesh.BatchHitProxyId = LocalHitProxies[Facade.GetHitProxyIndexAttribute()[MeshBatchData.GeomIndex]]->Id;
+				Mesh.BatchHitProxyId = LocalHitProxies[VertexHitProxyIndex[VertexBatchData.VertexIndex]]->Id;
 #endif
 				Collector.AddMesh(ViewIndex, Mesh);
 			}
@@ -224,6 +456,13 @@ void FDataflowEngineSceneProxy::GetDynamicMeshElements(
 }
 
 
+
+
+bool FDataflowEngineSceneProxy::CanBeOccluded() const
+{
+	return false;
+}
+
 FPrimitiveViewRelevance FDataflowEngineSceneProxy::GetViewRelevance(const FSceneView* View) const
 {
 	FPrimitiveViewRelevance Result;
@@ -232,19 +471,13 @@ FPrimitiveViewRelevance FDataflowEngineSceneProxy::GetViewRelevance(const FScene
 	Result.bDynamicRelevance = true;
 	Result.bRenderInMainPass = ShouldRenderInMainPass();
 	Result.bUsesLightingChannels =GetLightingChannelMask() != GetDefaultLightingChannelMask();
-	Result.bRenderCustomDepth = false; // was ShouldRenderCustomDepth();
+	Result.bRenderCustomDepth = false;
 	Result.bTranslucentSelfShadow = bCastVolumetricTranslucentShadow;
 	MaterialRelevance.SetPrimitiveViewRelevance(Result);
-	Result.bVelocityRelevance = false;//was DrawsVelocity() && Result.bOpaque&& Result.bRenderInMainPass;
+	Result.bVelocityRelevance = false;
 	return Result;
 }
 
-
-bool FDataflowEngineSceneProxy::CanBeOccluded() const
-{
-	return false;//was !MaterialRelevance.bDisableDepthTest;
-}
-
 SIZE_T FDataflowEngineSceneProxy::GetTypeHash() const
 {
 	static SIZE_T UniquePointer;
diff --git a/Engine/Plugins/Experimental/Dataflow/Source/DataflowEnginePlugin/Private/Dataflow/DataflowEngineSceneProxy.h b/Engine/Plugins/Experimental/Dataflow/Source/DataflowEnginePlugin/Private/Dataflow/DataflowEngineSceneProxy.h
--- a/Engine/Plugins/Experimental/Dataflow/Source/DataflowEnginePlugin/Private/Dataflow/DataflowEngineSceneProxy.h
+++ b/Engine/Plugins/Experimental/Dataflow/Source/DataflowEnginePlugin/Private/Dataflow/DataflowEngineSceneProxy.h
@@ -2,10 +2,11 @@
 #pragma once
 
 #include "Async/ParallelFor.h"
+#include "Dataflow/DataflowEngineSceneHitProxies.h"
+#include "DynamicMeshBuilder.h"
 #include "Engine/CollisionProfile.h"
-#include "PrimitiveSceneProxy.h"
 #include "Materials/Material.h"
-#include "DynamicMeshBuilder.h"
+#include "PrimitiveSceneProxy.h"
 #include "StaticMeshResources.h"
 
 class UDataflowComponent;
@@ -13,7 +14,6 @@ struct FManagedArrayCollection;
 
 struct FDataflowTriangleSetMeshBatchData
 {
-	FMaterialRenderProxy* MaterialProxy = nullptr;
 	int32 FirstTriangleIndex = INDEX_NONE;
 	int32 NumTriangles = INDEX_NONE;
 	int32 MinVertexIndex = INDEX_NONE;
@@ -21,51 +21,87 @@ struct FDataflowTriangleSetMeshBatchData
 	int32 GeomIndex = INDEX_NONE;
 };
 
+struct FDataflowVertexBatchData
+{
+	int32 FirstElementIndex = INDEX_NONE;
+	int32 NumElements = INDEX_NONE;
+	int32 MinVertexIndex = INDEX_NONE;
+	int32 MaxVertexIndex = INDEX_NONE;
+	int32 VertexIndex = INDEX_NONE;
+};
+
 class FDataflowEngineSceneProxy final : public FPrimitiveSceneProxy
 {
 public:
 
 	FDataflowEngineSceneProxy(UDataflowComponent* Component);
-
 	virtual ~FDataflowEngineSceneProxy();
 
+	/** Create the rendering buffer resources */
+	void InitResources();
+
+	/** Return the rendering buffer resources */
+	void ReleaseResources();
 
 	//~ FPrimitiveSceneProxy
-	virtual void GetDynamicMeshElements(const TArray<const FSceneView*>& Views, const FSceneViewFamily& ViewFamily, uint32 VisibilityMap, FMeshElementCollector& Collector) const override;
 	virtual FPrimitiveViewRelevance GetViewRelevance(const FSceneView* View) const override;
 	virtual bool CanBeOccluded() const override;
 	virtual uint32 GetMemoryFootprint() const override { return sizeof(*this) + GetAllocatedSize(); }
 	virtual SIZE_T GetTypeHash() const override;
 	virtual void CreateRenderThreadResources() override;
 	virtual void DestroyRenderThreadResources() override;
-
+	virtual void GetDynamicMeshElements(const TArray<const FSceneView*>& Views, const FSceneViewFamily& ViewFamily, uint32 VisibilityMap, FMeshElementCollector& Collector) const override;
 #if WITH_EDITOR
 	virtual HHitProxy* CreateHitProxies(UPrimitiveComponent* Component, TArray<TRefCountPtr<HHitProxy> >& OutHitProxies) override;
-#endif // WITH_EDITOR
+#endif
+	//~ End FPrimitiveSceneProxy
 
+private:
 
-protected:
+	UDataflowComponent* DataflowComponent = nullptr;
+	FManagedArrayCollection* ConstantData = nullptr;
+	UMaterialInterface* RenderMaterial = nullptr;
+#if WITH_EDITOR
+	TArray<HHitProxy*> LocalHitProxies;
+#endif
+	UMaterialInterface* GetRenderMaterial() const;
 
-	/** Create the rendering buffer resources */
-	void InitResources();
 
-	/** Return the rendering buffer resources */
-	void ReleaseResources();
+	//
+	// Surface Rendering
+	//
 
-private:
 	TArray<FDataflowTriangleSetMeshBatchData> MeshBatchDatas;
 	FMaterialRelevance MaterialRelevance;
 	FLocalVertexFactory VertexFactory;
 	FStaticMeshVertexBuffers VertexBuffers;
 	FDynamicMeshIndexBuffer32 IndexBuffer;
-#if WITH_EDITOR
-	FColorVertexBuffer HitProxyIdBuffer;
-	TArray<TRefCountPtr<HHitProxy> > LocalHitProxies;
-#endif // WITH_EDITOR
 
-	// Render thread copy of data. 
-	UMaterialInterface* RenderMaterial = nullptr;
-	FManagedArrayCollection* ConstantData = nullptr;
+	/** Create rendering resources for surface selection data. */
+	void CreateMeshRenderThreadResources();
+	/** Destroy rendering resources for surface selection data. */
+	void DestroyMeshRenderThreadResources();
+	/** Build MeshElements for rendering surface selection data. */
+	void GetMeshDynamicMeshElements(int32 ViewIndex, FMeshElementCollector& Collector) const;
+
+
+	//
+	// Vertex Rendering
+	//
+
+	int32 NumRenderedVerts = 0;
+	TArray<FDataflowVertexBatchData> VertexBatchDatas;
+	FLocalVertexFactory BoxVertexFactory;
+	FStaticMeshVertexBuffers BoxVertexBuffers;
+	FDynamicMeshIndexBuffer32 BoxIndexBuffer;
+
+	/** Create rendering resources for vertex selection data. */
+	void CreateInstancedVertexRenderThreadResources();
+	/** Destroy rendering resources for vertex selection data. */
+	void DestroyInstancedVertexRenderThreadResources();
+	/** Build MeshElements for rendering vertex selection data. */
+	void GetDynamicInstancedVertexMeshElements(int32 ViewIndex, FMeshElementCollector& Collector) const;
+
 };
 
 
diff --git a/Engine/Plugins/Experimental/Dataflow/Source/DataflowEnginePlugin/Public/Dataflow/DataflowComponent.h b/Engine/Plugins/Experimental/Dataflow/Source/DataflowEnginePlugin/Public/Dataflow/DataflowComponent.h
--- a/Engine/Plugins/Experimental/Dataflow/Source/DataflowEnginePlugin/Public/Dataflow/DataflowComponent.h
+++ b/Engine/Plugins/Experimental/Dataflow/Source/DataflowEnginePlugin/Public/Dataflow/DataflowComponent.h
@@ -58,7 +58,7 @@ public:
 	const UDataflow* GetDataflow() const { return Dataflow; }
 
 	/* Selection */
-	FDataflowSelectionState GetSelectionState() const { return SelectionState; }
+	const FDataflowSelectionState& GetSelectionState() const { return SelectionState; }
 	void SetSelectionState(const FDataflowSelectionState& InState) 
 	{
 		bUpdateSelection = true;
@@ -76,6 +76,6 @@ private:
 	bool bUpdateSelection = true;
 	bool bBoundsNeedsUpdate = true;
 	FBoxSphereBounds BoundingBox = FBoxSphereBounds(ForceInitToZero);
-	FDataflowSelectionState SelectionState;
+	FDataflowSelectionState SelectionState = FDataflowSelectionState(FDataflowSelectionState::EMode::DSS_Dataflow_None);
 };
 
diff --git a/Engine/Plugins/Experimental/Dataflow/Source/DataflowEnginePlugin/Public/Dataflow/DataflowComponentSelectionState.h b/Engine/Plugins/Experimental/Dataflow/Source/DataflowEnginePlugin/Public/Dataflow/DataflowComponentSelectionState.h
--- a/Engine/Plugins/Experimental/Dataflow/Source/DataflowEnginePlugin/Public/Dataflow/DataflowComponentSelectionState.h
+++ b/Engine/Plugins/Experimental/Dataflow/Source/DataflowEnginePlugin/Public/Dataflow/DataflowComponentSelectionState.h
@@ -21,32 +21,37 @@ public:
 		int32 ID;
 	};
 
-	FDataflowSelectionState() {}
-	FDataflowSelectionState(const UDataflowComponent* DataflowComponent)
-	{
-		check(DataflowComponent);
-	}
-
 	enum EMode
 	{
 		DSS_Dataflow_None,
 		DSS_Dataflow_Object,
+		DSS_Dataflow_Face,
+		DSS_Dataflow_Vertex,
 		DSS_Dataflow_Max
 	};
 
+	FDataflowSelectionState(EMode InMode) : Mode(InMode) {}
+	FDataflowSelectionState(const UDataflowComponent* DataflowComponent)
+	{
+		check(DataflowComponent);
+	}
+
 	void UpdateSelection(UDataflowComponent* DataflowComponent);
 
 	bool IsEmpty() const
 	{
-		return Nodes.Num() == 0;
+		return Nodes.Num() == 0 
+			&& Vertices.Num() == 0;
 	}
 
 	bool operator==(const FDataflowSelectionState& A) const {
-		return A.Nodes == Nodes;
+		return A.Nodes == Nodes && A.Vertices == Vertices;
 	}
 	bool operator!=(const FDataflowSelectionState& A) const {
 		return !this->operator==(A);
 	}
 
+	EMode Mode;
 	TArray<ObjectID> Nodes;
+	TArray<int32> Vertices;
 };
\ No newline at end of file
diff --git a/Engine/Plugins/Experimental/Dataflow/Source/DataflowEnginePlugin/Public/Dataflow/DataflowEnginePlugin.h b/Engine/Plugins/Experimental/Dataflow/Source/DataflowEnginePlugin/Public/Dataflow/DataflowEnginePlugin.h
--- a/Engine/Plugins/Experimental/Dataflow/Source/DataflowEnginePlugin/Public/Dataflow/DataflowEnginePlugin.h
+++ b/Engine/Plugins/Experimental/Dataflow/Source/DataflowEnginePlugin/Public/Dataflow/DataflowEnginePlugin.h
@@ -16,7 +16,9 @@ class IDataflowEnginePlugin : public IModuleInterface
 public:
 
 	static FColor SurfaceColor;
-	static FColor PrimarySelectionColor;
+	static FColor VertexColor;
+	static FColor SelectionPrimaryColor;
+	static FColor SelectionLockedPrimaryColor;
 
 
 	/**
diff --git a/Engine/Plugins/Experimental/Dataflow/Source/DataflowEnginePlugin/Public/Dataflow/DataflowEngineSceneHitProxies.h b/Engine/Plugins/Experimental/Dataflow/Source/DataflowEnginePlugin/Public/Dataflow/DataflowEngineSceneHitProxies.h
--- a/Engine/Plugins/Experimental/Dataflow/Source/DataflowEnginePlugin/Public/Dataflow/DataflowEngineSceneHitProxies.h
+++ b/Engine/Plugins/Experimental/Dataflow/Source/DataflowEnginePlugin/Public/Dataflow/DataflowEngineSceneHitProxies.h
@@ -10,6 +10,22 @@
 class UDataflowComponent;
 class UPrimitiveComponent;
 
+
+/** HitProxy with for dataflow actor.
+ */
+struct HDataflowDefault : public HActor
+{
+	DECLARE_HIT_PROXY(DATAFLOWENGINEPLUGIN_API)
+	HDataflowDefault(AActor* InActor, const UPrimitiveComponent* InPrimitiveComponent)
+		: HActor(InActor, InPrimitiveComponent){}
+
+	virtual EMouseCursor::Type GetMouseCursor() override
+	{
+		return EMouseCursor::Default;
+	}
+};
+
+
 /** HitProxy with for dataflow actor.
  */
 struct HDataflowNode : public HActor
@@ -30,4 +46,23 @@ struct HDataflowNode : public HActor
 	{
 		return EMouseCursor::Default;
 	}
+};
+
+
+/** HitProxy with for dataflow actor.
+ */
+struct HDataflowVertex : public HActor
+{
+	DECLARE_HIT_PROXY(DATAFLOWENGINEPLUGIN_API)
+
+	HDataflowVertex(AActor* InActor, const UPrimitiveComponent* InPrimitiveComponent, int32 InVertexIndex)
+		: HActor(InActor, InPrimitiveComponent)
+	{
+		SectionIndex = InVertexIndex;
+	}
+
+	virtual EMouseCursor::Type GetMouseCursor() override
+	{
+		return EMouseCursor::Crosshairs;
+	}
 };
\ No newline at end of file
diff --git a/Engine/Source/Runtime/Experimental/Chaos/Private/GeometryCollection/Facades/CollectionRenderingFacade.cpp b/Engine/Source/Runtime/Experimental/Chaos/Private/GeometryCollection/Facades/CollectionRenderingFacade.cpp
--- a/Engine/Source/Runtime/Experimental/Chaos/Private/GeometryCollection/Facades/CollectionRenderingFacade.cpp
+++ b/Engine/Source/Runtime/Experimental/Chaos/Private/GeometryCollection/Facades/CollectionRenderingFacade.cpp
@@ -15,16 +15,18 @@ namespace GeometryCollection::Facades
 		, Collection(&InCollection)
 		, VertexAttribute(InCollection, "Vertex", FGeometryCollection::VerticesGroup)
 		, VertexToGeometryIndexAttribute(InCollection, "GeometryIndex", FGeometryCollection::VerticesGroup, FGeometryCollection::GeometryGroup)
+		, VertexSelectionAttribute(InCollection, "SelectionState", FGeometryCollection::VerticesGroup)
+		, VertexHitProxyIndexAttribute(InCollection, "HitIndex", FGeometryCollection::VerticesGroup)
 		, IndicesAttribute(InCollection, "Indices", FGeometryCollection::FacesGroup, FGeometryCollection::VerticesGroup)
 		, MaterialIDAttribute(InCollection, "MaterialID", FGeometryCollection::FacesGroup)
 		, TriangleSectionAttribute(InCollection, "Sections", FGeometryCollection::MaterialGroup)
 		, GeometryNameAttribute(InCollection, "Name", FGeometryCollection::GeometryGroup)
-		, HitProxyIndexAttribute(InCollection, "HitIndex", FGeometryCollection::GeometryGroup)
-		, VertixStartAttribute(InCollection, "VertexStart", FGeometryCollection::GeometryGroup)
-		, VertixCountAttribute(InCollection, "VertexCount", FGeometryCollection::GeometryGroup)
+		, GeometryHitProxyIndexAttribute(InCollection, "HitIndex", FGeometryCollection::GeometryGroup)
+		, VertexStartAttribute(InCollection, "VertexStart", FGeometryCollection::GeometryGroup)
+		, VertexCountAttribute(InCollection, "VertexCount", FGeometryCollection::GeometryGroup)
 		, IndicesStartAttribute(InCollection, "IndicesStart", FGeometryCollection::GeometryGroup)
 		, IndicesCountAttribute(InCollection, "IndicesCount", FGeometryCollection::GeometryGroup)
-		, SelectionState(InCollection, "SelectionState", FGeometryCollection::GeometryGroup)
+		, GeometrySelectionAttribute(InCollection, "SelectionState", FGeometryCollection::GeometryGroup)
 
 	{}
 
@@ -33,16 +35,18 @@ namespace GeometryCollection::Facades
 		, Collection(nullptr)
 		, VertexAttribute(InCollection, "Vertex", FGeometryCollection::VerticesGroup)
 		, VertexToGeometryIndexAttribute(InCollection, "GeometryIndex", FGeometryCollection::VerticesGroup, FGeometryCollection::GeometryGroup)
+		, VertexSelectionAttribute(InCollection, "SelectionState", FGeometryCollection::VerticesGroup)
+		, VertexHitProxyIndexAttribute(InCollection, "HitIndex", FGeometryCollection::VerticesGroup)
 		, IndicesAttribute(InCollection, "Indices", FGeometryCollection::FacesGroup, FGeometryCollection::VerticesGroup)
 		, MaterialIDAttribute(InCollection, "MaterialID", FGeometryCollection::FacesGroup)
 		, TriangleSectionAttribute(InCollection, "Sections", FGeometryCollection::MaterialGroup)
 		, GeometryNameAttribute(InCollection, "Name", FGeometryCollection::GeometryGroup)
-		, HitProxyIndexAttribute(InCollection, "HitIndex", FGeometryCollection::GeometryGroup)
-		, VertixStartAttribute(InCollection, "VertexStart", FGeometryCollection::GeometryGroup, FGeometryCollection::VerticesGroup)
-		, VertixCountAttribute(InCollection, "VertexCount", FGeometryCollection::GeometryGroup)
+		, GeometryHitProxyIndexAttribute(InCollection, "HitIndex", FGeometryCollection::GeometryGroup)
+		, VertexStartAttribute(InCollection, "VertexStart", FGeometryCollection::GeometryGroup, FGeometryCollection::VerticesGroup)
+		, VertexCountAttribute(InCollection, "VertexCount", FGeometryCollection::GeometryGroup)
 		, IndicesStartAttribute(InCollection, "IndicesStart", FGeometryCollection::GeometryGroup, FGeometryCollection::FacesGroup)
 		, IndicesCountAttribute(InCollection, "IndicesCount", FGeometryCollection::GeometryGroup)
-		, SelectionState(InCollection, "SelectionState", FGeometryCollection::GeometryGroup)
+		, GeometrySelectionAttribute(InCollection, "SelectionState", FGeometryCollection::GeometryGroup)
 	{}
 
 	//
@@ -53,17 +57,19 @@ namespace GeometryCollection::Facades
 	{
 		check(!IsConst());
 		VertexAttribute.Add();
+		VertexSelectionAttribute.Add();
 		VertexToGeometryIndexAttribute.Add();
+		VertexHitProxyIndexAttribute.Add();
 		IndicesAttribute.Add();
 		MaterialIDAttribute.Add();
 		TriangleSectionAttribute.Add();
 		GeometryNameAttribute.Add();
-		HitProxyIndexAttribute.Add();
-		VertixStartAttribute.Add();
-		VertixCountAttribute.Add();
+		GeometryHitProxyIndexAttribute.Add();
+		VertexStartAttribute.Add();
+		VertexCountAttribute.Add();
 		IndicesStartAttribute.Add();
 		IndicesCountAttribute.Add();
-		SelectionState.Add();
+		GeometrySelectionAttribute.Add();
 	}
 
 	bool FRenderingFacade::CanRenderSurface( ) const
@@ -74,12 +80,13 @@ namespace GeometryCollection::Facades
 	bool FRenderingFacade::IsValid( ) const
 	{
 		return VertexAttribute.IsValid() && VertexToGeometryIndexAttribute.IsValid() &&
+			VertexSelectionAttribute.IsValid() && VertexHitProxyIndexAttribute.IsValid() &&
 			IndicesAttribute.IsValid() &&
 			MaterialIDAttribute.IsValid() && TriangleSectionAttribute.IsValid() &&
-			GeometryNameAttribute.IsValid() && HitProxyIndexAttribute.IsValid() &&
-			VertixStartAttribute.IsValid() && VertixCountAttribute.IsValid() &&
+			GeometryNameAttribute.IsValid() && GeometryHitProxyIndexAttribute.IsValid() &&
+			VertexStartAttribute.IsValid() && VertexCountAttribute.IsValid() &&
 			IndicesStartAttribute.IsValid() && IndicesCountAttribute.IsValid() &&
-			SelectionState.IsValid();
+			GeometrySelectionAttribute.IsValid();
 	}
 
 	int32 FRenderingFacade::NumTriangles() const
@@ -161,11 +168,11 @@ namespace GeometryCollection::Facades
 			GeomIndex = GeometryNameAttribute.AddElements(1);
 			GeometryNameAttribute.Modify()[GeomIndex] = InName;
 
-			VertixStartAttribute.Modify()[GeomIndex] = VertexAttribute.Num();
-			VertixCountAttribute.Modify()[GeomIndex] = 0;
+			VertexStartAttribute.Modify()[GeomIndex] = VertexAttribute.Num();
+			VertexCountAttribute.Modify()[GeomIndex] = 0;
 			IndicesStartAttribute.Modify()[GeomIndex] = IndicesAttribute.Num();
 			IndicesCountAttribute.Modify()[GeomIndex] = 0;
-			SelectionState.Modify()[GeomIndex] = 0;
+			GeometrySelectionAttribute.Modify()[GeomIndex] = 0;
 		}
 		return GeomIndex;
 	}
@@ -177,19 +184,19 @@ namespace GeometryCollection::Facades
 		{
 			check( GeometryNameAttribute.Num()-1 == InGeomIndex );
 
-			if (VertixStartAttribute.Get()[InGeomIndex] < VertexAttribute.Num())
+			if (VertexStartAttribute.Get()[InGeomIndex] < VertexAttribute.Num())
 			{
-				VertixCountAttribute.Modify()[InGeomIndex] = VertexAttribute.Num() - VertixStartAttribute.Get()[InGeomIndex];
+				VertexCountAttribute.Modify()[InGeomIndex] = VertexAttribute.Num() - VertexStartAttribute.Get()[InGeomIndex];
 
 				TManagedArray<int32>& GeomIndexAttr = VertexToGeometryIndexAttribute.Modify();
-				for (int i = VertixStartAttribute.Get()[InGeomIndex]; i < VertexAttribute.Num(); i++)
+				for (int i = VertexStartAttribute.Get()[InGeomIndex]; i < VertexAttribute.Num(); i++)
 				{
 					GeomIndexAttr[i] = InGeomIndex;
 				}
 			}
 			else
 			{
-				VertixStartAttribute.Modify()[InGeomIndex] = VertexAttribute.Num();
+				VertexStartAttribute.Modify()[InGeomIndex] = VertexAttribute.Num();
 			}
 
 			if (IndicesStartAttribute.Get()[InGeomIndex] < IndicesAttribute.Num())
@@ -208,12 +215,24 @@ namespace GeometryCollection::Facades
 		FStringIntMap Map;
 		for (int32 i = 0; i < GeometryNameAttribute.Num(); i++)
 		{
-			Map.Add(GetGeometryNameAttribute()[i], i);
+			Map.Add(GetGeometryName()[i], i);
 		}
 		return Map;
 	}
 
 
+	int32 FRenderingFacade::NumVerticesOnSelectedGeometry() const
+	{
+		const TManagedArray<int32>& SelectedGeometry = GeometrySelectionAttribute.Get();
+		const TManagedArray<int32>& VertexCount = VertexCountAttribute.Get();
+		int32 RetCount = 0;
+		for (int i = 0; i < SelectedGeometry.Num(); i++)
+			if (SelectedGeometry[i])
+				RetCount += VertexCount[i];
+		return RetCount;
+	}
+
+
 }; // GeometryCollection::Facades
 
 
diff --git a/Engine/Source/Runtime/Experimental/Chaos/Public/GeometryCollection/Facades/CollectionRenderingFacade.h b/Engine/Source/Runtime/Experimental/Chaos/Public/GeometryCollection/Facades/CollectionRenderingFacade.h
--- a/Engine/Source/Runtime/Experimental/Chaos/Public/GeometryCollection/Facades/CollectionRenderingFacade.h
+++ b/Engine/Source/Runtime/Experimental/Chaos/Public/GeometryCollection/Facades/CollectionRenderingFacade.h
@@ -25,7 +25,7 @@ namespace GeometryCollection::Facades
 		typedef TMap<FString, int32> FStringIntMap;
 		/**
 		* FRenderingFacade Constuctor
-		* @param VertixDependencyGroup : GroupName the index attribute is dependent on.
+		* @param VertexDependencyGroup : GroupName the index attribute is dependent on.
 		*/
 		FRenderingFacade(FManagedArrayCollection& InSelf);
 		FRenderingFacade(const FManagedArrayCollection& InSelf);
@@ -58,12 +58,6 @@ namespace GeometryCollection::Facades
 		/** GetIndices */
 		const TManagedArray< FIntVector >& GetIndices() const { return IndicesAttribute.Get(); }
 
-		/** GetVertices */
-		const TManagedArray< FVector3f >& GetVertices() const { return VertexAttribute.Get(); }
-
-		/** GetVertexToGeometryIndexAttribute */
-		const TManagedArray< int32 >& GetVertexToGeometryIndexAttribute() const { return VertexToGeometryIndexAttribute.Get(); }
-
 		/** GetMaterialID */
 		const TManagedArray< int32 >& GetMaterialID() const { return MaterialIDAttribute.Get(); }
 
@@ -72,7 +66,29 @@ namespace GeometryCollection::Facades
 
 		/** BuildMeshSections */
 		TArray<FTriangleSection> BuildMeshSections(const TArray<FIntVector>& Indices, TArray<int32> BaseMeshOriginalIndicesIndex, TArray<FIntVector>& RetIndices) const;
-		
+
+
+		//
+		//  Vertices
+		//
+
+		/** GetVertices */
+		const TManagedArray< FVector3f >& GetVertices() const { return VertexAttribute.Get(); }
+
+		/** GetVertexSelectionAttribute */
+		const TManagedArray< int32 >& GetVertexSelection() const { return VertexSelectionAttribute.Get(); }
+		      TManagedArray< int32 >& ModifyVertexSelection() { check(!IsConst()); return VertexSelectionAttribute.Modify(); }
+
+		/** GetVertexToGeometryIndexAttribute */
+		const TManagedArray< int32 >& GetVertexToGeometryIndex() const { return VertexToGeometryIndexAttribute.Get(); }
+
+		/** HitProxyIDAttribute */
+		const TManagedArray< int32 >& GetVertexHitProxyIndex() const { return VertexHitProxyIndexAttribute.Get(); }
+		TManagedArray< int32 >& ModifyVertexHitProxyIndex() { check(!IsConst()); return VertexHitProxyIndexAttribute.Modify(); }
+
+		/** NumVertices */
+		int32 NumVertices() const { return VertexAttribute.Num(); }
+
 
 		//
 		// Geometry Group Attributes
@@ -87,27 +103,30 @@ namespace GeometryCollection::Facades
 		int32 NumGeometry() const { return GeometryNameAttribute.Num(); }
 
 		/** GeometryNameAttribute */
-		const TManagedArray< FString >& GetGeometryNameAttribute() const { return GeometryNameAttribute.Get(); }
+		const TManagedArray< FString >& GetGeometryName() const { return GeometryNameAttribute.Get(); }
 
 		/** HitProxyIDAttribute */
-		const TManagedArray< int32 >& GetHitProxyIndexAttribute() const { return HitProxyIndexAttribute.Get(); }
-		      TManagedArray< int32 >& ModifyHitProxyIndexAttribute() {check(!IsConst());return HitProxyIndexAttribute.Modify(); }
+		const TManagedArray< int32 >& GetGeometryHitProxyIndex() const { return GeometryHitProxyIndexAttribute.Get(); }
+		      TManagedArray< int32 >& ModifyGeometryHitProxyIndex() {check(!IsConst());return GeometryHitProxyIndexAttribute.Modify(); }
 
-		/** VertixStartAttribute */
-		const TManagedArray< int32 >& GetVertixStartAttribute() const { return VertixStartAttribute.Get(); }
+		/** VertexStartAttribute */
+		const TManagedArray< int32 >& GetVertexStart() const { return VertexStartAttribute.Get(); }
 
-		/** VertixCountAttribute */
-		const TManagedArray< int32 >& GetVertixCountAttribute() const { return VertixCountAttribute.Get(); }
+		/** VertexCountAttribute */
+		const TManagedArray< int32 >& GetVertexCount() const { return VertexCountAttribute.Get(); }
 
 		/** IndicesStartAttribute */
-		const TManagedArray< int32 >& GetIndicesStartAttribute() const { return IndicesStartAttribute.Get(); }
+		const TManagedArray< int32 >& GetIndicesStart() const { return IndicesStartAttribute.Get(); }
 
 		/** IndicesCountAttribute */
-		const TManagedArray< int32 >& GetIndicesCountAttribute() const { return IndicesCountAttribute.Get(); }
+		const TManagedArray< int32 >& GetIndicesCount() const { return IndicesCountAttribute.Get(); }
 
 		/** SelectionState */
-		const TManagedArray< int32 >& GetSelectionState() const { return SelectionState.Get(); }
-		      TManagedArray< int32 >& ModifySelectionState() { check(!IsConst()); return SelectionState.Modify(); }
+		const TManagedArray< int32 >& GetSelectionState() const { return GeometrySelectionAttribute.Get(); }
+		      TManagedArray< int32 >& ModifySelectionState() { check(!IsConst()); return GeometrySelectionAttribute.Modify(); }
+
+		/** NumVerticesOnSelectedGeometry */
+		int32 NumVerticesOnSelectedGeometry() const;
 
 		/** GetGeometryNameToIndexMap */
 		FStringIntMap GetGeometryNameToIndexMap() const;
@@ -118,17 +137,20 @@ namespace GeometryCollection::Facades
 
 		TManagedArrayAccessor<FVector3f> VertexAttribute;
 		TManagedArrayAccessor<int32> VertexToGeometryIndexAttribute;
+		TManagedArrayAccessor<int32> VertexSelectionAttribute;
+		TManagedArrayAccessor<int32> VertexHitProxyIndexAttribute;
+
 		TManagedArrayAccessor<FIntVector> IndicesAttribute;
 		TManagedArrayAccessor<int32> MaterialIDAttribute;
 		TManagedArrayAccessor<FTriangleSection> TriangleSectionAttribute;
 
 		TManagedArrayAccessor<FString> GeometryNameAttribute;
-		TManagedArrayAccessor<int32> HitProxyIndexAttribute;
-		TManagedArrayAccessor<int32> VertixStartAttribute;
-		TManagedArrayAccessor<int32> VertixCountAttribute;
+		TManagedArrayAccessor<int32> GeometryHitProxyIndexAttribute;
+		TManagedArrayAccessor<int32> VertexStartAttribute;
+		TManagedArrayAccessor<int32> VertexCountAttribute;
 		TManagedArrayAccessor<int32> IndicesStartAttribute;
 		TManagedArrayAccessor<int32> IndicesCountAttribute;
-		TManagedArrayAccessor<int32> SelectionState;
+		TManagedArrayAccessor<int32> GeometrySelectionAttribute;
 
 
 	};
