Dataflow : Copy/Paste Selection

diff --git a/Engine/Plugins/Experimental/ChaosFlesh/Source/ChaosFleshNodes/Private/Dataflow/ChaosFleshKinematicInitializationNodes.cpp b/Engine/Plugins/Experimental/ChaosFlesh/Source/ChaosFleshNodes/Private/Dataflow/ChaosFleshKinematicInitializationNodes.cpp
--- a/Engine/Plugins/Experimental/ChaosFlesh/Source/ChaosFleshNodes/Private/Dataflow/ChaosFleshKinematicInitializationNodes.cpp
+++ b/Engine/Plugins/Experimental/ChaosFlesh/Source/ChaosFleshNodes/Private/Dataflow/ChaosFleshKinematicInitializationNodes.cpp
@@ -22,7 +22,6 @@ namespace Dataflow
 	void RegisterChaosFleshKinematicInitializationNodes()
 	{
 		DATAFLOW_NODE_REGISTER_CREATION_FACTORY(FKinematicInitializationDataflowNode);
-		DATAFLOW_NODE_REGISTER_CREATION_FACTORY(FSelectionSetDataflowNode);
 		DATAFLOW_NODE_REGISTER_CREATION_FACTORY(FKinematicTetrahedralBindingsDataflowNode);
 		DATAFLOW_NODE_REGISTER_CREATION_FACTORY(FSetVerticesKinematicDataflowNode);
 		DATAFLOW_NODE_REGISTER_CREATION_FACTORY(FBinVerticesDataflowNode);
@@ -239,26 +238,6 @@ void FSetVerticesKinematicDataflowNode::Evaluate(Dataflow::FContext& Context, co
 }
 
 
-void FSelectionSetDataflowNode::Evaluate(Dataflow::FContext& Context, const FDataflowOutput* Out) const
-{
-	if (Out->IsA<DataType>(&IndicesOut))
-	{
-		TArray<int32> IntArray;
-		TArray<FString> StrArray;
-		Indices.ParseIntoArray(StrArray, *FString(" "));
-		for (FString Elem : StrArray)
-		{
-			if (FCString::IsNumeric(*Elem))
-			{
-				IntArray.Add(FCString::Atoi(*Elem));
-			}
-		}
-
-		Out->SetValue<DataType>(IntArray, Context);
-	}
-}
-
-
 void FBinVerticesDataflowNode::Evaluate(Dataflow::FContext& Context, const FDataflowOutput* Out) const
 {
 	if (Out->IsA<DataType>(&VertexIndicesOut))
diff --git a/Engine/Plugins/Experimental/ChaosFlesh/Source/ChaosFleshNodes/Public/Dataflow/ChaosFleshKinematicInitializationNodes.h b/Engine/Plugins/Experimental/ChaosFlesh/Source/ChaosFleshNodes/Public/Dataflow/ChaosFleshKinematicInitializationNodes.h
--- a/Engine/Plugins/Experimental/ChaosFlesh/Source/ChaosFleshNodes/Public/Dataflow/ChaosFleshKinematicInitializationNodes.h
+++ b/Engine/Plugins/Experimental/ChaosFlesh/Source/ChaosFleshNodes/Public/Dataflow/ChaosFleshKinematicInitializationNodes.h
@@ -84,32 +84,6 @@ public:
 
 
 
-USTRUCT(meta = (DataflowFlesh))
-struct FSelectionSetDataflowNode : public FDataflowNode
-{
-	GENERATED_USTRUCT_BODY()
-	DATAFLOW_NODE_DEFINE_INTERNAL(FSelectionSetDataflowNode, "SelectionSet", "Flesh", "")
-
-public:
-	typedef TArray<int32> DataType;
-
-	UPROPERTY(EditAnywhere, Category = "Dataflow")
-	FString Indices = FString("1 2 3");
-
-	UPROPERTY(meta = (DataflowOutput, DisplayName = "SelectionSet"))
-	TArray<int32> IndicesOut;
-
-	FSelectionSetDataflowNode(const Dataflow::FNodeParameters& InParam, FGuid InGuid = FGuid::NewGuid())
-		: FDataflowNode(InParam, InGuid)
-	{
-		RegisterOutputConnection(&IndicesOut);
-	}
-
-	virtual void Evaluate(Dataflow::FContext& Context, const FDataflowOutput* Out) const override;
-};
-
-
-
 USTRUCT(meta = (DataflowFlesh))
 struct FSetVerticesKinematicDataflowNode : public FDataflowNode
 {
diff --git a/Engine/Plugins/Experimental/Dataflow/Source/DataflowEditor/DataflowEditor.Build.cs b/Engine/Plugins/Experimental/Dataflow/Source/DataflowEditor/DataflowEditor.Build.cs
--- a/Engine/Plugins/Experimental/Dataflow/Source/DataflowEditor/DataflowEditor.Build.cs
+++ b/Engine/Plugins/Experimental/Dataflow/Source/DataflowEditor/DataflowEditor.Build.cs
@@ -10,6 +10,7 @@ namespace UnrealBuildTool.Rules
 				new string[]
 				{
 					"Core",
+					"ApplicationCore",
 					"CoreUObject",
 				    "Slate",
 				    "SlateCore",
@@ -33,8 +34,10 @@ namespace UnrealBuildTool.Rules
 					"DataflowCore",
 					"DataflowEngine",
 					"DataflowEnginePlugin",
+					"DataflowNodes",
 					"Slate",
 					"Chaos",
+					"XmlParser",
 				}
 			);
 		}
diff --git a/Engine/Plugins/Experimental/Dataflow/Source/DataflowEditor/Private/Dataflow/DataflowEditorViewport.cpp b/Engine/Plugins/Experimental/Dataflow/Source/DataflowEditor/Private/Dataflow/DataflowEditorViewport.cpp
--- a/Engine/Plugins/Experimental/Dataflow/Source/DataflowEditor/Private/Dataflow/DataflowEditorViewport.cpp
+++ b/Engine/Plugins/Experimental/Dataflow/Source/DataflowEditor/Private/Dataflow/DataflowEditorViewport.cpp
@@ -2,23 +2,13 @@
 #include "Dataflow/DataflowEditorViewport.h"
 
 #include "AdvancedPreviewScene.h"
-#include "Dataflow/DataflowEditorCommands.h"
-#include "Dataflow/DataflowEdNode.h"
-#include "Dataflow/DataflowObject.h"
 #include "Dataflow/DataflowActor.h"
-#include "Dataflow/DataflowComponent.h"
+#include "Dataflow/DataflowEditorCommands.h"
 #include "Dataflow/DataflowEditorViewportToolbar.h"
-#include "Dataflow/DataflowEngineSceneHitProxies.h"
-#include "DynamicMeshBuilder.h"
-#include "Dataflow/DataflowEditorToolkit.h"
-#include "EditorModes.h"
+#include "Dataflow/DataflowEditorViewportClient.h"
 #include "EditorViewportClient.h"
-#include "EditorViewportCommands.h"
 #include "Framework/Application/SlateApplication.h"
-#include "Framework/Commands/UIAction.h"
-#include "Materials/Material.h"
-#include "SCommonEditorViewportToolbarBase.h"
-#include "Settings/EditorStyleSettings.h"
+#include "PreviewScene.h"
 
 SDataflowEditorViewport::SDataflowEditorViewport()
 {
@@ -105,220 +95,3 @@ void SDataflowEditorViewport::BindCommands()
 		);
 	}
 }
-
-// ----------------------------------------------------------------------------------
-
-FDataflowEditorViewportClient::FDataflowEditorViewportClient(FPreviewScene* InPreviewScene, 
-	const TWeakPtr<SEditorViewport> InEditorViewportWidget,
-	TWeakPtr<FDataflowEditorToolkit> InDataflowEditorToolkitPtr)
-	: 
-	FEditorViewportClient(nullptr, InPreviewScene, InEditorViewportWidget)
-	, DataflowEditorToolkitPtr(InDataflowEditorToolkitPtr)
-{
-	SetRealtime(true);
-	SetViewModes(VMI_Lit, VMI_Lit);
-	bSetListenerPosition = false;
-	EngineShowFlags.Grid = false;
-}
-
-void FDataflowEditorViewportClient::SetSelectionMode(FDataflowSelectionState::EMode InState)
-{
-	FDataflowSelectionState State = DataflowActor->DataflowComponent->GetSelectionState();
-		
-	if (SelectionMode == InState)
-	{
-		SelectionMode = FDataflowSelectionState::EMode::DSS_Dataflow_None;
-	}
-	else
-	{
-		SelectionMode = InState;
-	}
-
-	State.Mode = SelectionMode;
-	DataflowActor->DataflowComponent->SetSelectionState(State);
-
-	if (SelectionMode == FDataflowSelectionState::EMode::DSS_Dataflow_None)
-	{
-		if (!DataflowActor->DataflowComponent->GetSelectionState().IsEmpty())
-		{
-			DataflowActor->DataflowComponent->SetSelectionState(FDataflowSelectionState(SelectionMode));
-		}
-	}
-}
-bool FDataflowEditorViewportClient::CanSetSelectionMode(FDataflowSelectionState::EMode InState)
-{
-	TSharedPtr<FDataflowEditorToolkit> DataflowEditorToolkit = DataflowEditorToolkitPtr.Pin();
-	if (DataflowEditorToolkitPtr.IsValid())
-	{
-		if (const UDataflow* Dataflow = DataflowEditorToolkit->GetDataflow())
-		{
-			if (Dataflow->GetRenderTargets().Num())
-			{
-				if (InState == FDataflowSelectionState::EMode::DSS_Dataflow_Object)
-				{
-					return true;
-				}
-
-				if (InState == FDataflowSelectionState::EMode::DSS_Dataflow_Vertex
-					&& !DataflowActor->DataflowComponent->GetSelectionState().Nodes.IsEmpty())
-				{
-					return true;
-				}
-			}
-		}
-	}
-
-	return false;
-}
-bool FDataflowEditorViewportClient::IsSelectionModeActive(FDataflowSelectionState::EMode InState) 
-{ 
-	return SelectionMode == InState;
-}
-
-Dataflow::FTimestamp FDataflowEditorViewportClient::LatestTimestamp(const UDataflow* Dataflow, const Dataflow::FContext* Context)
-{
-	if (Dataflow && Context)
-	{
-		return FMath::Max(Dataflow->GetRenderingTimestamp().Value, Context->GetTimestamp().Value);
-	}
-	return Dataflow::FTimestamp::Invalid;
-}
-
-
-void FDataflowEditorViewportClient::ProcessClick(FSceneView& View, HHitProxy* HitProxy, FKey Key, EInputEvent Event, uint32 HitX, uint32 HitY)
-{
-	Super::ProcessClick(View, HitProxy, Key, Event, HitX, HitY);
-	if (DataflowActor && DataflowActor->DataflowComponent)
-	{
-		const bool bIsShiftKeyDown = Viewport->KeyState(EKeys::LeftShift) || Viewport->KeyState(EKeys::RightShift);
-		const bool bIsCtrltKeyDown = Viewport->KeyState(EKeys::LeftControl) || Viewport->KeyState(EKeys::RightControl);
-
-		FDataflowSelectionState SelectionState = DataflowActor->DataflowComponent->GetSelectionState();
-		FDataflowSelectionState PreState = SelectionState;
-
-		if (SelectionMode == FDataflowSelectionState::EMode::DSS_Dataflow_Object)
-		{
-			if (HitProxy && HitProxy->IsA(HDataflowNode::StaticGetType()))
-			{
-				HDataflowNode* DataflowNode = (HDataflowNode*)(HitProxy);
-				FDataflowSelectionState::ObjectID ID(DataflowNode->NodeName, DataflowNode->GeometryIndex);
-
-				if (bIsShiftKeyDown)
-				{
-					if (!SelectionState.Nodes.Contains(ID))
-					{
-						SelectionState.Nodes.AddUnique(ID);
-					}
-				}
-				else if (bIsCtrltKeyDown)
-				{
-					if (SelectionState.Nodes.Contains(ID))
-					{
-						SelectionState.Nodes.Remove(ID);
-					}
-				}
-				else
-				{
-					SelectionState.Nodes.Empty();
-					SelectionState.Nodes.AddUnique(ID);
-				}
-			}
-			else if (!bIsShiftKeyDown && !bIsCtrltKeyDown)
-			{
-				SelectionState.Nodes.Empty();
-			}
-		}
-
-		if (SelectionMode == FDataflowSelectionState::EMode::DSS_Dataflow_Vertex)
-		{
-			if (HitProxy && HitProxy->IsA(HDataflowVertex::StaticGetType()))
-			{
-				HDataflowVertex* DataflowVertex = (HDataflowVertex*)(HitProxy);
-				int32 ID = DataflowVertex->SectionIndex;
-				if (bIsShiftKeyDown)
-				{
-					if (!SelectionState.Vertices.Contains(ID))
-					{
-						SelectionState.Vertices.AddUnique(ID);
-					}
-				}
-				else if (bIsCtrltKeyDown)
-				{
-					if (SelectionState.Vertices.Contains(ID))
-					{
-						SelectionState.Vertices.Remove(ID);
-					}
-				}
-				else
-				{
-					SelectionState.Vertices.Empty();
-					SelectionState.Vertices.AddUnique(ID);
-				}
-			}
-			else if (!bIsShiftKeyDown && !bIsCtrltKeyDown)
-			{
-				SelectionState.Vertices.Empty();
-			}
-		}
-
-		if (PreState != SelectionState)
-		{
-			DataflowActor->DataflowComponent->SetSelectionState(SelectionState);
-		}
-	}
-}
-
-void FDataflowEditorViewportClient::Tick(float DeltaSeconds)
-{
-	Super::Tick(DeltaSeconds);
-
-	TSharedPtr<FDataflowEditorToolkit> DataflowEditorToolkit = DataflowEditorToolkitPtr.Pin();
-
-	if (DataflowActor && DataflowEditorToolkitPtr.IsValid())
-	{
-		if (TSharedPtr<Dataflow::FContext> Context = DataflowEditorToolkit->GetContext())
-		{
-			if (const UDataflow* Dataflow = DataflowEditorToolkit->GetDataflow())
-			{
-				if (UDataflowComponent* DataflowComponent = DataflowActor->GetDataflowComponent())
-				{
-					Dataflow::FTimestamp SystemTimestamp = LatestTimestamp(Dataflow, Context.Get());
-					if (SystemTimestamp >= LastModifiedTimestamp)
-					{
-						if (Dataflow->GetRenderTargets().Num())
-						{
-							// Component Object Rendering
-							DataflowComponent->ResetRenderTargets();
-							DataflowComponent->SetDataflow(Dataflow);
-							DataflowComponent->SetContext(Context);
-							for (const UDataflowEdNode* Node : Dataflow->GetRenderTargets())
-							{
-								DataflowComponent->AddRenderTarget(Node);
-							}
-						}
-						else
-						{
-							DataflowComponent->ResetRenderTargets();
-						}
-
-						LastModifiedTimestamp = LatestTimestamp(Dataflow, Context.Get()).Value + 1;
-					}
-				}
-			}
-		}
-	}
-
-	// Tick the preview scene world.
-	if (!GIntraFrameDebuggingGameThread)
-	{
-		PreviewScene->GetWorld()->Tick(IsRealtime() ? LEVELTICK_All : LEVELTICK_TimeOnly, DeltaSeconds);
-	}
-}
-
-
-void FDataflowEditorViewportClient::AddReferencedObjects(FReferenceCollector& Collector)
-{
-	Super::AddReferencedObjects(Collector);
-}
-
-
diff --git a/Engine/Plugins/Experimental/Dataflow/Source/DataflowEditor/Private/Dataflow/DataflowEditorViewportClient.cpp b/Engine/Plugins/Experimental/Dataflow/Source/DataflowEditor/Private/Dataflow/DataflowEditorViewportClient.cpp
new file mode 100644
--- /dev/null
+++ b/Engine/Plugins/Experimental/Dataflow/Source/DataflowEditor/Private/Dataflow/DataflowEditorViewportClient.cpp
@@ -0,0 +1,250 @@
+// Copyright Epic Games, Inc. All Rights Reserved.
+#include "Dataflow/DataflowEditorViewportClient.h"
+
+#include "Dataflow/DataflowActor.h"
+#include "Dataflow/DataflowEditorToolkit.h"
+#include "Dataflow/DataflowEngineSceneHitProxies.h"
+#include "Dataflow/DataflowXml.h"
+#include "HAL/PlatformApplicationMisc.h"
+#include "PreviewScene.h"
+
+FDataflowEditorViewportClient::FDataflowEditorViewportClient(FPreviewScene* InPreviewScene, 
+	const TWeakPtr<SEditorViewport> InEditorViewportWidget,
+	TWeakPtr<FDataflowEditorToolkit> InDataflowEditorToolkitPtr)
+	: 
+	FEditorViewportClient(nullptr, InPreviewScene, InEditorViewportWidget)
+	, DataflowEditorToolkitPtr(InDataflowEditorToolkitPtr)
+{
+	SetRealtime(true);
+	SetViewModes(VMI_Lit, VMI_Lit);
+	bSetListenerPosition = false;
+	EngineShowFlags.Grid = false;
+}
+
+void FDataflowEditorViewportClient::SetSelectionMode(FDataflowSelectionState::EMode InState)
+{
+	FDataflowSelectionState State = DataflowActor->DataflowComponent->GetSelectionState();
+		
+	if (SelectionMode == InState)
+	{
+		SelectionMode = FDataflowSelectionState::EMode::DSS_Dataflow_None;
+	}
+	else
+	{
+		SelectionMode = InState;
+	}
+
+	State.Mode = SelectionMode;
+	DataflowActor->DataflowComponent->SetSelectionState(State);
+
+	if (SelectionMode == FDataflowSelectionState::EMode::DSS_Dataflow_None)
+	{
+		if (!DataflowActor->DataflowComponent->GetSelectionState().IsEmpty())
+		{
+			DataflowActor->DataflowComponent->SetSelectionState(FDataflowSelectionState(SelectionMode));
+		}
+	}
+}
+bool FDataflowEditorViewportClient::CanSetSelectionMode(FDataflowSelectionState::EMode InState)
+{
+	TSharedPtr<FDataflowEditorToolkit> DataflowEditorToolkit = DataflowEditorToolkitPtr.Pin();
+	if (DataflowEditorToolkitPtr.IsValid())
+	{
+		if (const UDataflow* Dataflow = DataflowEditorToolkit->GetDataflow())
+		{
+			if (Dataflow->GetRenderTargets().Num())
+			{
+				if (InState == FDataflowSelectionState::EMode::DSS_Dataflow_Object)
+				{
+					return true;
+				}
+
+				if (InState == FDataflowSelectionState::EMode::DSS_Dataflow_Vertex
+					&& !DataflowActor->DataflowComponent->GetSelectionState().Nodes.IsEmpty())
+				{
+					return true;
+				}
+			}
+		}
+	}
+
+	return false;
+}
+bool FDataflowEditorViewportClient::IsSelectionModeActive(FDataflowSelectionState::EMode InState) 
+{ 
+	return SelectionMode == InState;
+}
+
+Dataflow::FTimestamp FDataflowEditorViewportClient::LatestTimestamp(const UDataflow* Dataflow, const Dataflow::FContext* Context)
+{
+	if (Dataflow && Context)
+	{
+		return FMath::Max(Dataflow->GetRenderingTimestamp().Value, Context->GetTimestamp().Value);
+	}
+	return Dataflow::FTimestamp::Invalid;
+}
+
+bool FDataflowEditorViewportClient::InputKey(const FInputKeyEventArgs& EventArgs)
+{
+	bool bHandled = false;
+
+	FInputEventState InputState(EventArgs.Viewport, EventArgs.Key, EventArgs.Event);
+	if (InputState.IsCtrlButtonPressed() && EventArgs.Key == EKeys::C)
+	{
+		if (DataflowActor->DataflowComponent->GetSelectionState().Vertices.Num())
+		{
+			FString XmlBuffer = FDataflowXmlWrite()
+				.Begin()
+				.MakeVertexSelectionBlock(DataflowActor->DataflowComponent->GetSelectionState().Vertices)
+				.End()
+				.ToString();
+			FPlatformApplicationMisc::ClipboardCopy(*XmlBuffer);
+		}
+	}
+
+
+	if (!bHandled)
+	{
+		bHandled = FEditorViewportClient::InputKey(EventArgs);
+	}
+
+	return bHandled;
+}
+
+void FDataflowEditorViewportClient::ProcessClick(FSceneView& View, HHitProxy* HitProxy, FKey Key, EInputEvent Event, uint32 HitX, uint32 HitY)
+{
+	Super::ProcessClick(View, HitProxy, Key, Event, HitX, HitY);
+	if (DataflowActor && DataflowActor->DataflowComponent)
+	{
+		const bool bIsShiftKeyDown = Viewport->KeyState(EKeys::LeftShift) || Viewport->KeyState(EKeys::RightShift);
+		const bool bIsCtrltKeyDown = Viewport->KeyState(EKeys::LeftControl) || Viewport->KeyState(EKeys::RightControl);
+
+		FDataflowSelectionState SelectionState = DataflowActor->DataflowComponent->GetSelectionState();
+		FDataflowSelectionState PreState = SelectionState;
+
+		if (SelectionMode == FDataflowSelectionState::EMode::DSS_Dataflow_Object)
+		{
+			if (HitProxy && HitProxy->IsA(HDataflowNode::StaticGetType()))
+			{
+				HDataflowNode* DataflowNode = (HDataflowNode*)(HitProxy);
+				FDataflowSelectionState::ObjectID ID(DataflowNode->NodeName, DataflowNode->GeometryIndex);
+
+				if (bIsShiftKeyDown)
+				{
+					if (!SelectionState.Nodes.Contains(ID))
+					{
+						SelectionState.Nodes.AddUnique(ID);
+					}
+				}
+				else if (bIsCtrltKeyDown)
+				{
+					if (SelectionState.Nodes.Contains(ID))
+					{
+						SelectionState.Nodes.Remove(ID);
+					}
+				}
+				else
+				{
+					SelectionState.Nodes.Empty();
+					SelectionState.Nodes.AddUnique(ID);
+				}
+			}
+			else if (!bIsShiftKeyDown && !bIsCtrltKeyDown)
+			{
+				SelectionState.Nodes.Empty();
+			}
+		}
+
+		if (SelectionMode == FDataflowSelectionState::EMode::DSS_Dataflow_Vertex)
+		{
+			if (HitProxy && HitProxy->IsA(HDataflowVertex::StaticGetType()))
+			{
+				HDataflowVertex* DataflowVertex = (HDataflowVertex*)(HitProxy);
+				int32 ID = DataflowVertex->SectionIndex;
+				if (bIsShiftKeyDown)
+				{
+					if (!SelectionState.Vertices.Contains(ID))
+					{
+						SelectionState.Vertices.AddUnique(ID);
+					}
+				}
+				else if (bIsCtrltKeyDown)
+				{
+					if (SelectionState.Vertices.Contains(ID))
+					{
+						SelectionState.Vertices.Remove(ID);
+					}
+				}
+				else
+				{
+					SelectionState.Vertices.Empty();
+					SelectionState.Vertices.AddUnique(ID);
+				}
+			}
+			else if (!bIsShiftKeyDown && !bIsCtrltKeyDown)
+			{
+				SelectionState.Vertices.Empty();
+			}
+		}
+
+		if (PreState != SelectionState)
+		{
+			DataflowActor->DataflowComponent->SetSelectionState(SelectionState);
+		}
+	}
+}
+
+void FDataflowEditorViewportClient::Tick(float DeltaSeconds)
+{
+	Super::Tick(DeltaSeconds);
+
+	TSharedPtr<FDataflowEditorToolkit> DataflowEditorToolkit = DataflowEditorToolkitPtr.Pin();
+
+	if (DataflowActor && DataflowEditorToolkitPtr.IsValid())
+	{
+		if (TSharedPtr<Dataflow::FContext> Context = DataflowEditorToolkit->GetContext())
+		{
+			if (const UDataflow* Dataflow = DataflowEditorToolkit->GetDataflow())
+			{
+				if (UDataflowComponent* DataflowComponent = DataflowActor->GetDataflowComponent())
+				{
+					Dataflow::FTimestamp SystemTimestamp = LatestTimestamp(Dataflow, Context.Get());
+					if (SystemTimestamp >= LastModifiedTimestamp)
+					{
+						if (Dataflow->GetRenderTargets().Num())
+						{
+							// Component Object Rendering
+							DataflowComponent->ResetRenderTargets();
+							DataflowComponent->SetDataflow(Dataflow);
+							DataflowComponent->SetContext(Context);
+							for (const UDataflowEdNode* Node : Dataflow->GetRenderTargets())
+							{
+								DataflowComponent->AddRenderTarget(Node);
+							}
+						}
+						else
+						{
+							DataflowComponent->ResetRenderTargets();
+						}
+
+						LastModifiedTimestamp = LatestTimestamp(Dataflow, Context.Get()).Value + 1;
+					}
+				}
+			}
+		}
+	}
+
+	// Tick the preview scene world.
+	if (!GIntraFrameDebuggingGameThread)
+	{
+		PreviewScene->GetWorld()->Tick(IsRealtime() ? LEVELTICK_All : LEVELTICK_TimeOnly, DeltaSeconds);
+	}
+}
+
+
+void FDataflowEditorViewportClient::AddReferencedObjects(FReferenceCollector& Collector)
+{
+	Super::AddReferencedObjects(Collector);
+}
+
+
diff --git a/Engine/Plugins/Experimental/Dataflow/Source/DataflowEditor/Private/Dataflow/DataflowGraphEditor.cpp b/Engine/Plugins/Experimental/Dataflow/Source/DataflowEditor/Private/Dataflow/DataflowGraphEditor.cpp
--- a/Engine/Plugins/Experimental/Dataflow/Source/DataflowEditor/Private/Dataflow/DataflowGraphEditor.cpp
+++ b/Engine/Plugins/Experimental/Dataflow/Source/DataflowEditor/Private/Dataflow/DataflowGraphEditor.cpp
@@ -2,21 +2,21 @@
 
 #include "Dataflow/DataflowGraphEditor.h"
 
+#include "BoneDragDropOp.h"
 #include "Dataflow/DataflowEditorToolkit.h"
 #include "Dataflow/DataflowEngine.h"
 #include "Dataflow/DataflowSNodeFactories.h"
-#include "BoneDragDropOp.h"
-#include "IStructureDetailsView.h"
 #include "Dataflow/DataflowSCommentNode.h"
 #include "Dataflow/DataflowNodeParameters.h"
-
+#include "Dataflow/DataflowXml.h"
+#include "Dataflow/DataflowSNode.h"
+#include "Dataflow/DataflowSelectionNodes.h"
 #include "Dataflow/DataflowSchema.h"
-
-DEFINE_LOG_CATEGORY_STATIC(SDataflowGraphEditorLog, Log, All);
+#include "HAL/PlatformApplicationMisc.h"
+#include "IStructureDetailsView.h"
 
 #define LOCTEXT_NAMESPACE "DataflowGraphEditor"
 
-
 void SDataflowGraphEditor::Construct(const FArguments& InArgs, UObject* InAssetOwner)
 {
 	check(InArgs._GraphToEdit);
@@ -146,9 +146,65 @@ void SDataflowGraphEditor::OnSelectedNodesChanged(const TSet<UObject*>& NewSelec
 	OnSelectionChangedMulticast.Broadcast(NewSelection);
 }
 
+FReply SDataflowGraphEditor::OnKeyDown(const FGeometry& MyGeometry, const FKeyEvent& InKeyEvent)
+{
+	if (InKeyEvent.GetKey() == EKeys::LeftControl)
+	{
+		LeftControlKeyDown = true;
+	}
+	if (InKeyEvent.GetKey() == EKeys::RightControl)
+	{
+		RightControlKeyDown = true;
+	}
+	if (InKeyEvent.GetKey() == EKeys::V)
+	{
+		VKeyDown = true;
+	}
+	return SGraphEditor::OnKeyUp(MyGeometry, InKeyEvent);
+}
+
+
+FReply SDataflowGraphEditor::OnKeyUp(const FGeometry& MyGeometry, const FKeyEvent& InKeyEvent)
+{
+	// process paste
+	if ( (VKeyDown && LeftControlKeyDown) || (VKeyDown && RightControlKeyDown) )
+	{
+		if (InKeyEvent.GetKey() == EKeys::LeftControl
+			|| InKeyEvent.GetKey() == EKeys::RightControl
+			|| InKeyEvent.GetKey() == EKeys::V)
+		{
+			FString XmlBuffer;
+			FPlatformApplicationMisc::ClipboardPaste(XmlBuffer);
+			FDataflowXmlRead Xml(this); 
+			if (Xml.LoadFromBuffer(XmlBuffer))
+			{
+				Xml.ParseXmlFile();
+			}
+		}
+	}
+
+	if (InKeyEvent.GetKey() == EKeys::LeftControl)
+	{
+		LeftControlKeyDown = false;
+	}
+	if (InKeyEvent.GetKey() == EKeys::RightControl)
+	{
+		RightControlKeyDown = false;
+	}
+	if (InKeyEvent.GetKey() == EKeys::V)
+	{
+		VKeyDown = false;
+	}
+	if (InKeyEvent.GetKey() == EKeys::LeftControl)
+	{
+		return FReply::Unhandled();
+	}
+	return SGraphEditor::OnKeyUp(MyGeometry, InKeyEvent);
+}
+
+
 FReply SDataflowGraphEditor::OnDragOver(const FGeometry& MyGeometry, const FDragDropEvent& DragDropEvent)
 {
-	UE_LOG(SDataflowGraphEditorLog, All, TEXT("SDataflowGraphEditor::OnDragOver"));
 	if (TSharedPtr<FBoneDragDropOp> SchemeDragDropOp = DragDropEvent.GetOperationAs<FBoneDragDropOp>())
 	{
 		return FReply::Handled();
@@ -163,7 +219,6 @@ FReply SDataflowGraphEditor::OnDrop(const FGeometry& MyGeometry, const FDragDrop
 	{
 		if (OnDragDropEventCallback)
 		{
-			UE_LOG(SDataflowGraphEditorLog, All, TEXT("SDataflowGraphEditor::OnDrop"));
 			OnDragDropEventCallback(MyGeometry,DragDropEvent);
 		}
 	}
@@ -204,6 +259,30 @@ void SDataflowGraphEditor::CreateComment()
 	CommentAction->PerformAction(Graph, nullptr, GetGraphEditor()->GetPasteLocation(), false);
 }
 
+void SDataflowGraphEditor::CreateVertexSelectionNode(const FString & InArray)
+{
+	UDataflow* Graph = DataflowAsset.Get();
+	const TSharedPtr<SGraphEditor>& InGraphEditor = SharedThis(GetGraphEditor());
+
+	if (TSharedPtr<FAssetSchemaAction_Dataflow_CreateNode_DataflowEdNode> NodeAction = FAssetSchemaAction_Dataflow_CreateNode_DataflowEdNode::CreateAction(Graph, "FSelectionSetDataflowNode"))
+	{
+		if (UEdGraphNode* NewEdNode = NodeAction->PerformAction(Graph, nullptr, GetGraphEditor()->GetPasteLocation(), false))
+		{
+			FDataflowAssetEdit Edit = Graph->EditDataflow();
+			if (Dataflow::FGraph* DataflowGraph = Edit.GetGraph())
+			{
+				if (TSharedPtr<FDataflowNode> Node = DataflowGraph->FindBaseNode(((UDataflowEdNode*)NewEdNode)->DataflowNodeGuid))
+				{
+					if (FSelectionSetDataflowNode * SelectionNode = Node->AsType<FSelectionSetDataflowNode>())
+					{
+						SelectionNode->Indices = InArray;
+					}
+				}
+			}
+		}
+	}
+}
+
 void SDataflowGraphEditor::AlignTop()
 {
 	GetGraphEditor()->OnAlignTop();
diff --git a/Engine/Plugins/Experimental/Dataflow/Source/DataflowEditor/Private/Dataflow/DataflowXml.cpp b/Engine/Plugins/Experimental/Dataflow/Source/DataflowEditor/Private/Dataflow/DataflowXml.cpp
new file mode 100644
--- /dev/null
+++ b/Engine/Plugins/Experimental/Dataflow/Source/DataflowEditor/Private/Dataflow/DataflowXml.cpp
@@ -0,0 +1,77 @@
+// Copyright Epic Games, Inc. All Rights Reserved.
+
+#include "Dataflow/DataflowXml.h"
+
+#include "Dataflow/DataflowGraphEditor.h"
+#include "Misc/MessageDialog.h"
+
+
+FString IntListToString(TArray<int32> & InList)
+{
+	FString Val;
+	for (int32 i : InList)
+	{
+		Val += FString::FromInt(i);
+		Val += TEXT(" ");
+	}
+	return Val;
+}
+
+FDataflowXmlWrite& FDataflowXmlWrite::Begin()
+{
+	Buffer += TEXT("<Dataflow>");
+	return *this;
+}
+
+FDataflowXmlWrite& FDataflowXmlWrite::MakeVertexSelectionBlock(TArray<int32> VertexList)
+{
+	Buffer += TEXT("<VertexSelection indices=\"");
+	Buffer += IntListToString(VertexList);
+	Buffer += TEXT("\"></VertexSelection>");
+	return *this;
+}
+
+FDataflowXmlWrite& FDataflowXmlWrite::End()
+{
+	Buffer += TEXT("</Dataflow>");
+	return *this;
+}
+
+bool FDataflowXmlRead::LoadFromBuffer(FString InBuffer)
+{
+	XmlFile = MakeUnique< FXmlFile >(InBuffer, EConstructMethod::ConstructFromBuffer);
+	const FXmlNode* RootNode = XmlFile->GetRootNode();
+	if (RootNode == NULL)
+	{
+		FMessageDialog::Open(EAppMsgType::Ok, FText::FromString(TEXT("Invalid dataflow xml-markup.")));
+		XmlFile.Reset();
+		return false;
+	}
+
+	if (RootNode->GetTag() != TEXT("Dataflow"))
+	{
+		FText DialogTitle = FText::FromString(TEXT("Error parsing xml-markup"));
+		FMessageDialog::Open(EAppMsgType::Ok, FText::FromString(TEXT("Invalid dataflow xml-markup.")));
+		XmlFile.Reset();
+	}
+
+	return true;
+}
+
+void FDataflowXmlRead::ParseXmlFile()
+{
+	if (XmlFile)
+	{
+		const TArray<FXmlNode*>& Nodes = XmlFile->GetRootNode()->GetChildrenNodes();
+
+		for (int i = 0; i < Nodes.Num(); i++)
+		{
+			if (Nodes[i]->GetTag() == TEXT("VertexSelection") )
+			{
+				FString Indices = Nodes[i]->GetAttribute(TEXT("indices"));
+				Editor->CreateVertexSelectionNode(Indices);
+			}
+		}
+	}
+}
+
diff --git a/Engine/Plugins/Experimental/Dataflow/Source/DataflowEditor/Public/Dataflow/DataflowEditorViewport.h b/Engine/Plugins/Experimental/Dataflow/Source/DataflowEditor/Public/Dataflow/DataflowEditorViewport.h
--- a/Engine/Plugins/Experimental/Dataflow/Source/DataflowEditor/Public/Dataflow/DataflowEditorViewport.h
+++ b/Engine/Plugins/Experimental/Dataflow/Source/DataflowEditor/Public/Dataflow/DataflowEditorViewport.h
@@ -3,11 +3,6 @@
 
 #include "CanvasTypes.h"
 #include "CoreMinimal.h"
-#include "Dataflow/DataflowComponent.h"
-#include "Dataflow/DataflowNodeParameters.h"
-#include "EditorViewportClient.h"
-#include "GeometryCollection/Facades/CollectionRenderingFacade.h"
-#include "GeometryCollection/ManagedArrayCollection.h"
 #include "SAssetEditorViewport.h"
 #include "SCommonEditorViewportToolbarBase.h"
 
@@ -17,7 +12,6 @@ class FEditorViewportClient;
 class FDataflowEditorViewportClient;
 class SEditorViewport;
 class ADataflowActor;
-class FDynamicMeshBuilder;
 
 // ----------------------------------------------------------------------------------
 
@@ -59,47 +53,3 @@ private:
 	TWeakPtr<FDataflowEditorToolkit> DataflowEditorToolkitPtr;
 	ADataflowActor* CustomDataflowActor = nullptr;
 };
-
-
-// ----------------------------------------------------------------------------------
-
-class FDataflowEditorViewportClient : public FEditorViewportClient
-{
-public:
-	using Super = FEditorViewportClient;
-	
-	FDataflowEditorViewportClient(FPreviewScene* InPreviewScene,
-		const TWeakPtr<SEditorViewport> InEditorViewportWidget = nullptr,
-		TWeakPtr<FDataflowEditorToolkit> InDataflowEditorToolkitPtr = nullptr);
-
-	Dataflow::FTimestamp LatestTimestamp(const UDataflow* Dataflow, const Dataflow::FContext* Context);
-	void SetDataflowActor(ADataflowActor* InActor) { DataflowActor = InActor; }
-
-	void SetSelectionMode(FDataflowSelectionState::EMode InState);
-	bool CanSetSelectionMode(FDataflowSelectionState::EMode InState);
-	bool IsSelectionModeActive(FDataflowSelectionState::EMode InState);
-	FDataflowSelectionState::EMode GetSelectionMode() const { return SelectionMode; }
-
-
-	// FEditorViewportClient interface
-	virtual void ProcessClick(FSceneView& View, HHitProxy* HitProxy, FKey Key, EInputEvent Event, uint32 HitX, uint32 HitY) override;
-	virtual void Tick(float DeltaSeconds) override;
-	// End of FEditorViewportClient
-
-	//~ FGCObject Interface
-	virtual void AddReferencedObjects(FReferenceCollector& Collector) override;
-	virtual FString GetReferencerName() const override { return TEXT("FDataflowEditorViewportClient"); }
-
-private:
-
-	TWeakPtr<FDataflowEditorToolkit> DataflowEditorToolkitPtr = nullptr;
-	Dataflow::FTimestamp LastModifiedTimestamp = Dataflow::FTimestamp::Invalid;
-	ADataflowActor* DataflowActor = nullptr;
-
-	//
-	// Selection
-	//
-
-	FDataflowSelectionState::EMode SelectionMode = FDataflowSelectionState::EMode::DSS_Dataflow_None;
-
-};
diff --git a/Engine/Plugins/Experimental/Dataflow/Source/DataflowEditor/Public/Dataflow/DataflowEditorViewportClient.h b/Engine/Plugins/Experimental/Dataflow/Source/DataflowEditor/Public/Dataflow/DataflowEditorViewportClient.h
new file mode 100644
--- /dev/null
+++ b/Engine/Plugins/Experimental/Dataflow/Source/DataflowEditor/Public/Dataflow/DataflowEditorViewportClient.h
@@ -0,0 +1,54 @@
+// Copyright Epic Games, Inc. All Rights Reserved.
+#pragma once
+
+#include "CoreMinimal.h"
+#include "EditorViewportClient.h"
+#include "Dataflow/DataflowNodeParameters.h"
+#include "Dataflow/DataflowComponentSelectionState.h"
+
+class FDataflowEditorToolkit;
+class ADataflowActor;
+
+// ----------------------------------------------------------------------------------
+
+class FDataflowEditorViewportClient : public FEditorViewportClient
+{
+public:
+	using Super = FEditorViewportClient;
+	
+	FDataflowEditorViewportClient(FPreviewScene* InPreviewScene,
+		const TWeakPtr<SEditorViewport> InEditorViewportWidget = nullptr,
+		TWeakPtr<FDataflowEditorToolkit> InDataflowEditorToolkitPtr = nullptr);
+
+	Dataflow::FTimestamp LatestTimestamp(const UDataflow* Dataflow, const Dataflow::FContext* Context);
+	void SetDataflowActor(ADataflowActor* InActor) { DataflowActor = InActor; }
+
+	void SetSelectionMode(FDataflowSelectionState::EMode InState);
+	bool CanSetSelectionMode(FDataflowSelectionState::EMode InState);
+	bool IsSelectionModeActive(FDataflowSelectionState::EMode InState);
+	FDataflowSelectionState::EMode GetSelectionMode() const { return SelectionMode; }
+
+
+	// FEditorViewportClient interface
+	virtual bool InputKey(const FInputKeyEventArgs& EventArgs) override;
+	virtual void ProcessClick(FSceneView& View, HHitProxy* HitProxy, FKey Key, EInputEvent Event, uint32 HitX, uint32 HitY) override;
+	virtual void Tick(float DeltaSeconds) override;
+	// End of FEditorViewportClient
+
+	//~ FGCObject Interface
+	virtual void AddReferencedObjects(FReferenceCollector& Collector) override;
+	virtual FString GetReferencerName() const override { return TEXT("FDataflowEditorViewportClient"); }
+
+private:
+
+	TWeakPtr<FDataflowEditorToolkit> DataflowEditorToolkitPtr = nullptr;
+	Dataflow::FTimestamp LastModifiedTimestamp = Dataflow::FTimestamp::Invalid;
+	ADataflowActor* DataflowActor = nullptr;
+
+	//
+	// Selection
+	//
+
+	FDataflowSelectionState::EMode SelectionMode = FDataflowSelectionState::EMode::DSS_Dataflow_None;
+
+};
diff --git a/Engine/Plugins/Experimental/Dataflow/Source/DataflowEditor/Public/Dataflow/DataflowGraphEditor.h b/Engine/Plugins/Experimental/Dataflow/Source/DataflowEditor/Public/Dataflow/DataflowGraphEditor.h
--- a/Engine/Plugins/Experimental/Dataflow/Source/DataflowEditor/Public/Dataflow/DataflowGraphEditor.h
+++ b/Engine/Plugins/Experimental/Dataflow/Source/DataflowEditor/Public/Dataflow/DataflowGraphEditor.h
@@ -55,6 +55,8 @@ public:
 	FOnSelectionChangedMulticast OnSelectionChangedMulticast;
 
 	// SWidget overrides
+	virtual FReply OnKeyDown(const FGeometry& MyGeometry, const FKeyEvent& InKeyEvent) override;
+	virtual FReply OnKeyUp(const FGeometry& MyGeometry, const FKeyEvent& InKeyEvent) override;
 	virtual FReply OnDragOver(const FGeometry& MyGeometry, const FDragDropEvent& DragDropEvent) override;
 	virtual FReply OnDrop(const FGeometry& MyGeometry, const FDragDropEvent& DragDropEvent) override;
 	//virtual void OnDragEnter(const FGeometry& MyGeometry, const FDragDropEvent& DragDropEvent) override;
@@ -76,6 +78,9 @@ public:
 	/** */
 	void CreateComment();
 
+	/** */
+	void CreateVertexSelectionNode(const FString& InArray);
+
 	/** */
 	void AlignTop();
 
@@ -126,4 +131,8 @@ private:
 
 	/** The details view that responds to this widget. */
 	TSharedPtr<IStructureDetailsView> DetailsView;
+
+	bool VKeyDown = false;
+	bool LeftControlKeyDown = false;
+	bool RightControlKeyDown = false;
 };
diff --git a/Engine/Plugins/Experimental/Dataflow/Source/DataflowEditor/Public/Dataflow/DataflowXml.h b/Engine/Plugins/Experimental/Dataflow/Source/DataflowEditor/Public/Dataflow/DataflowXml.h
new file mode 100644
--- /dev/null
+++ b/Engine/Plugins/Experimental/Dataflow/Source/DataflowEditor/Public/Dataflow/DataflowXml.h
@@ -0,0 +1,41 @@
+// Copyright Epic Games, Inc. All Rights Reserved.
+
+#pragma once
+
+#include "CoreMinimal.h"
+#include "Serialization/ArchiveProxy.h"
+#include "XmlParser.h"
+
+class SDataflowGraphEditor;
+
+/*
+* FDataflowEditorXml
+*
+*/
+
+class FDataflowXmlWrite
+{
+public:
+	FDataflowXmlWrite() {}
+
+	FDataflowXmlWrite& Begin();
+	FDataflowXmlWrite& MakeVertexSelectionBlock(TArray<int32> InVertices);
+	FDataflowXmlWrite& End();
+	FString ToString() { return Buffer; }
+private:
+
+	FString Buffer;
+};
+
+class FDataflowXmlRead
+{
+public:
+	FDataflowXmlRead(SDataflowGraphEditor* InEditor) : Editor(InEditor) { check(Editor); };
+
+	bool LoadFromBuffer(FString InBuffer);
+	void ParseXmlFile();
+
+private:
+	SDataflowGraphEditor* Editor;
+	TUniquePtr<FXmlFile>  XmlFile;
+};
\ No newline at end of file
diff --git a/Engine/Plugins/Experimental/Dataflow/Source/DataflowNodes/Private/Dataflow/DataflowNodesPlugin.cpp b/Engine/Plugins/Experimental/Dataflow/Source/DataflowNodes/Private/Dataflow/DataflowNodesPlugin.cpp
--- a/Engine/Plugins/Experimental/Dataflow/Source/DataflowNodes/Private/Dataflow/DataflowNodesPlugin.cpp
+++ b/Engine/Plugins/Experimental/Dataflow/Source/DataflowNodes/Private/Dataflow/DataflowNodesPlugin.cpp
@@ -6,6 +6,7 @@
 #include "Dataflow/DataflowSkeletalMeshNodes.h"
 #include "Dataflow/DataflowStaticMeshNodes.h"
 #include "Dataflow/DataflowNodeFactory.h"
+#include "Dataflow/DataflowSelectionNodes.h"
 
 #define LOCTEXT_NAMESPACE "DataflowNodes"
 
@@ -14,6 +15,7 @@ void IDataflowNodesPlugin::StartupModule()
 {
 	Dataflow::RegisterSkeletalMeshNodes();
 	Dataflow::RegisterStaticMeshNodes();
+	Dataflow::RegisterSelectionNodes();
 }
 
 void IDataflowNodesPlugin::ShutdownModule()
diff --git a/Engine/Plugins/Experimental/Dataflow/Source/DataflowNodes/Private/Dataflow/DataflowSelectionNodes.cpp b/Engine/Plugins/Experimental/Dataflow/Source/DataflowNodes/Private/Dataflow/DataflowSelectionNodes.cpp
new file mode 100644
--- /dev/null
+++ b/Engine/Plugins/Experimental/Dataflow/Source/DataflowNodes/Private/Dataflow/DataflowSelectionNodes.cpp
@@ -0,0 +1,36 @@
+// Copyright Epic Games, Inc. All Rights Reserved.
+
+#include "Dataflow/DataflowSelectionNodes.h"
+#include "Engine/StaticMesh.h"
+
+#include UE_INLINE_GENERATED_CPP_BY_NAME(DataflowSelectionNodes)
+
+namespace Dataflow
+{
+	void RegisterSelectionNodes()
+	{
+		DATAFLOW_NODE_REGISTER_CREATION_FACTORY(FSelectionSetDataflowNode);
+	}
+};
+
+void FSelectionSetDataflowNode::Evaluate(Dataflow::FContext& Context, const FDataflowOutput* Out) const
+{
+	if (Out->IsA<DataType>(&IndicesOut))
+	{
+		TArray<int32> IntArray;
+		TArray<FString> StrArray;
+		Indices.ParseIntoArray(StrArray, *FString(" "));
+		for (FString Elem : StrArray)
+		{
+			if (FCString::IsNumeric(*Elem))
+			{
+				IntArray.Add(FCString::Atoi(*Elem));
+			}
+		}
+
+		Out->SetValue<DataType>(IntArray, Context);
+	}
+}
+
+
+
diff --git a/Engine/Plugins/Experimental/Dataflow/Source/DataflowNodes/Public/Dataflow/DataflowSelectionNodes.h b/Engine/Plugins/Experimental/Dataflow/Source/DataflowNodes/Public/Dataflow/DataflowSelectionNodes.h
new file mode 100644
--- /dev/null
+++ b/Engine/Plugins/Experimental/Dataflow/Source/DataflowNodes/Public/Dataflow/DataflowSelectionNodes.h
@@ -0,0 +1,42 @@
+// Copyright Epic Games, Inc. All Rights Reserved.
+
+#pragma once 
+
+#include "CoreMinimal.h"
+#include "Dataflow/DataflowCore.h"
+#include "Dataflow/DataflowEngine.h"
+
+#include "DataflowSelectionNodes.generated.h"
+
+
+
+USTRUCT(meta = (DataflowFlesh))
+struct DATAFLOWNODES_API FSelectionSetDataflowNode : public FDataflowNode
+{
+	GENERATED_USTRUCT_BODY()
+	DATAFLOW_NODE_DEFINE_INTERNAL(FSelectionSetDataflowNode, "SelectionSet", "Dataflow", "")
+
+public:
+	typedef TArray<int32> DataType;
+
+	UPROPERTY(EditAnywhere, Category = "Dataflow")
+		FString Indices = FString("1 2 3");
+
+	UPROPERTY(meta = (DataflowOutput, DisplayName = "SelectionSet"))
+		TArray<int32> IndicesOut;
+
+	FSelectionSetDataflowNode(const Dataflow::FNodeParameters& InParam, FGuid InGuid = FGuid::NewGuid())
+		: FDataflowNode(InParam, InGuid)
+	{
+		RegisterOutputConnection(&IndicesOut);
+	}
+
+	virtual void Evaluate(Dataflow::FContext& Context, const FDataflowOutput* Out) const override;
+};
+
+
+namespace Dataflow
+{
+	void RegisterSelectionNodes();
+}
+
diff --git a/Engine/Source/Runtime/Experimental/Dataflow/Core/Public/Dataflow/DataflowNode.h b/Engine/Source/Runtime/Experimental/Dataflow/Core/Public/Dataflow/DataflowNode.h
--- a/Engine/Source/Runtime/Experimental/Dataflow/Core/Public/Dataflow/DataflowNode.h
+++ b/Engine/Source/Runtime/Experimental/Dataflow/Core/Public/Dataflow/DataflowNode.h
@@ -220,6 +220,17 @@ struct DATAFLOWCORE_API FDataflowNode
 		return nullptr;
 	}
 
+	template<class T>
+	T* AsType()
+	{
+		FName TargetType = T::StaticType();
+		if (IsA(TargetType))
+		{
+			return (T*)this;
+		}
+		return nullptr;
+	}
+
 
 private:
 	virtual TArray<Dataflow::FRenderingParameter> GetRenderParametersImpl() const { return TArray<Dataflow::FRenderingParameter>(); }
