Dataflow : Moved selection controls behind a cvar.

diff --git a/Engine/Plugins/Experimental/Dataflow/Source/DataflowEditor/Private/Dataflow/DataflowEditorViewportToolbar.cpp b/Engine/Plugins/Experimental/Dataflow/Source/DataflowEditor/Private/Dataflow/DataflowEditorViewportToolbar.cpp
--- a/Engine/Plugins/Experimental/Dataflow/Source/DataflowEditor/Private/Dataflow/DataflowEditorViewportToolbar.cpp
+++ b/Engine/Plugins/Experimental/Dataflow/Source/DataflowEditor/Private/Dataflow/DataflowEditorViewportToolbar.cpp
@@ -5,8 +5,13 @@
 #include "Dataflow/DataflowEditorStyle.h"
 #include "Dataflow/DataflowEditorViewport.h"
 #include "EditorViewportCommands.h"
+#include "HAL/IConsoleManager.h"
 #include "Widgets/Layout/SBorder.h"
 
+bool bDataflowPrototypeSelectionMenu = false;
+FAutoConsoleVariableRef CVarDataflowPrototypeSelectionMenu(TEXT("p.Dataflow.prototype.selection"), bDataflowPrototypeSelectionMenu, TEXT("Work in progress development for selection controls in the 3D Viewport.[def:false]"));
+
+
 void SDataflowViewportSelectionToolBar::Construct(const FArguments& InArgs)
 {
 	EditorViewport = InArgs._EditorViewport;
@@ -46,30 +51,30 @@ TSharedRef<SWidget> SDataflowViewportSelectionToolBar::MakeSelectionModeToolBar(
 	{
 		ToolbarBuilder.BeginBlockGroup();
 
-		ToolbarBuilder.AddToolBarButton(FDataflowEditorCommands::Get().ToggleObjectSelection,
-			NAME_None,
-			TAttribute<FText>(),
-			TAttribute<FText>(),
-			FSlateIcon(FDataflowEditorStyle::Get().GetStyleSetName(), TEXT("Dataflow.SelectObject")),
-			FName(TEXT("SelectObjectMode"))
-		);
-
-		ToolbarBuilder.AddToolBarButton(FDataflowEditorCommands::Get().ToggleFaceSelection,
-										NAME_None,
-										TAttribute<FText>(),
-										TAttribute<FText>(),
-										FSlateIcon(FDataflowEditorStyle::Get().GetStyleSetName(), TEXT("Dataflow.SelectFace")),
-										FName(TEXT("SelectFaceMode"))
-		);
-
-		ToolbarBuilder.AddToolBarButton(FDataflowEditorCommands::Get().ToggleVertexSelection,
-										NAME_None,
-										TAttribute<FText>(),
-										TAttribute<FText>(),
-										FSlateIcon(FDataflowEditorStyle::Get().GetStyleSetName(), TEXT("Dataflow.SelectVertex")),
-										FName(TEXT("SelectVertexMode"))
-		);
-
+		if (bDataflowPrototypeSelectionMenu)
+		{
+			ToolbarBuilder.AddToolBarButton(FDataflowEditorCommands::Get().ToggleObjectSelection,
+											NAME_None,
+											TAttribute<FText>(),
+											TAttribute<FText>(),
+											FSlateIcon(FDataflowEditorStyle::Get().GetStyleSetName(), TEXT("Dataflow.SelectObject")),
+											FName(TEXT("SelectObjectMode"))
+			);
+			ToolbarBuilder.AddToolBarButton(FDataflowEditorCommands::Get().ToggleFaceSelection,
+											NAME_None,
+											TAttribute<FText>(),
+											TAttribute<FText>(),
+											FSlateIcon(FDataflowEditorStyle::Get().GetStyleSetName(), TEXT("Dataflow.SelectFace")),
+											FName(TEXT("SelectFaceMode"))
+			);
+			ToolbarBuilder.AddToolBarButton(FDataflowEditorCommands::Get().ToggleVertexSelection,
+											NAME_None,
+											TAttribute<FText>(),
+											TAttribute<FText>(),
+											FSlateIcon(FDataflowEditorStyle::Get().GetStyleSetName(), TEXT("Dataflow.SelectVertex")),
+											FName(TEXT("SelectVertexMode"))
+			);
+		}
 
 		ToolbarBuilder.EndBlockGroup();
 	}
