Dataflow : Tick component render target access issue.

diff --git a/Engine/Plugins/Experimental/Dataflow/Source/DataflowEnginePlugin/Private/Dataflow/DataflowComponent.cpp b/Engine/Plugins/Experimental/Dataflow/Source/DataflowEnginePlugin/Private/Dataflow/DataflowComponent.cpp
--- a/Engine/Plugins/Experimental/Dataflow/Source/DataflowEnginePlugin/Private/Dataflow/DataflowComponent.cpp
+++ b/Engine/Plugins/Experimental/Dataflow/Source/DataflowEnginePlugin/Private/Dataflow/DataflowComponent.cpp
@@ -66,7 +66,10 @@ void UDataflowComponent::TickComponent(float DeltaTime, enum ELevelTick TickType
 		{
 			for (const UDataflowEdNode* Target : RenderTargets)
 			{
-				bNeedsRefresh |= Target->Render(Facade, Context);
+				if (Target)
+				{
+					bNeedsRefresh |= Target->Render(Facade, Context);
+				}
 			}
 		}
 
