Dataflow:
- Added SelectionView window
- Bug fixes

diff --git a/Engine/Plugins/Experimental/Dataflow/Source/DataflowEditor/Private/Dataflow/DataflowEditorCommands.cpp b/Engine/Plugins/Experimental/Dataflow/Source/DataflowEditor/Private/Dataflow/DataflowEditorCommands.cpp
--- a/Engine/Plugins/Experimental/Dataflow/Source/DataflowEditor/Private/Dataflow/DataflowEditorCommands.cpp
+++ b/Engine/Plugins/Experimental/Dataflow/Source/DataflowEditor/Private/Dataflow/DataflowEditorCommands.cpp
@@ -226,28 +226,24 @@ void FDataflowEditorCommands::OnNodeTitleCommitted(const FText& InNewText, EText
 }
 
 
-void FDataflowEditorCommands::OnPropertyValueChanged(UDataflow* OutDataflow, TSharedPtr<Dataflow::FEngineContext>& Context, Dataflow::FTimestamp& OutLastNodeTimestamp, const FPropertyChangedEvent& InPropertyChangedEvent)
+void FDataflowEditorCommands::OnPropertyValueChanged(UDataflow* OutDataflow, TSharedPtr<Dataflow::FEngineContext>& Context, Dataflow::FTimestamp& OutLastNodeTimestamp, const FPropertyChangedEvent& InPropertyChangedEvent, const TSet<UObject*>& SelectedNodes)
 {
 	if (InPropertyChangedEvent.ChangeType == EPropertyChangeType::ValueSet)
 	{
 		TSharedPtr<const FDataflowNode> UpdatedNode = nullptr;
 		if (OutDataflow && InPropertyChangedEvent.Property && InPropertyChangedEvent.Property->GetOwnerUObject())
 		{
-			OutDataflow->MarkPackageDirty();
+//			OutDataflow->MarkPackageDirty();
+			OutDataflow->Modify();
 
-			FString Name = InPropertyChangedEvent.Property->GetOwnerUObject()->GetName();
-			Name = Name.Left(Name.Len() - FString("DataflowNode").Len());
+			for (UObject* SelectedNode : SelectedNodes)
 			{
-				FDataflowAssetEdit Edit = OutDataflow->EditDataflow();
-
-				Dataflow::FGraph* Graph = Edit.GetGraph();
-				for (auto Node : Graph->GetNodes())
+				if (UDataflowEdNode* Node = Cast<UDataflowEdNode>(SelectedNode))
 				{
-					FString NodeName = Node->GetName().ToString();
-					if (NodeName.StartsWith(Name))
+					if (TSharedPtr<FDataflowNode> DataflowNode = Node->GetDataflowNode())
 					{
-						UpdatedNode = Node;
-						Node->Invalidate();
+						UpdatedNode = DataflowNode;
+						DataflowNode->Invalidate();
 					}
 				}
 			}
diff --git a/Engine/Plugins/Experimental/Dataflow/Source/DataflowEditor/Private/Dataflow/DataflowEditorToolkit.cpp b/Engine/Plugins/Experimental/Dataflow/Source/DataflowEditor/Private/Dataflow/DataflowEditorToolkit.cpp
--- a/Engine/Plugins/Experimental/Dataflow/Source/DataflowEditor/Private/Dataflow/DataflowEditorToolkit.cpp
+++ b/Engine/Plugins/Experimental/Dataflow/Source/DataflowEditor/Private/Dataflow/DataflowEditorToolkit.cpp
@@ -38,6 +38,7 @@ const FName FDataflowEditorToolkit::GraphCanvasTabId(TEXT("DataflowEditor_GraphC
 const FName FDataflowEditorToolkit::AssetDetailsTabId(TEXT("DataflowEditor_AssetDetails"));
 const FName FDataflowEditorToolkit::NodeDetailsTabId(TEXT("DataflowEditor_NodeDetails"));
 const FName FDataflowEditorToolkit::SkeletalTabId(TEXT("DataflowEditor_Skeletal"));
+const FName FDataflowEditorToolkit::SelectionViewTabId(TEXT("DataflowEditor_SelectionView"));
 
 
 UDataflow* GetDataflowFrom(UObject* InObject)
@@ -88,6 +89,15 @@ bool FDataflowEditorToolkit::CanOpenDataflowEditor(UObject* ObjectToEdit)
 	return Dataflow != nullptr;
 }
 
+FDataflowEditorToolkit::~FDataflowEditorToolkit()
+{
+	if (GraphEditor)
+	{
+		GraphEditor->OnSelectionChangedMulticast.Remove(OnSelectionChangedMulticastDelegateHandle);
+		GraphEditor->OnNodeDeletedMulticast.Remove(OnNodeDeletedMulticastDelegateHandle);
+	}
+}
+
 void FDataflowEditorToolkit::InitializeEditor(const EToolkitMode::Type Mode, const TSharedPtr<IToolkitHost>& InitToolkitHost, UObject* ObjectToEdit)
 {
 	Asset = nullptr;
@@ -104,6 +114,9 @@ void FDataflowEditorToolkit::InitializeEditor(const EToolkitMode::Type Mode, con
 
 	if (Dataflow != nullptr)
 	{
+		Context = TSharedPtr< Dataflow::FEngineContext>(new Dataflow::FAssetContext(Asset, Dataflow, FPlatformTime::Cycles64()));
+		LastNodeTimestamp = Context->GetTimestamp();
+
 		Dataflow->Schema = UDataflowSchema::StaticClass();
 
 		NodeDetailsEditor = CreateNodeDetailsEditorWidget(ObjectToEdit);
@@ -111,9 +124,6 @@ void FDataflowEditorToolkit::InitializeEditor(const EToolkitMode::Type Mode, con
 		GraphEditor = CreateGraphEditorWidget(Dataflow, NodeDetailsEditor);
 		SkeletalEditor = CreateSkeletalEditorWidget(ObjectToEdit);
 
-		Context = TSharedPtr< Dataflow::FEngineContext>(new Dataflow::FAssetContext(Asset, Dataflow, FPlatformTime::Cycles64()));
-		LastNodeTimestamp = Context->GetTimestamp();
-
 		const TSharedRef<FTabManager::FLayout> StandaloneDefaultLayout = FTabManager::NewLayout("Dataflow_Layout.V1")
 			->AddArea
 			(
@@ -161,7 +171,7 @@ void FDataflowEditorToolkit::InitializeEditor(const EToolkitMode::Type Mode, con
 
 void FDataflowEditorToolkit::OnPropertyValueChanged(const FPropertyChangedEvent& PropertyChangedEvent)
 {
-	FDataflowEditorCommands::OnPropertyValueChanged(this->GetDataflow(), Context, LastNodeTimestamp, PropertyChangedEvent);
+	FDataflowEditorCommands::OnPropertyValueChanged(this->GetDataflow(), Context, LastNodeTimestamp, PropertyChangedEvent, PrevNodeSelection);
 }
 
 bool FDataflowEditorToolkit::OnNodeVerifyTitleCommit(const FText& NewText, UEdGraphNode* GraphNode, FText& OutErrorMessage)
@@ -174,6 +184,74 @@ void FDataflowEditorToolkit::OnNodeTitleCommitted(const FText& InNewText, ETextC
 	FDataflowEditorCommands::OnNodeTitleCommitted(InNewText, InCommitType, GraphNode);
 }
 
+void FDataflowEditorToolkit::OnNodeSelectionChanged(const TSet<UObject*>& NewSelection)
+{
+	if (Dataflow)
+	{
+		// Only keep UDataflowEdNode from NewSelection
+		TSet<UObject*> ValidatedSelection;
+
+		for (UObject* Item : NewSelection)
+		{
+			if (UDataflowEdNode* Node = Cast<UDataflowEdNode>(Item))
+			{
+				ValidatedSelection.Add(Item);
+			}
+		}
+
+		if (ValidatedSelection.Num() > 0)
+		{
+			TSet<UObject*> SelectionToUse, SelectionDifference;
+
+			if (PrevNodeSelection.Num() > 0)
+			{
+				SelectionDifference = ValidatedSelection.Difference(PrevNodeSelection);
+
+				if (SelectionDifference.Num() > 0)
+				{
+					SelectionToUse.Add(SelectionDifference.Array()[SelectionDifference.Num() - 1]);
+				}
+				else
+				{
+					SelectionToUse.Add(ValidatedSelection.Array()[ValidatedSelection.Num() - 1]);
+				}
+			}
+			else
+			{
+				SelectionToUse.Add(ValidatedSelection.Array()[ValidatedSelection.Num() - 1]);
+			}
+
+			if (UDataflowEdNode* Node = Cast<UDataflowEdNode>(SelectionToUse.Array()[0]))
+			{
+				for (IDataflowViewListener* Listener : ViewListeners)
+				{
+					Listener->OnSelectedNodeChanged(Node);
+				}
+			}
+		}
+		else
+		{
+			for (IDataflowViewListener* Listener : ViewListeners)
+			{
+				Listener->OnSelectedNodeChanged(nullptr);
+			}
+		}
+		
+		PrevNodeSelection = ValidatedSelection;
+	}
+}
+
+void FDataflowEditorToolkit::OnNodeDeleted(const TSet<UObject*>& NewSelection)
+{
+	for (UObject* Node : NewSelection)
+	{
+		if (PrevNodeSelection.Contains(Node))
+		{
+			PrevNodeSelection.Remove(Node);
+		}
+	}
+}
+
 void FDataflowEditorToolkit::Tick(float DeltaTime)
 {
 	if (Dataflow && Asset)
@@ -193,7 +271,7 @@ TStatId FDataflowEditorToolkit::GetStatId() const
 	RETURN_QUICK_DECLARE_CYCLE_STAT(FDataflowEditorToolkit, STATGROUP_Tickables);
 }
 
-TSharedRef<SGraphEditor> FDataflowEditorToolkit::CreateGraphEditorWidget(UDataflow* DataflowToEdit, TSharedPtr<IStructureDetailsView> InNodeDetailsEditor)
+TSharedRef<SDataflowGraphEditor> FDataflowEditorToolkit::CreateGraphEditorWidget(UDataflow* DataflowToEdit, TSharedPtr<IStructureDetailsView> InNodeDetailsEditor)
 {
 	ensure(DataflowToEdit);
 	using namespace Dataflow;
@@ -213,12 +291,16 @@ TSharedRef<SGraphEditor> FDataflowEditorToolkit::CreateGraphEditorWidget(UDatafl
 	InEvents.OnVerifyTextCommit = FOnNodeVerifyTextCommit::CreateSP(this, &FDataflowEditorToolkit::OnNodeVerifyTitleCommit);
 	InEvents.OnTextCommitted = FOnNodeTextCommitted::CreateSP(this, &FDataflowEditorToolkit::OnNodeTitleCommitted);
 
-	return SNew(SDataflowGraphEditor, DataflowToEdit)
+	TSharedRef<SDataflowGraphEditor> NewGraphEditor = SNew(SDataflowGraphEditor, DataflowToEdit)
 		.GraphToEdit(DataflowToEdit)
 		.GraphEvents(InEvents)
 		.DetailsView(InNodeDetailsEditor)
 		.EvaluateGraph(Evaluate);
 
+	OnSelectionChangedMulticastDelegateHandle = NewGraphEditor->OnSelectionChangedMulticast.AddSP(this, &FDataflowEditorToolkit::OnNodeSelectionChanged);
+	OnNodeDeletedMulticastDelegateHandle = NewGraphEditor->OnNodeDeletedMulticast.AddSP(this, &FDataflowEditorToolkit::OnNodeDeleted);
+
+	return NewGraphEditor;
 }
 
 TSharedPtr<IStructureDetailsView> FDataflowEditorToolkit::CreateNodeDetailsEditorWidget(UObject* ObjectToEdit)
@@ -378,6 +460,35 @@ TSharedRef<SDockTab> FDataflowEditorToolkit::SpawnTab_Skeletal(const FSpawnTabAr
 		];
 }
 
+
+TSharedRef<SDockTab> FDataflowEditorToolkit::SpawnTab_SelectionView(const FSpawnTabArgs& Args)
+{
+	check(Args.GetTabId() == SelectionViewTabId);
+
+	DataflowSelectionView = MakeShared<FDataflowSelectionView>(FDataflowSelectionView());
+	ViewListeners.Add(DataflowSelectionView.Get());
+
+	TSharedPtr<SSelectionViewWidget> SelectionViewWidget;
+
+	TSharedRef<SDockTab> DockableTab = SNew(SDockTab)
+	[
+		SAssignNew(SelectionViewWidget, SSelectionViewWidget)
+	];
+
+	if (SelectionViewWidget)
+	{
+		DataflowSelectionView->SetSelectionView(SelectionViewWidget);
+
+		// Set the Context on the interface
+		if (TSharedPtr<Dataflow::FContext> CurrentContext = this->GetContext())
+		{
+			DataflowSelectionView->SetContext(CurrentContext);
+		}
+	}
+
+	return DockableTab;
+}
+
 void FDataflowEditorToolkit::RegisterTabSpawners(const TSharedRef<FTabManager>& InTabManager)
 {
 	TSharedRef<FWorkspaceItem> WorkspaceMenuCategoryRef = InTabManager->AddLocalWorkspaceMenuCategory(LOCTEXT("WorkspaceMenu_DataflowEditor", "Dataflow Editor"));
@@ -407,6 +518,11 @@ void FDataflowEditorToolkit::RegisterTabSpawners(const TSharedRef<FTabManager>&
 		.SetGroup(WorkspaceMenuCategoryRef)
 		.SetIcon(FSlateIcon(FAppStyle::GetAppStyleSetName(), "LevelEditor.Tabs.SkeletalHierarchy"));
 
+	InTabManager->RegisterTabSpawner(SelectionViewTabId, FOnSpawnTab::CreateSP(this, &FDataflowEditorToolkit::SpawnTab_SelectionView))
+		.SetDisplayName(LOCTEXT("DataflowSelectionViewTab", "Selection View"))
+		.SetGroup(WorkspaceMenuCategoryRef)
+		.SetIcon(FSlateIcon(FAppStyle::GetAppStyleSetName(), "LevelEditor.Tabs.SelectionView"));
+
 	FAssetEditorToolkit::RegisterTabSpawners(InTabManager);
 }
 
diff --git a/Engine/Plugins/Experimental/Dataflow/Source/DataflowEditor/Private/Dataflow/DataflowGraphEditor.cpp b/Engine/Plugins/Experimental/Dataflow/Source/DataflowEditor/Private/Dataflow/DataflowGraphEditor.cpp
--- a/Engine/Plugins/Experimental/Dataflow/Source/DataflowEditor/Private/Dataflow/DataflowGraphEditor.cpp
+++ b/Engine/Plugins/Experimental/Dataflow/Source/DataflowEditor/Private/Dataflow/DataflowGraphEditor.cpp
@@ -132,7 +132,10 @@ void SDataflowGraphEditor::DeleteNode()
 			DetailsView->SetStructureData(nullptr);
 		}
 
-		FDataflowEditorCommands::DeleteNodes(DataflowAsset.Get(), GetSelectedNodes());
+		const FGraphPanelSelectionSet& SelectedNodes = GetSelectedNodes();
+		FDataflowEditorCommands::DeleteNodes(DataflowAsset.Get(), SelectedNodes);
+
+		OnNodeDeletedMulticast.Broadcast(SelectedNodes);
 	}
 }
 
diff --git a/Engine/Plugins/Experimental/Dataflow/Source/DataflowEditor/Private/Dataflow/DataflowSelectionView.cpp b/Engine/Plugins/Experimental/Dataflow/Source/DataflowEditor/Private/Dataflow/DataflowSelectionView.cpp
new file mode 100644
--- /dev/null
+++ b/Engine/Plugins/Experimental/Dataflow/Source/DataflowEditor/Private/Dataflow/DataflowSelectionView.cpp
@@ -0,0 +1,93 @@
+// Copyright Epic Games, Inc. All Rights Reserved.
+
+#include "Dataflow/DataflowSelectionView.h"
+#include "Templates/EnableIf.h"
+#include "Dataflow/DataflowSelection.h"
+
+
+void FDataflowSelectionView::SetSupportedOutputTypes()
+{
+	GetSupportedOutputTypes().Empty();
+
+	GetSupportedOutputTypes().Add("FDataflowTransformSelection");
+	GetSupportedOutputTypes().Add("FDataflowVertexSelection");
+	GetSupportedOutputTypes().Add("FDataflowFaceSelection");
+}
+
+void FDataflowSelectionView::UpdateViewData()
+{
+	if (SelectionView)
+	{
+		SelectionView->GetSelectionTable()->GetSelectionInfoMap().Empty();
+
+		if (GetSelectedNode())
+		{
+			if (GetSelectedNode()->IsBound())
+			{
+				if (TSharedPtr<FDataflowNode> DataflowNode = GetSelectedNode()->DataflowGraph->FindBaseNode(GetSelectedNode()->DataflowNodeGuid))
+				{
+					TArray<FDataflowOutput*> Outputs = DataflowNode->GetOutputs();
+
+					for (FDataflowOutput* Output : Outputs)
+					{
+						FName Name = Output->GetName();
+						FName Type = Output->GetType();
+
+						if (Output->GetType() == "FDataflowTransformSelection")
+						{
+							const FDataflowTransformSelection& Value = Output->GetValue<FDataflowTransformSelection>(*GetContext(), FDataflowTransformSelection());
+
+							SelectionView->GetSelectionTable()->GetSelectionInfoMap().Add(Name.ToString(), { Type.ToString(), TBitArray<>(Value.GetBitArray()) });
+						}
+						else if (Output->GetType() == "FDataflowVertexSelection")
+						{
+							const FDataflowVertexSelection& Value = Output->GetValue<FDataflowVertexSelection>(*GetContext(), FDataflowVertexSelection());
+
+							SelectionView->GetSelectionTable()->GetSelectionInfoMap().Add(Name.ToString(), { Type.ToString(), TBitArray<>(Value.GetBitArray()) });
+						}
+						else if (Output->GetType() == "FDataflowFaceSelection")
+						{
+							const FDataflowFaceSelection& Value = Output->GetValue<FDataflowFaceSelection>(*GetContext(), FDataflowFaceSelection());
+
+							SelectionView->GetSelectionTable()->GetSelectionInfoMap().Add(Name.ToString(), { Type.ToString(), TBitArray<>(Value.GetBitArray()) });
+						}
+					}
+				}
+			}
+
+			SelectionView->SetData(GetSelectedNode()->GetName());
+		}
+		else
+		{
+			SelectionView->SetData(FString());
+		}
+
+		SelectionView->RefreshWidget();
+	}
+}
+
+
+void FDataflowSelectionView::SetSelectionView(TSharedPtr<SSelectionViewWidget>& InSelectionView)
+{
+	ensure(!SelectionView);
+
+	SelectionView = InSelectionView;
+
+	if (SelectionView)
+	{
+		OnPinnedDownChangedDelegateHandle = SelectionView->GetOnPinnedDownChangedDelegate().AddRaw(this, &FDataflowSelectionView::OnPinnedDownChanged);
+		OnRefreshLockedChangedDelegateHandle = SelectionView->GetOnRefreshLockedChangedDelegate().AddRaw(this, &FDataflowSelectionView::OnRefreshLockedChanged);
+	}
+}
+
+
+FDataflowSelectionView::~FDataflowSelectionView()
+{
+	if (SelectionView)
+	{
+		SelectionView->GetOnPinnedDownChangedDelegate().Remove(OnPinnedDownChangedDelegateHandle);
+	}
+}
+
+
+#undef LOCTEXT_NAMESPACE
diff --git a/Engine/Plugins/Experimental/Dataflow/Source/DataflowEditor/Private/Dataflow/DataflowView.cpp b/Engine/Plugins/Experimental/Dataflow/Source/DataflowEditor/Private/Dataflow/DataflowView.cpp
new file mode 100644
--- /dev/null
+++ b/Engine/Plugins/Experimental/Dataflow/Source/DataflowEditor/Private/Dataflow/DataflowView.cpp
@@ -0,0 +1,114 @@
+// Copyright Epic Games, Inc. All Rights Reserved.
+
+#include "Dataflow/DataflowView.h"
+#include "Templates/EnableIf.h"
+#include "Dataflow/DataflowSelection.h"
+
+
+FDataflowNodeView::~FDataflowNodeView()
+{
+	if (SelectedNode)
+	{
+		if (TSharedPtr<FDataflowNode> DataflowNode = SelectedNode->GetDataflowNode())
+		{
+			if (DataflowNode->GetOnNodeInvalidatedDelegate().IsBound() && OnNodeInvalidatedDelegateHandle.IsValid())
+			{
+				DataflowNode->GetOnNodeInvalidatedDelegate().Remove(OnNodeInvalidatedDelegateHandle);
+			}
+		}
+	}
+}
+
+
+bool FDataflowNodeView::SelectedNodeHaveSupportedOutputTypes(UDataflowEdNode* InNode)
+{
+	SetSupportedOutputTypes();
+
+	if (InNode->IsBound())
+	{
+		if (TSharedPtr<FDataflowNode> DataflowNode = InNode->DataflowGraph->FindBaseNode(InNode->DataflowNodeGuid))
+		{
+			TArray<FDataflowOutput*> Outputs = DataflowNode->GetOutputs();
+
+			for (FDataflowOutput* Output : Outputs)
+			{
+				for (const FString& OutputType : SupportedOutputTypes)
+				{
+					if (Output->GetType() == FName(*OutputType))
+					{
+						return true;
+					}
+				}
+			}
+		}
+	}
+
+	return false;
+}
+
+
+void FDataflowNodeView::SetContext(TSharedPtr<Dataflow::FContext>& InContext)
+{
+	Context = InContext;
+}
+
+
+void FDataflowNodeView::OnSelectedNodeChanged(UDataflowEdNode* InNode)
+{
+	if (!bIsPinnedDown)
+	{
+		//
+		// Remove from broadcast
+		//
+		if (SelectedNode)
+		{
+			if (TSharedPtr<FDataflowNode> DataflowNode = SelectedNode->GetDataflowNode())
+			{
+				if (DataflowNode->GetOnNodeInvalidatedDelegate().IsBound() && OnNodeInvalidatedDelegateHandle.IsValid())
+				{
+					DataflowNode->GetOnNodeInvalidatedDelegate().Remove(OnNodeInvalidatedDelegateHandle);
+				}
+			}
+		}
+
+		SelectedNode = nullptr;
+
+		if (InNode)  // nullptr is valid
+		{
+			if (SelectedNodeHaveSupportedOutputTypes(InNode))
+			{
+				SelectedNode = InNode;
+			}
+
+			// 
+			// Bind OnNodeInvalidated() to new SelectedNode
+			// 
+			if (SelectedNode)
+			{
+				if (TSharedPtr<FDataflowNode> DataflowNode = SelectedNode->GetDataflowNode())
+				{
+					OnNodeInvalidatedDelegateHandle = DataflowNode->GetOnNodeInvalidatedDelegate().AddRaw(this, &FDataflowNodeView::OnNodeInvalidated);
+				}
+			}
+		}
+
+		UpdateViewData();
+	}
+}
+
+void FDataflowNodeView::OnNodeInvalidated(FDataflowNode* InvalidatedNode)
+{
+	if (!bIsRefreshLocked)
+	{
+		if (TSharedPtr<FDataflowNode> DataflowNode = SelectedNode->GetDataflowNode())
+		{
+			if (InvalidatedNode == DataflowNode.Get())
+			{
+				UpdateViewData();
+			}
+		}
+	}
+}
+
+
+#undef LOCTEXT_NAMESPACE
diff --git a/Engine/Plugins/Experimental/Dataflow/Source/DataflowEditor/Private/Dataflow/SelectionViewWidget.cpp b/Engine/Plugins/Experimental/Dataflow/Source/DataflowEditor/Private/Dataflow/SelectionViewWidget.cpp
new file mode 100644
--- /dev/null
+++ b/Engine/Plugins/Experimental/Dataflow/Source/DataflowEditor/Private/Dataflow/SelectionViewWidget.cpp
@@ -0,0 +1,446 @@
+// Copyright Epic Games, Inc. All Rights Reserved.
+
+#include "Dataflow/SelectionViewWidget.h"
+#include "Widgets/Input/SButton.h"
+#include "Styling/StarshipCoreStyle.h"
+
+
+#define LOCTEXT_NAMESPACE "DataflowGraphEditor"
+
+
+const FName FSelectionViewHeader::IndexColumnNameTransform = FName("Transform Index");
+const FName FSelectionViewHeader::IndexColumnNameFace = FName("Face Index");
+const FName FSelectionViewHeader::IndexColumnNameVertex = FName("Vertex Index");
+const FName FSelectionViewHeader::SelectionStatusColumnName = FName("Selection Status");
+
+const FName FSelectionViewItem::SelectedName = FName("Selected");
+const FName FSelectionViewItem::NotSelectedName = FName("-");
+
+
+void SSelectionViewRow::Construct(const FArguments& InArgs, TSharedRef<STableViewBase> OwnerTableView, const TSharedPtr<const FSelectionViewHeader>& InHeader, const TSharedPtr<const FSelectionViewItem>& InItem)
+{
+	Header = InHeader;
+	Item = InItem;
+
+	SMultiColumnTableRow<TSharedPtr<const FSelectionViewItem>>::Construct(
+		FSuperRowType::FArguments()
+		.Style(&FStarshipCoreStyle::GetCoreStyle().GetWidgetStyle<FTableRowStyle>("TableView.AlternatingRow")),
+		OwnerTableView);
+}
+
+
+TSharedRef<SWidget> SSelectionViewRow::GenerateWidgetForColumn(const FName& ColumnName)
+{
+	int32 FoundIndex;
+	if (Header->ColumnNames.Find(ColumnName, FoundIndex))
+	{
+		const FString& AttrValue = Item->Values[FoundIndex];
+
+		return SNew(SHorizontalBox)
+			+ SHorizontalBox::Slot()
+			.AutoWidth()
+			[
+				SNew(STextBlock)
+				.Text(FText::FromString(AttrValue))
+				.ShadowColorAndOpacity(FLinearColor(0.1f, 0.1f, 0.1f, 1.f))
+				.Visibility(EVisibility::Visible)
+				.ToolTipText(FText::FromString("This is an example tooltip."))
+				.Font_Lambda([this]() -> FSlateFontInfo
+				{
+					if (Item->Values[1] == FSelectionViewItem::SelectedName.ToString())
+					{
+						return FCoreStyle::GetDefaultFontStyle("Bold", 10);
+					}
+
+					return FCoreStyle::GetDefaultFontStyle("Regular", 10);
+				})
+			];
+	}
+
+	return SNullWidget::NullWidget;
+}
+
+//
+// ----------------------------------------------------------------------------
+//
+
+void SSelectionView::Construct(
+	const FArguments& InArgs
+)
+{
+	SelectedOutput = InArgs._SelectedOutput;
+
+	HeaderRowWidget =
+		SNew(SHeaderRow)
+		.Visibility(EVisibility::Visible);
+
+	if (!SelectionInfoMap.IsEmpty())
+	{
+		RegenerateHeader();
+		RepopulateListView();
+	}
+
+	ChildSlot
+	[
+		SNew(SVerticalBox)
+		+ SVerticalBox::Slot()
+		.Padding(FMargin(0.0f, 3.f))
+		[
+			SNew(SOverlay)
+			+ SOverlay::Slot()
+			[
+				SAssignNew(ListView, SListView<TSharedPtr<const FSelectionViewItem>>)
+				.SelectionMode(ESelectionMode::None)
+				.ListItemsSource(&ListItems)
+				.OnGenerateRow(this, &SSelectionView::GenerateRow)
+				.HeaderRow(HeaderRowWidget)
+			]
+		]
+	];
+}
+
+
+void SSelectionView::SetSelectedOutput(const FName& InSelectedOutput)
+{
+	SelectedOutput = InSelectedOutput;
+
+	RegenerateHeader();
+	RepopulateListView();
+}
+
+
+const FName& SSelectionView::GetSelectedOutput() const
+{
+	return SelectedOutput;
+}
+
+
+void SSelectionView::RegenerateHeader()
+{
+	HeaderRowWidget->ClearColumns();
+
+	Header = MakeShared<FSelectionViewHeader>();
+
+	if (SelectionInfoMap.Num() > 0 && !SelectedOutput.IsNone() && SelectedOutput.ToString() != "")
+	{
+		if (SelectionInfoMap[SelectedOutput.ToString()].OutputType == "FDataflowTransformSelection")
+		{
+			Header->ColumnNames.Add(FSelectionViewHeader::IndexColumnNameTransform);
+		}
+		else if (SelectionInfoMap[SelectedOutput.ToString()].OutputType == "FDataflowVertexSelection")
+		{
+			Header->ColumnNames.Add(FSelectionViewHeader::IndexColumnNameVertex);
+		}
+		else if (SelectionInfoMap[SelectedOutput.ToString()].OutputType == "FDataflowFaceSelection")
+		{
+			Header->ColumnNames.Add(FSelectionViewHeader::IndexColumnNameFace);
+		}
+	}
+	else
+	{
+		Header->ColumnNames.Add(FName(" "));
+	}
+
+	Header->ColumnNames.Add(FSelectionViewHeader::SelectionStatusColumnName);
+
+	constexpr float CustomFillWidth = 1.0f;
+
+	for (int32 AttributeNameIndex = 0; AttributeNameIndex < Header->ColumnNames.Num(); ++AttributeNameIndex)
+	{
+		const FName& ColumnName = Header->ColumnNames[AttributeNameIndex];
+
+		HeaderRowWidget->AddColumn(
+			SHeaderRow::Column(ColumnName)
+			.DefaultLabel(FText::FromName(ColumnName))
+			.FillWidth(CustomFillWidth)
+		);
+	}
+}
+
+
+void SSelectionView::RepopulateListView()
+{
+	ListItems.Empty();
+
+	if (SelectionInfoMap.Num() > 0 && !SelectedOutput.IsNone() && SelectedOutput.ToString() != "")
+	{
+		const int32 NumElems = SelectionInfoMap[SelectedOutput.ToString()].SelectionArray.Num();
+		int32 NumSelectedElems = 0;
+
+		for (int32 Idx = 0; Idx < NumElems; ++Idx)
+		{
+			const TSharedPtr<FSelectionViewItem> NewItem = MakeShared<FSelectionViewItem>();
+			NewItem->Values.SetNum(Header->ColumnNames.Num());
+
+			for (int32 ColumnIdx = 0; ColumnIdx < Header->ColumnNames.Num(); ++ColumnIdx)
+			{
+				const FName& ColumnName = Header->ColumnNames[ColumnIdx];
+				if (ColumnName == FSelectionViewHeader::SelectionStatusColumnName)
+				{
+					if (SelectionInfoMap[SelectedOutput.ToString()].SelectionArray[Idx])
+					{
+						NewItem->Values[ColumnIdx] = FSelectionViewItem::SelectedName.ToString();
+
+						NumSelectedElems++;
+					}
+					else
+					{
+						NewItem->Values[ColumnIdx] = FSelectionViewItem::NotSelectedName.ToString();
+					}
+				}
+				else
+				{
+					NewItem->Values[ColumnIdx] = FString::FromInt(Idx);
+				}
+			}
+
+			ListItems.Add(NewItem);
+		}
+
+		NumItems = NumElems;
+		NumSelectedItems = NumSelectedElems;
+	}
+
+	ListView->RequestListRefresh();
+}
+
+TSharedRef<ITableRow> SSelectionView::GenerateRow(TSharedPtr<const FSelectionViewItem> InItem, const TSharedRef<STableViewBase>& OwnerTable)
+{
+	const FTableRowStyle AlternatingTableRowStyle = FStarshipCoreStyle::GetCoreStyle().GetWidgetStyle<FTableRowStyle>("TableView.AlternatingRow");
+
+	TSharedRef<SSelectionViewRow> NewSelectionViewRow = SNew(SSelectionViewRow, OwnerTable, this->Header, InItem);
+
+	return NewSelectionViewRow;
+}
+
+
+//
+// ----------------------------------------------------------------------------
+//
+
+
+void SSelectionViewWidget::Construct(const FArguments& InArgs)
+{
+	ChildSlot
+	[
+		SNew(SVerticalBox)
+
+		+ SVerticalBox::Slot()
+		.AutoHeight()
+		[
+			// Node: [ MakeBox_0                 ] [O]
+			SNew(SHorizontalBox)
+
+			+ SHorizontalBox::Slot()
+			.AutoWidth()
+			.Padding(10.0f, 10.0f, 0.0f, 5.0f)
+			[
+				SNew(STextBlock)
+				.Text(FText::FromString("Node: "))
+			]
+
+			+ SHorizontalBox::Slot()
+			.FillWidth(1.f)
+			.Padding(10.0f, 10.0f, 10.0f, 5.0f)
+			[
+				SAssignNew(NodeNameTextBlock, STextBlock)
+			]
+
+			+ SHorizontalBox::Slot()
+			.AutoWidth()
+			.Padding(10.0f, 10.0f, 4.0f, 5.0f)
+			[
+				SNew(SButton)
+				.ToolTipText(LOCTEXT("PinDownButtonToolTip", "The button pins down the panel. When it pinned down it doesn't react to node selection change."))
+				.ButtonStyle(FAppStyle::Get(), "Button")
+				.HAlign(HAlign_Center)
+				.VAlign(VAlign_Center)
+				.Content()
+				[
+					SNew(STextBlock)
+					.Text_Lambda([this]
+					{
+						if (!NodeNameTextBlock->GetText().IsEmpty())
+						{
+							return FText::FromString(bIsPinnedDown ? " X " : " O ");
+						}
+						else
+						{
+							return FText::FromString(" O ");
+						}
+					})
+				]
+				.OnClicked_Lambda([this]
+				{
+					if (!NodeNameTextBlock->GetText().IsEmpty())
+					{
+						bIsPinnedDown = !bIsPinnedDown;
+
+						OnPinnedDownChangedDelegate.Broadcast(bIsPinnedDown);
+					}
+
+					return FReply::Handled();					
+				})
+			]
+
+			+ SHorizontalBox::Slot()
+			.AutoWidth()
+			.Padding(4.0f, 10.0f, 10.0f, 5.0f)
+			[
+				SNew(SButton)
+				.ToolTipText(LOCTEXT("LockRefreshButtonToolTip", "The button locks the refresh of the values in the panel."))
+				.ButtonStyle(FAppStyle::Get(), "Button")
+				.HAlign(HAlign_Center)
+				.VAlign(VAlign_Center)
+				.Content()
+				[
+					SNew(STextBlock)
+					.Text_Lambda([this]
+					{
+						if (!NodeNameTextBlock->GetText().IsEmpty())
+						{
+							return FText::FromString(bIsRefreshLocked ? " L " : " U ");
+						}
+						else
+						{
+							return FText::FromString(" U ");
+						}
+					})
+				]
+				.OnClicked_Lambda([this]
+				{
+					if (!NodeNameTextBlock->GetText().IsEmpty())
+					{
+						bIsRefreshLocked = !bIsRefreshLocked;
+
+						OnRefreshLockedChangedDelegate.Broadcast(bIsRefreshLocked);
+					}
+
+					return FReply::Handled();
+				})
+			]
+
+		]
+
+		+ SVerticalBox::Slot()
+		.AutoHeight()
+		[
+			// Output: [ TransformSelection        |V|]
+			SNew(SHorizontalBox)
+
+			+ SHorizontalBox::Slot()
+			.AutoWidth()
+			.Padding(10.0f, 0.0f, 0.0f, 10.0f)
+			[
+				SNew(STextBlock)
+				.Text(FText::FromString("Output: "))
+			]
+
+			+ SHorizontalBox::Slot()
+			.FillWidth(1.0f)
+			.Padding(10.0f, 0.0f, 10.0f, 10.0f)
+			[
+				SAssignNew(NodeOutputsComboBox, SComboBox<FName>)
+				.OptionsSource(&NodeOutputs)
+				.OnSelectionChanged(SComboBox<FName>::FOnSelectionChanged::CreateLambda(
+					[this](FName SelectedName, ESelectInfo::Type)
+					{
+						if (SelectionTable)
+						{
+							SelectionTable->SetSelectedOutput(SelectedName);
+						}
+
+						SetStatusText();
+					}))
+				.OnGenerateWidget(SComboBox<FName>::FOnGenerateWidget::CreateLambda(
+					[](FName Item)
+					{
+						return SNew(STextBlock)
+							.Text(FText::FromName(Item));
+					}))
+				[
+					SNew(STextBlock)
+					.Text_Lambda([this]()
+					{
+						if (SelectionTable)
+						{
+							return FText::FromName(SelectionTable->GetSelectedOutput());
+						}
+						else
+						{
+							return FText::FromString("");
+						}
+					})
+				]
+			]
+		]
+
+		+ SVerticalBox::Slot()
+		[
+			SAssignNew(SelectionTable, SSelectionView)
+		]
+
+		+ SVerticalBox::Slot()
+		.AutoHeight()
+		[
+			SNew(SHorizontalBox)
+
+			+ SHorizontalBox::Slot()
+			.AutoWidth()
+			.Padding(10.0f, 10.0f, 0.0f, 5.0f)
+			[
+				SAssignNew(StatusTextBlock, STextBlock)
+			]
+		]
+	];
+}
+
+
+void SSelectionViewWidget::SetData(const FString& InNodeName)
+{
+	NodeName = InNodeName;
+
+	NodeOutputs.Empty();
+
+	if (!NodeName.IsEmpty())
+	{
+		if (SelectionTable->GetSelectionInfoMap().Num() > 0)
+		{
+			for (auto& Info : SelectionTable->GetSelectionInfoMap())
+			{
+				NodeOutputs.Add(FName(*Info.Key));
+			}
+		}
+	}
+}
+
+
+void SSelectionViewWidget::RefreshWidget()
+{
+	NodeNameTextBlock->SetText(FText::FromString(NodeName));
+
+	NodeOutputsComboBox->RefreshOptions();
+	NodeOutputsComboBox->ClearSelection();
+
+	if (NodeOutputs.Num() > 0)
+	{
+		NodeOutputsComboBox->SetSelectedItem(NodeOutputs[0]);
+	}
+}
+
+
+void SSelectionViewWidget::SetStatusText()
+{
+	if (!NodeName.IsEmpty())
+	{
+		FString Str = FString::Printf(TEXT("Selected %d of %d"), SelectionTable->GetNumSelectedItems(), SelectionTable->GetNumItems());
+		StatusTextBlock->SetText(FText::FromString(Str));
+	}
+	else
+	{
+		StatusTextBlock->SetText(FText::FromString(" "));
+	}
+}
+
+
+#undef LOCTEXT_NAMESPACE
diff --git a/Engine/Plugins/Experimental/Dataflow/Source/DataflowEditor/Public/Dataflow/DataflowEditorCommands.h b/Engine/Plugins/Experimental/Dataflow/Source/DataflowEditor/Public/Dataflow/DataflowEditorCommands.h
--- a/Engine/Plugins/Experimental/Dataflow/Source/DataflowEditor/Public/Dataflow/DataflowEditorCommands.h
+++ b/Engine/Plugins/Experimental/Dataflow/Source/DataflowEditor/Public/Dataflow/DataflowEditorCommands.h
@@ -92,7 +92,7 @@ public:
 	/*
 	*  OnPropertyValueChanged
 	*/
-	static void OnPropertyValueChanged(UDataflow* Graph, TSharedPtr<Dataflow::FEngineContext>& Context, Dataflow::FTimestamp& OutLastNodeTimestamp, const FPropertyChangedEvent& PropertyChangedEvent);
+	static void OnPropertyValueChanged(UDataflow* Graph, TSharedPtr<Dataflow::FEngineContext>& Context, Dataflow::FTimestamp& OutLastNodeTimestamp, const FPropertyChangedEvent& PropertyChangedEvent, const TSet<UObject*>& SelectedNodes);
 
 	/*
 	*  OnSelectedNodesChanged
diff --git a/Engine/Plugins/Experimental/Dataflow/Source/DataflowEditor/Public/Dataflow/DataflowEditorToolkit.h b/Engine/Plugins/Experimental/Dataflow/Source/DataflowEditor/Public/Dataflow/DataflowEditorToolkit.h
--- a/Engine/Plugins/Experimental/Dataflow/Source/DataflowEditor/Public/Dataflow/DataflowEditorToolkit.h
+++ b/Engine/Plugins/Experimental/Dataflow/Source/DataflowEditor/Public/Dataflow/DataflowEditorToolkit.h
@@ -8,6 +8,8 @@
 #include "Misc/NotifyHook.h"
 #include "GraphEditor.h"
 #include "TickableEditorObject.h"
+#include "Dataflow/DataflowSelectionView.h"
+#include "Dataflow/SelectionViewWidget.h"
 
 class FEditorViewportTabContent;
 class IDetailsView;
@@ -16,6 +18,7 @@ class IStructureDetailsView;
 class IToolkitHost;
 class UDataflow;
 class USkeletalMesh;
+class SDataflowGraphEditor;
 
 namespace Dataflow
 {
@@ -33,6 +36,8 @@ namespace Dataflow
 class DATAFLOWEDITOR_API FDataflowEditorToolkit : public FAssetEditorToolkit, public FTickableEditorObject, public FNotifyHook, public FGCObject
 {
 public:
+	~FDataflowEditorToolkit();
+
 	static bool CanOpenDataflowEditor(UObject* ObjectToEdit);
 
 	virtual void InitializeEditor(const EToolkitMode::Type Mode, const TSharedPtr<IToolkitHost>& InitToolkitHost, UObject* ObjectToEdit);
@@ -61,6 +66,7 @@ public:
 	TSharedRef<SDockTab> SpawnTab_AssetDetails(const FSpawnTabArgs& Args);
 	TSharedRef<SDockTab> SpawnTab_NodeDetails(const FSpawnTabArgs& Args);
 	TSharedRef<SDockTab> SpawnTab_Skeletal(const FSpawnTabArgs& Args);
+	TSharedRef<SDockTab> SpawnTab_SelectionView(const FSpawnTabArgs& Args);
 
 
 	// Member Access
@@ -79,8 +85,8 @@ public:
 	TSharedPtr<IStructureDetailsView> GetNodeDetailsEditor() { return NodeDetailsEditor; }
 	const TSharedPtr<IStructureDetailsView> GetNodeDetailsEditor() const { return NodeDetailsEditor; }
 
-	TSharedPtr<SGraphEditor> GetGraphEditor() { return GraphEditor; }
-	const TSharedPtr<SGraphEditor> GetGraphEditor() const { return GraphEditor; }
+	TSharedPtr<SDataflowGraphEditor> GetGraphEditor() { return GraphEditor; }
+	const TSharedPtr<SDataflowGraphEditor> GetGraphEditor() const { return GraphEditor; }
 
 protected:
 
@@ -88,6 +94,8 @@ protected:
 	void OnPropertyValueChanged(const FPropertyChangedEvent& PropertyChangedEvent);
 	bool OnNodeVerifyTitleCommit(const FText& NewText, UEdGraphNode* GraphNode, FText& OutErrorMessage);
 	void OnNodeTitleCommitted(const FText& InNewText, ETextCommit::Type InCommitType, UEdGraphNode* GraphNode);
+	void OnNodeSelectionChanged(const TSet<UObject*>& NewSelection);
+	void OnNodeDeleted(const TSet<UObject*>& NewSelection);
 	//~ End DataflowEditorActions
 
 private:
@@ -100,9 +108,9 @@ private:
 	TSharedPtr<FEditorViewportTabContent> ViewportEditor;
 
 	static const FName GraphCanvasTabId;
-	TSharedPtr<SGraphEditor> GraphEditor;
+	TSharedPtr<SDataflowGraphEditor> GraphEditor;
 	TSharedPtr<FUICommandList> GraphEditorCommands;
-	TSharedRef<SGraphEditor> CreateGraphEditorWidget(UDataflow* ObjectToEdit, TSharedPtr<IStructureDetailsView> PropertiesEditor);
+	TSharedRef<SDataflowGraphEditor> CreateGraphEditorWidget(UDataflow* ObjectToEdit, TSharedPtr<IStructureDetailsView> PropertiesEditor);
 
 	static const FName AssetDetailsTabId;
 	TSharedPtr<IDetailsView> AssetDetailsEditor;
@@ -118,6 +126,18 @@ private:
 	TSharedPtr<class ISkeletonTree> SkeletalEditor;
 	TSharedPtr<ISkeletonTree> CreateSkeletalEditorWidget(UObject* ObjectToEdit);
 
+	static const FName SelectionViewTabId;
+
+	TSet<UObject*> PrevNodeSelection;
+	Dataflow::FTimestamp LastContextRefreshTimestamp = Dataflow::FTimestamp::Invalid;
+	int32 ContextRefreshCounter = 0;
+
+	TSharedPtr<FDataflowSelectionView> DataflowSelectionView;
+	TArray<IDataflowViewListener*> ViewListeners;
+
 	TSharedPtr<Dataflow::FEngineContext> Context;
 	Dataflow::FTimestamp LastNodeTimestamp = Dataflow::FTimestamp::Invalid;
+
+	FDelegateHandle OnSelectionChangedMulticastDelegateHandle;
+	FDelegateHandle OnNodeDeletedMulticastDelegateHandle;
 };
diff --git a/Engine/Plugins/Experimental/Dataflow/Source/DataflowEditor/Public/Dataflow/DataflowGraphEditor.h b/Engine/Plugins/Experimental/Dataflow/Source/DataflowEditor/Public/Dataflow/DataflowGraphEditor.h
--- a/Engine/Plugins/Experimental/Dataflow/Source/DataflowEditor/Public/Dataflow/DataflowGraphEditor.h
+++ b/Engine/Plugins/Experimental/Dataflow/Source/DataflowEditor/Public/Dataflow/DataflowGraphEditor.h
@@ -54,6 +54,9 @@ public:
 	DECLARE_MULTICAST_DELEGATE_OneParam(FOnSelectionChangedMulticast, const FGraphPanelSelectionSet&)
 	FOnSelectionChangedMulticast OnSelectionChangedMulticast;
 
+	DECLARE_MULTICAST_DELEGATE_OneParam(FOnNodeDeletedMulticast, const FGraphPanelSelectionSet&)
+	FOnNodeDeletedMulticast OnNodeDeletedMulticast;
+
 	// SWidget overrides
 	virtual FReply OnKeyDown(const FGeometry& MyGeometry, const FKeyEvent& InKeyEvent) override;
 	virtual FReply OnKeyUp(const FGeometry& MyGeometry, const FKeyEvent& InKeyEvent) override;
diff --git a/Engine/Plugins/Experimental/Dataflow/Source/DataflowEditor/Public/Dataflow/DataflowSelectionView.h b/Engine/Plugins/Experimental/Dataflow/Source/DataflowEditor/Public/Dataflow/DataflowSelectionView.h
new file mode 100644
--- /dev/null
+++ b/Engine/Plugins/Experimental/Dataflow/Source/DataflowEditor/Public/Dataflow/DataflowSelectionView.h
@@ -0,0 +1,34 @@
+// Copyright Epic Games, Inc. All Rights Reserved.
+
+#pragma once
+
+#include "Widgets/SCompoundWidget.h"
+#include "Widgets/Views/SListView.h"
+#include "Dataflow/DataflowEdNode.h"
+#include "Dataflow/SelectionViewWidget.h"
+#include "Dataflow/DataflowSelection.h"
+#include "Dataflow/DataflowView.h"
+
+
+/**
+*
+* Class to handle the SelectionView widget
+*
+*/
+class FDataflowSelectionView : public FDataflowNodeView
+{
+public:
+	~FDataflowSelectionView();
+
+	virtual void SetSupportedOutputTypes() override;
+	virtual void UpdateViewData() override;
+
+	void SetSelectionView(TSharedPtr<SSelectionViewWidget>& InSelectionView);
+
+private:
+	TSharedPtr<SSelectionViewWidget> SelectionView;
+
+	FDelegateHandle OnPinnedDownChangedDelegateHandle;
+	FDelegateHandle OnRefreshLockedChangedDelegateHandle;
+};
+
diff --git a/Engine/Plugins/Experimental/Dataflow/Source/DataflowEditor/Public/Dataflow/DataflowView.h b/Engine/Plugins/Experimental/Dataflow/Source/DataflowEditor/Public/Dataflow/DataflowView.h
new file mode 100644
--- /dev/null
+++ b/Engine/Plugins/Experimental/Dataflow/Source/DataflowEditor/Public/Dataflow/DataflowView.h
@@ -0,0 +1,76 @@
+// Copyright Epic Games, Inc. All Rights Reserved.
+
+#pragma once
+
+#include "Widgets/SCompoundWidget.h"
+#include "Widgets/Views/SListView.h"
+#include "Dataflow/DataflowEdNode.h"
+#include "Dataflow/SelectionViewWidget.h"
+#include "Dataflow/DataflowSelection.h"
+
+
+/**
+*
+* Base listener class to interface between the DataflowToolkit and Dataflow views
+*
+*/
+class IDataflowViewListener
+{
+public:
+	virtual void OnSelectedNodeChanged(UDataflowEdNode* InNode) = 0;  // nullptr is valid
+	virtual void OnNodeInvalidated(FDataflowNode* InvalidatedNode) = 0;
+};
+
+
+/**
+*
+* FDataflowNodeView class implements common functions for single node based Dataflow views
+*
+*/
+class FDataflowNodeView : public IDataflowViewListener
+{
+public:
+	virtual ~FDataflowNodeView();
+
+	UDataflowEdNode* GetSelectedNode() const { return SelectedNode; }
+	bool SelectedNodeHaveSupportedOutputTypes(UDataflowEdNode* InNode);
+
+	TArray<FString>& GetSupportedOutputTypes() { return SupportedOutputTypes; }
+
+	TSharedPtr<Dataflow::FContext> GetContext() { return Context; }
+	void SetContext(TSharedPtr<Dataflow::FContext>& InContext);
+
+	/**
+	* Virtual functions to overwrite in view widget classes
+	*/
+	virtual void UpdateViewData() = 0;
+	virtual void SetSupportedOutputTypes() = 0;
+
+	/**
+	* Callback for PinnedDown change
+	*/
+	void OnPinnedDownChanged(const bool State) { bIsPinnedDown = State; }
+
+	/**
+	* Callback for RefreshLock change
+	*/
+	void OnRefreshLockedChanged(const bool State) { bIsRefreshLocked = State; }
+
+	/**
+	* Virtual function overrides from IDataflowViewListener base class
+	*/
+	virtual void OnSelectedNodeChanged(UDataflowEdNode* InNode) override;  // nullptr is valid
+	virtual void OnNodeInvalidated(FDataflowNode* InvalidatedNode) override;
+
+private:
+	UDataflowEdNode* SelectedNode = nullptr;
+	TSharedPtr<Dataflow::FContext> Context;
+	bool bIsPinnedDown = false;
+	bool bIsRefreshLocked = false;
+
+	TArray<FString> SupportedOutputTypes;
+
+	FDelegateHandle OnNodeInvalidatedDelegateHandle;
+};
+
+
diff --git a/Engine/Plugins/Experimental/Dataflow/Source/DataflowEditor/Public/Dataflow/SelectionViewWidget.h b/Engine/Plugins/Experimental/Dataflow/Source/DataflowEditor/Public/Dataflow/SelectionViewWidget.h
new file mode 100644
--- /dev/null
+++ b/Engine/Plugins/Experimental/Dataflow/Source/DataflowEditor/Public/Dataflow/SelectionViewWidget.h
@@ -0,0 +1,156 @@
+// Copyright Epic Games, Inc. All Rights Reserved.
+
+#pragma once
+
+#include "CoreMinimal.h"
+#include "Dataflow/DataflowSelection.h"
+#include "DetailLayoutBuilder.h"
+#include "Widgets/Views/STableRow.h"
+#include "Widgets/Text/STextBlock.h"
+#include "Widgets/Input/SComboBox.h"
+
+/**
+* Struct to hold OutputType/Selection data for the outputs
+* Data is stored in a map using the OutputName as key: TMap<FString, SSelectionInfo>
+*/
+struct SSelectionInfo
+{
+	FString OutputType;
+	TBitArray<> SelectionArray;
+};
+
+
+/** 
+* Header
+* 1st column possible values: Transform Index/Face Index/Vertex Index based on the OutputType
+* 2nd column is SelectionStatus
+*/
+struct FSelectionViewHeader
+{
+	static const FName IndexColumnNameTransform;
+	static const FName IndexColumnNameFace;
+	static const FName IndexColumnNameVertex;
+	static const FName SelectionStatusColumnName;
+
+	TArray<FName> ColumnNames;
+};
+
+
+/** 
+* Representing a row in the table 
+* Index/SelectionStatus
+*/
+struct FSelectionViewItem
+{
+	static const FName SelectedName;
+	static const FName NotSelectedName;
+
+	TArray<FString> Values;
+};
+
+
+/** 
+* 
+*/
+class SSelectionViewRow : public SMultiColumnTableRow<TSharedPtr<const FSelectionViewItem>>
+{
+public:
+	SLATE_BEGIN_ARGS(SSelectionViewRow) 
+	{}
+		SLATE_ARGUMENT(TSharedPtr<const FSelectionViewHeader>, Header)
+		SLATE_ARGUMENT(TSharedPtr<const FSelectionViewItem>, Item)
+	SLATE_END_ARGS()
+
+	void Construct(const FArguments& InArgs, TSharedRef<STableViewBase> OwnerTableView, const TSharedPtr<const FSelectionViewHeader>& InHeader, const TSharedPtr<const FSelectionViewItem>& InItem);
+
+	virtual TSharedRef<SWidget> GenerateWidgetForColumn(const FName& ColumnName) override;
+
+private:
+	TSharedPtr<const FSelectionViewHeader> Header;
+	TSharedPtr<const FSelectionViewItem> Item;
+};
+
+/** 
+* 2xn grid to display Index/SelectionStatus data
+*/
+class SSelectionView : public SCompoundWidget
+{
+public:
+	SLATE_BEGIN_ARGS(SSelectionView) {}
+		SLATE_ARGUMENT(FName, SelectedOutput)
+	SLATE_END_ARGS()
+
+	void Construct(const FArguments& InArgs);
+
+	TMap<FString, SSelectionInfo>& GetSelectionInfoMap() { return SelectionInfoMap; }
+
+	const FName& GetSelectedOutput() const;
+	void SetSelectedOutput(const FName& InSelectedOutput);
+	
+	int32 GetNumItems() const { return NumItems; }
+	int32 GetNumSelectedItems() const { return NumSelectedItems; }
+
+	TSharedRef<ITableRow> GenerateRow(TSharedPtr<const FSelectionViewItem> InItem, const TSharedRef<STableViewBase>& OwnerTable);
+
+private:
+	FName SelectedOutput = FName(TEXT(""));
+	TMap<FString, SSelectionInfo> SelectionInfoMap;
+
+	TSharedPtr<SListView<TSharedPtr<const FSelectionViewItem>>> ListView;
+	TArray<TSharedPtr<const FSelectionViewItem>> ListItems;
+
+	TSharedPtr<FSelectionViewHeader> Header;
+	TSharedPtr<SHeaderRow> HeaderRowWidget;
+
+	int32 NumItems = 0, NumSelectedItems = 0;
+	void RegenerateHeader();
+	void RepopulateListView();
+};
+
+/**
+* Widget for the SelectionView panel
+*/
+class SSelectionViewWidget : public SCompoundWidget
+{
+public:
+	SLATE_BEGIN_ARGS(SSelectionViewWidget)
+	{}
+	SLATE_END_ARGS()
+
+	/**
+	* Slate construction function
+	*
+	* @param InArgs - Arguments passed from Slate
+	*
+	*/
+	void Construct(const FArguments& InArgs);
+
+	void SetData(const FString& InNodeName);
+	void RefreshWidget();
+	TSharedPtr<SSelectionView> GetSelectionTable() { return SelectionTable; }
+	void SetStatusText();
+
+	/** Gets a multicast delegate which is called whenever the PinnedDown button clicked */
+	DECLARE_MULTICAST_DELEGATE_OneParam(FOnPinnedDownChanged, const bool);
+	FOnPinnedDownChanged& GetOnPinnedDownChangedDelegate() { return OnPinnedDownChangedDelegate; }
+
+	/** Gets a multicast delegate which is called whenever the LockRefresh button clicked */
+	DECLARE_MULTICAST_DELEGATE_OneParam(FOnRefreshLockedChanged, const bool);
+	FOnRefreshLockedChanged& GetOnRefreshLockedChangedDelegate() { return OnRefreshLockedChangedDelegate; }
+
+private:
+	TSharedPtr<STextBlock> NodeNameTextBlock;
+	TSharedPtr<SComboBox<FName>> NodeOutputsComboBox;
+	TSharedPtr<SSelectionView> SelectionTable;
+	TSharedPtr<STextBlock> StatusTextBlock;
+
+	FString NodeName;
+	TArray<FName> NodeOutputs;
+	bool bIsPinnedDown = false;
+	bool bIsRefreshLocked = false;
+
+protected:
+	FOnPinnedDownChanged OnPinnedDownChangedDelegate;
+	FOnRefreshLockedChanged OnRefreshLockedChangedDelegate;
+
+};
\ No newline at end of file
diff --git a/Engine/Source/Runtime/Experimental/Chaos/Private/GeometryCollection/Facades/CollectionTransformSelectionFacade.cpp b/Engine/Source/Runtime/Experimental/Chaos/Private/GeometryCollection/Facades/CollectionTransformSelectionFacade.cpp
--- a/Engine/Source/Runtime/Experimental/Chaos/Private/GeometryCollection/Facades/CollectionTransformSelectionFacade.cpp
+++ b/Engine/Source/Runtime/Experimental/Chaos/Private/GeometryCollection/Facades/CollectionTransformSelectionFacade.cpp
@@ -1097,7 +1097,7 @@ namespace GeometryCollection::Facades
 	{
 		TArray<int32> OutSelection;
 
-		if (TransformIndexAttribute.IsValid() && FaceStartAttribute.IsValid() && FaceCountAttribute.IsValid())
+		if (TransformToGeometryIndexAttribute.IsValid() && FaceStartAttribute.IsValid() && FaceCountAttribute.IsValid())
 		{
 			const TManagedArray<int32>& TransformToGeometryIndices = TransformToGeometryIndexAttribute.Get();
 			const TManagedArray<int32>& FaceStarts = FaceStartAttribute.Get();
diff --git a/Engine/Source/Runtime/Experimental/Dataflow/Core/Private/Dataflow/DataflowNode.cpp b/Engine/Source/Runtime/Experimental/Dataflow/Core/Private/Dataflow/DataflowNode.cpp
--- a/Engine/Source/Runtime/Experimental/Dataflow/Core/Private/Dataflow/DataflowNode.cpp
+++ b/Engine/Source/Runtime/Experimental/Dataflow/Core/Private/Dataflow/DataflowNode.cpp
@@ -227,6 +227,8 @@ void FDataflowNode::Invalidate()
 		FDataflowOutput* Con = Elem.Value;
 		Con->Invalidate();
 	}
+
+	OnNodeInvalidatedDelegate.Broadcast(this);
 }
 
 void FDataflowNode::RegisterInputConnection(const void* InProperty)
diff --git a/Engine/Source/Runtime/Experimental/Dataflow/Core/Public/Dataflow/DataflowNode.h b/Engine/Source/Runtime/Experimental/Dataflow/Core/Public/Dataflow/DataflowNode.h
--- a/Engine/Source/Runtime/Experimental/Dataflow/Core/Public/Dataflow/DataflowNode.h
+++ b/Engine/Source/Runtime/Experimental/Dataflow/Core/Public/Dataflow/DataflowNode.h
@@ -231,6 +231,8 @@ struct DATAFLOWCORE_API FDataflowNode
 		return nullptr;
 	}
 
+	DECLARE_MULTICAST_DELEGATE_OneParam(FOnNodeInvalidated, FDataflowNode*);
+	FOnNodeInvalidated& GetOnNodeInvalidatedDelegate() { return OnNodeInvalidatedDelegate; }
 
 private:
 	virtual TArray<Dataflow::FRenderingParameter> GetRenderParametersImpl() const { return TArray<Dataflow::FRenderingParameter>(); }
@@ -238,6 +240,9 @@ private:
 
 	bool bValid = true;
 
+protected:
+	FOnNodeInvalidated OnNodeInvalidatedDelegate;
+
 };
 
 namespace Dataflow
diff --git a/Engine/Source/Runtime/Experimental/Dataflow/Core/Public/Dataflow/DataflowSelection.h b/Engine/Source/Runtime/Experimental/Dataflow/Core/Public/Dataflow/DataflowSelection.h
--- a/Engine/Source/Runtime/Experimental/Dataflow/Core/Public/Dataflow/DataflowSelection.h
+++ b/Engine/Source/Runtime/Experimental/Dataflow/Core/Public/Dataflow/DataflowSelection.h
@@ -26,7 +26,7 @@ struct DATAFLOWCORE_API FDataflowSelection
 	void SetFromArray(const TArray<int32>& SelectionArr);
 	
 	/**
-	* Sets the selection from a dense array (it contains element for every item)
+	* Sets the selection from a dense array (it contains a true/false element for every item)
 	* (for example from the "Internal" attr from FacesGroup)
 	*/
 	void SetFromArray(const TArray<bool>& SelectionArr);
@@ -35,6 +35,7 @@ struct DATAFLOWCORE_API FDataflowSelection
 	void XOR(const FDataflowSelection& Other, FDataflowSelection& Result) const;
 	void Invert() { SelectionArray.BitwiseNOT(); }
 	void SetWithMask(const bool Value, const FDataflowSelection& Mask);
+	const TBitArray<>& GetBitArray() const { return SelectionArray; };
 
 private:
 
diff --git a/Engine/Source/Runtime/Experimental/GeometryCollectionEngine/Private/GeometryCollection/GeometryCollectionEngineConversion.cpp b/Engine/Source/Runtime/Experimental/GeometryCollectionEngine/Private/GeometryCollection/GeometryCollectionEngineConversion.cpp
--- a/Engine/Source/Runtime/Experimental/GeometryCollectionEngine/Private/GeometryCollection/GeometryCollectionEngineConversion.cpp
+++ b/Engine/Source/Runtime/Experimental/GeometryCollectionEngine/Private/GeometryCollection/GeometryCollectionEngineConversion.cpp
@@ -1327,7 +1327,7 @@ int32 FGeometryCollectionEngineConversion::AppendSkeletalMeshMaterials(const USk
 
 void FGeometryCollectionEngineConversion::AppendGeometryCollectionSource(const FGeometryCollectionSource& GeometryCollectionSource, FGeometryCollection& GeometryCollectionInOut, TArray<UMaterial*>& MaterialsInOut, bool ReindexMaterials)
 {
-	if (const UObject* SourceObject = GeometryCollectionSource.SourceGeometryObject.ResolveObject())
+	if (const UObject* SourceObject = GeometryCollectionSource.SourceGeometryObject.TryLoad())
 	{
 		const int32 StartMaterialIndex = MaterialsInOut.Num();
 
