Dataflow:
- Fixed missing '#define LOCTEXT_NAMESPACE'

diff --git a/Engine/Plugins/Experimental/Dataflow/Source/DataflowEditor/Private/Dataflow/DataflowSelectionView.cpp b/Engine/Plugins/Experimental/Dataflow/Source/DataflowEditor/Private/Dataflow/DataflowSelectionView.cpp
--- a/Engine/Plugins/Experimental/Dataflow/Source/DataflowEditor/Private/Dataflow/DataflowSelectionView.cpp
+++ b/Engine/Plugins/Experimental/Dataflow/Source/DataflowEditor/Private/Dataflow/DataflowSelectionView.cpp
@@ -4,6 +4,7 @@
 #include "Templates/EnableIf.h"
 #include "Dataflow/DataflowSelection.h"
 
+#define LOCTEXT_NAMESPACE "DataflowSelectionView"
 
 void FDataflowSelectionView::SetSupportedOutputTypes()
 {
diff --git a/Engine/Plugins/Experimental/Dataflow/Source/DataflowEditor/Private/Dataflow/DataflowView.cpp b/Engine/Plugins/Experimental/Dataflow/Source/DataflowEditor/Private/Dataflow/DataflowView.cpp
--- a/Engine/Plugins/Experimental/Dataflow/Source/DataflowEditor/Private/Dataflow/DataflowView.cpp
+++ b/Engine/Plugins/Experimental/Dataflow/Source/DataflowEditor/Private/Dataflow/DataflowView.cpp
@@ -4,6 +4,7 @@
 #include "Templates/EnableIf.h"
 #include "Dataflow/DataflowSelection.h"
 
+#define LOCTEXT_NAMESPACE "DataflowView"
 
 FDataflowNodeView::~FDataflowNodeView()
 {
