Dataflow:
- Added CollectionSpreadSheet window
- Fixed issues with SelectionView window
- both windows can have up to 4 instances
- bug fixes

diff --git a/Engine/Plugins/Experimental/Dataflow/Source/DataflowEditor/Private/Dataflow/CollectionSpreadSheetWidget.cpp b/Engine/Plugins/Experimental/Dataflow/Source/DataflowEditor/Private/Dataflow/CollectionSpreadSheetWidget.cpp
new file mode 100644
--- /dev/null
+++ b/Engine/Plugins/Experimental/Dataflow/Source/DataflowEditor/Private/Dataflow/CollectionSpreadSheetWidget.cpp
@@ -0,0 +1,808 @@
+// Copyright Epic Games, Inc. All Rights Reserved.
+
+#include "Dataflow/CollectionSpreadSheetWidget.h"
+#include "Widgets/Input/SButton.h"
+#include "Styling/StarshipCoreStyle.h"
+#include "Widgets/Layout/SScrollBox.h"
+
+
+#define LOCTEXT_NAMESPACE "CollectionSpreadSheetWidget"
+
+
+const FName FCollectionSpreadSheetHeader::IndexColumnName = FName("Index");
+
+namespace CollectionSpreadSheetHelpers
+{
+	template<typename T>
+	FString AttributeValueToString(const T& Value)
+	{
+		return Value.ToString();
+	}
+
+	FString AttributeValueToString(float Value)
+	{
+		return FString::SanitizeFloat(Value, 2);
+	}
+
+	FString AttributeValueToString(int32 Value)
+	{
+		return FString::FromInt(Value);
+	}
+
+	FString AttributeValueToString(FString Value)
+	{
+		return Value;
+	}
+
+	FString AttributeValueToString(FLinearColor Value)
+	{
+		return FString::Printf(TEXT("(R=%.2f G=%.2f B=%.2f A=%.2f)"), Value.R, Value.G, Value.B, Value.A);
+	}
+
+	FString AttributeValueToString(FVector Value)
+	{
+		return FString::Printf(TEXT("(X=%.2f Y=%.2f Z=%.2f)"), Value.X, Value.Y, Value.Z);
+	}
+
+	FString AttributeValueToString(bool Value)
+	{
+		return Value ? FString("true") : FString("false");
+	}
+
+	FString AttributeValueToString(TSet<int32> Value)
+	{
+		TArray<int32> Array = Value.Array();
+		FString Out;
+
+		for (int32 Idx = 0; Idx < Array.Num(); ++Idx)
+		{
+			Out += AttributeValueToString(Array[Idx]);
+
+			if (Idx != Array.Num() - 1)
+			{
+				Out += " ";
+			}
+		}
+
+		return Out;
+	}
+
+	template<typename T>
+	FString AttributeValueToString(const TArray<T>& Array)
+	{
+		FString Out;
+
+		for (int32 Idx = 0; Idx < Array.Num(); ++Idx)
+		{
+			Out += AttributeValueToString(Array[Idx]);
+
+			if (Idx != Array.Num() - 1)
+			{
+				Out += "; ";
+			}
+		}
+
+		return Out;
+	}
+
+	FString AttributeValueToString(FTransform Value)
+	{
+		const FVector Translation = Value.GetTranslation();
+		const FVector Rotation = Value.GetRotation().Euler();
+		const FVector Scale = Value.GetScale3D();
+
+		return FString::Printf(TEXT("T:(%s) R:(%s) S:(%s)"), *Translation.ToString(), *Rotation.ToString(), *Scale.ToString());
+	}
+
+	FString AttributeValueToString(FBox Value)
+	{
+		return FString::Printf(TEXT("Min:(%s) Max:(%s)"), *Value.Min.ToString(), *Value.Max.ToString());
+	}
+
+	FString AttributeValueToString(FIntVector Value)
+	{
+		return FString::Printf(TEXT("%d %d %d"), Value.X, Value.Y, Value.Z);
+	}
+
+	template<typename T>
+	FString AttributeValueToString(const FManagedArrayCollection& InCollection, const FName& InAttributeName, const FName& InGroupName, int32 InIdxColumn)
+	{
+		const TManagedArray<T>* const Array = InCollection.FindAttributeTyped<T>(InAttributeName, InGroupName);
+		if (Array == nullptr)
+		{
+			return FString("<Unknown Attribute>");
+		}
+
+		return AttributeValueToString((*Array)[InIdxColumn]);
+	}
+
+	FString AttributeValueToString(const FManagedArrayCollection& InCollection, const FName& InAttributeName, const FName& InGroupName, int32 InIdxColumn)
+	{
+		const FManagedArrayCollection::EArrayType ArrayType = InCollection.GetAttributeType(InAttributeName, InGroupName);
+
+		FString ValueAsString;
+
+		switch (ArrayType)
+		{
+		case FManagedArrayCollection::EArrayType::FFloatType:
+			ValueAsString = AttributeValueToString<float>(InCollection, InAttributeName, InGroupName, InIdxColumn);
+			break;
+
+		case FManagedArrayCollection::EArrayType::FInt32Type:
+			ValueAsString = AttributeValueToString<int32>(InCollection, InAttributeName, InGroupName, InIdxColumn);
+			break;
+
+		case FManagedArrayCollection::EArrayType::FBoolType:
+			ValueAsString = AttributeValueToString<bool>(InCollection, InAttributeName, InGroupName, InIdxColumn);
+			break;
+
+		case FManagedArrayCollection::EArrayType::FStringType:
+			ValueAsString = AttributeValueToString<FString>(InCollection, InAttributeName, InGroupName, InIdxColumn);
+			break;
+
+		case FManagedArrayCollection::EArrayType::FLinearColorType:
+			ValueAsString = AttributeValueToString<FLinearColor>(InCollection, InAttributeName, InGroupName, InIdxColumn);
+			break;
+
+		case FManagedArrayCollection::EArrayType::FVectorType:
+			ValueAsString = AttributeValueToString<FVector3f>(InCollection, InAttributeName, InGroupName, InIdxColumn);
+			break;
+
+		case FManagedArrayCollection::EArrayType::FVector2DType:
+			ValueAsString = AttributeValueToString<FVector2f>(InCollection, InAttributeName, InGroupName, InIdxColumn);
+			break;
+
+		case FManagedArrayCollection::EArrayType::FVector3dType:
+			ValueAsString = AttributeValueToString<FVector2f>(InCollection, InAttributeName, InGroupName, InIdxColumn);
+			break;
+
+		case FManagedArrayCollection::EArrayType::FIntVectorType:
+			ValueAsString = AttributeValueToString<FIntVector>(InCollection, InAttributeName, InGroupName, InIdxColumn);
+			break;
+
+		case FManagedArrayCollection::EArrayType::FTransformType:
+			ValueAsString = AttributeValueToString<FTransform>(InCollection, InAttributeName, InGroupName, InIdxColumn);
+			break;
+
+		case FManagedArrayCollection::EArrayType::FVector2DArrayType:
+			ValueAsString = AttributeValueToString<TArray<FVector2f>>(InCollection, InAttributeName, InGroupName, InIdxColumn);
+			break;
+
+		case FManagedArrayCollection::EArrayType::FIntArrayType:
+			ValueAsString = AttributeValueToString<TSet<int32>>(InCollection, InAttributeName, InGroupName, InIdxColumn);
+			break;
+
+		case FManagedArrayCollection::EArrayType::FInt32ArrayType:
+			ValueAsString = AttributeValueToString<TArray<int32>>(InCollection, InAttributeName, InGroupName, InIdxColumn);
+			break;
+
+		case FManagedArrayCollection::EArrayType::FFloatArrayType:
+			ValueAsString = AttributeValueToString<TArray<float>>(InCollection, InAttributeName, InGroupName, InIdxColumn);
+			break;
+
+		case FManagedArrayCollection::EArrayType::FBoxType:
+			ValueAsString = AttributeValueToString<FBox>(InCollection, InAttributeName, InGroupName, InIdxColumn);
+			break;
+
+		default:
+			ensure(false);
+			ValueAsString = "<Unknown Data Type>";
+		}
+
+		return ValueAsString;
+	}
+}
+
+
+void SCollectionSpreadSheetRow::Construct(const FArguments& InArgs, TSharedRef<STableViewBase> OwnerTableView, const TSharedPtr<const FCollectionSpreadSheetHeader>& InHeader, const TSharedPtr<const FCollectionSpreadSheetItem>& InItem)
+{
+	Header = InHeader;
+	Item = InItem;
+
+	SMultiColumnTableRow<TSharedPtr<const FCollectionSpreadSheetItem>>::Construct(
+		FSuperRowType::FArguments()
+		.Style(&FStarshipCoreStyle::GetCoreStyle().GetWidgetStyle<FTableRowStyle>("TableView.AlternatingRow")),
+		OwnerTableView);
+}
+
+
+TSharedRef<SWidget> SCollectionSpreadSheetRow::GenerateWidgetForColumn(const FName& ColumnName)
+{
+	int32 FoundIndex;
+	if (Header->ColumnNames.Find(ColumnName, FoundIndex))
+	{
+		const FString& AttrValue = Item->Values[FoundIndex];
+
+		return SNew(SHorizontalBox)
+			+ SHorizontalBox::Slot()
+			.AutoWidth()
+			[
+				SNew(STextBlock)
+				.Text(FText::FromString(AttrValue))
+				.ShadowColorAndOpacity(FLinearColor(0.1f, 0.1f, 0.1f, 1.f))
+				.Visibility(EVisibility::Visible)
+			];
+	}
+
+	return SNullWidget::NullWidget;
+}
+
+//
+// ----------------------------------------------------------------------------
+//
+
+void SCollectionSpreadSheet::Construct(const FArguments& InArgs)
+{
+	SelectedOutput = InArgs._SelectedOutput;
+
+	SpreadSheetVerticalScrollBar = SNew(SScrollBar)
+		.Orientation(Orient_Vertical)
+		.Thickness(FVector2D(12.0f, 12.0f));
+
+	HeaderRowWidget = SNew(SHeaderRow)
+		.Visibility(EVisibility::Visible);
+
+	if (!CollectionInfoMap.IsEmpty())
+	{
+		RegenerateHeader();
+		RepopulateListView();
+	}
+
+	ChildSlot
+	[
+		SNew(SVerticalBox)
+		+ SVerticalBox::Slot()
+		.Padding(FMargin(0.0f, 3.f))
+		[
+			SNew(SOverlay)
+			+ SOverlay::Slot()
+			[
+				SAssignNew(ListView, SListView<TSharedPtr<const FCollectionSpreadSheetItem>>)
+//				SAssignNew(ListView, STreeView<TSharedPtr<const FCollectionSpreadSheetItem>>)
+				.SelectionMode(ESelectionMode::Multi)
+				.ListItemsSource(&ListItems)
+//				.TreeItemsSource(&ListItems)
+				.OnGenerateRow(this, &SCollectionSpreadSheet::GenerateRow)
+				.HeaderRow(HeaderRowWidget)
+				.ExternalScrollbar(SpreadSheetVerticalScrollBar)
+			]
+		]
+	];
+
+	AttrTypeWidthMap.Add("Transform", 600);
+	AttrTypeWidthMap.Add("String", 200);
+	AttrTypeWidthMap.Add("LinearColor", 250);
+	AttrTypeWidthMap.Add("int32", 100);
+	AttrTypeWidthMap.Add("IntArray", 200);
+	AttrTypeWidthMap.Add("Vector", 300);
+	AttrTypeWidthMap.Add("Vector2D", 160);
+	AttrTypeWidthMap.Add("Float", 150);
+	AttrTypeWidthMap.Add("IntVector", 150);
+	AttrTypeWidthMap.Add("Bool", 75);
+	AttrTypeWidthMap.Add("Box", 550);
+	AttrTypeWidthMap.Add("MeshSection", 0);
+	AttrTypeWidthMap.Add("UInt8", 0);
+}
+
+
+void SCollectionSpreadSheet::SetSelectedOutput(const FName& InSelectedOutput)
+{
+	SelectedOutput = InSelectedOutput;
+
+	RegenerateHeader();
+	RepopulateListView();
+}
+
+
+const FName& SCollectionSpreadSheet::GetSelectedOutput() const
+{
+	return SelectedOutput;
+}
+
+
+void SCollectionSpreadSheet::SetSelectedGroup(const FName& InSelectedGroup)
+{
+	SelectedGroup = InSelectedGroup;
+
+	RegenerateHeader();
+	RepopulateListView();
+}
+
+
+const FName& SCollectionSpreadSheet::GetSelectedGroup() const
+{
+	return SelectedGroup;
+}
+
+namespace {
+	inline FName GetArrayTypeString(FManagedArrayCollection::EArrayType ArrayType)
+	{
+		switch (ArrayType)
+		{
+#define MANAGED_ARRAY_TYPE(a,A)	case EManagedArrayType::F##A##Type:\
+			return FName(#A);
+#include "GeometryCollection/ManagedArrayTypeValues.inl"
+#undef MANAGED_ARRAY_TYPE
+		}
+		return FName();
+	}
+}
+
+void SCollectionSpreadSheet::RegenerateHeader()
+{
+	HeaderRowWidget->ClearColumns();
+
+	Header = MakeShared<FCollectionSpreadSheetHeader>();
+	TArray<FString> AttrTypes;
+
+	if (CollectionInfoMap.Num() > 0 && !SelectedOutput.IsNone() && SelectedOutput.ToString() != "" &&
+		!SelectedGroup.IsNone() && SelectedGroup.ToString() != "")
+	{
+		Header->ColumnNames.Add(FCollectionSpreadSheetHeader::IndexColumnName);
+
+		for (FName Attr : CollectionInfoMap[SelectedOutput.ToString()].Collection.AttributeNames(SelectedGroup))
+		{
+			Header->ColumnNames.Add(Attr);
+			AttrTypes.Add(GetArrayTypeString(CollectionInfoMap[SelectedOutput.ToString()].Collection.GetAttributeType(Attr, SelectedGroup)).ToString());
+		}
+
+	}
+
+	for (int32 IdxAttr = 0; IdxAttr < Header->ColumnNames.Num(); ++IdxAttr)
+	{
+		const FName& ColumnName = Header->ColumnNames[IdxAttr];
+		FName ToolTip;
+
+		FString ColumnNameStr = ColumnName.ToString();
+		FString AttrTypeStr;
+
+		if (IdxAttr > 0)
+		{
+			ToolTip = FName(*FString::Printf(TEXT("Attr: %s\nType: %s"), *ColumnName.ToString(), *AttrTypes[IdxAttr - 1])); // IdxAttr needs to be adjusted because of the first Index column
+
+			AttrTypeStr = AttrTypes[IdxAttr - 1];
+		}
+		else
+		{
+			ToolTip = FName("");
+		}
+
+		int32 ColumnWidth = 100;
+		if (ColumnNameStr == "Index")
+		{
+			ColumnWidth = 100;
+		}
+		else
+		{
+			int32 ColumnNameStrLen = ColumnNameStr.Len() * 9;
+			int32 AttrTypeWidth = 100;
+			if (AttrTypeWidthMap.Contains(AttrTypeStr))
+			{
+				AttrTypeWidth = AttrTypeWidthMap[AttrTypeStr];
+			}
+			ColumnWidth = ColumnNameStrLen > AttrTypeWidth ? ColumnNameStrLen : AttrTypeWidth;
+		}
+
+		HeaderRowWidget->AddColumn(
+			SHeaderRow::Column(ColumnName)
+			.DefaultLabel(FText::FromName(ColumnName))
+			.DefaultTooltip(FText::FromName(ToolTip))
+			.ManualWidth(ColumnWidth)
+			//			.SortMode(EColumnSizeMode)
+			.HAlignCell(HAlign_Center)
+			.HAlignHeader(HAlign_Center)
+			.VAlignCell(VAlign_Center)
+		);
+	}
+}
+
+
+void SCollectionSpreadSheet::RepopulateListView()
+{
+	ListItems.Empty();
+
+	if (CollectionInfoMap.Num() > 0 &&
+		!SelectedOutput.IsNone() && SelectedOutput.ToString() != "" &&
+		!SelectedGroup.IsNone() && SelectedGroup.ToString() != "")
+	{
+		int32 NumElems = CollectionInfoMap[SelectedOutput.ToString()].Collection.NumElements(SelectedGroup);
+
+		for (int32 IdxElem = 0; IdxElem < NumElems; ++IdxElem)
+		{
+			const TSharedPtr<FCollectionSpreadSheetItem> NewItem = MakeShared<FCollectionSpreadSheetItem>();
+			NewItem->Values.SetNum(Header->ColumnNames.Num());
+
+			for (int32 IdxColumn = 0; IdxColumn < Header->ColumnNames.Num(); ++IdxColumn)
+			{
+				const FName& ColumnName = Header->ColumnNames[IdxColumn];
+				if (ColumnName == FCollectionSpreadSheetHeader::IndexColumnName)
+				{
+					NewItem->Values[IdxColumn] = FString::FromInt(IdxElem);
+				}
+				else
+				{
+					NewItem->Values[IdxColumn] = CollectionSpreadSheetHelpers::AttributeValueToString(CollectionInfoMap[SelectedOutput.ToString()].Collection, ColumnName, SelectedGroup, IdxElem);
+				}
+			}
+
+			ListItems.Add(NewItem);
+		}
+	}
+
+	ListView->RequestListRefresh();
+}
+
+TSharedRef<ITableRow> SCollectionSpreadSheet::GenerateRow(TSharedPtr<const FCollectionSpreadSheetItem> InItem, const TSharedRef<STableViewBase>& OwnerTable)
+{
+	TSharedRef<SCollectionSpreadSheetRow> NewCollectionSpreadSheetRow = SNew(SCollectionSpreadSheetRow, OwnerTable, this->Header, InItem);
+
+	return NewCollectionSpreadSheetRow;
+}
+
+
+//
+// ----------------------------------------------------------------------------
+//
+
+void SCollectionSpreadSheetWidget::NodeOutputsComboBoxSelectionChanged(FName InSelectedOutput, ESelectInfo::Type InSelectInfo)
+{
+	if (CollectionTable)
+	{
+		if (CollectionTable->GetSelectedOutput() != InSelectedOutput)
+		{
+			CollectionTable->SetSelectedOutput(InSelectedOutput);
+
+			NodeOutputsComboBoxLabel->SetText(FText::FromName(CollectionTable->GetSelectedOutput()));
+
+			CollectionGroupsComboBox->RefreshOptions();
+			CollectionGroupsComboBox->ClearSelection();
+
+			UpdateCollectionGroups(InSelectedOutput);
+
+			if (CollectionGroups.Num() > 0)
+			{
+				CollectionGroupsComboBox->SetSelectedItem(CollectionGroups[0]);
+			}
+		}
+	}
+}
+
+
+void SCollectionSpreadSheetWidget::CollectionGroupsComboBoxSelectionChanged(FName InSelectedGroup, ESelectInfo::Type InSelectInfo)
+{
+	if (CollectionTable)
+	{
+		if (CollectionTable->GetSelectedGroup() != InSelectedGroup)
+		{
+			CollectionTable->SetSelectedGroup(InSelectedGroup);
+
+			CollectionGroupsComboBoxLabel->SetText(FText::FromName(CollectionTable->GetSelectedGroup()));
+
+			if (!InSelectedGroup.IsNone())
+			{
+				int32 NumElems = CollectionTable->GetCollectionInfoMap()[CollectionTable->GetSelectedOutput().ToString()].Collection.NumElements(InSelectedGroup);
+				CollectionTable->SetNumItems(NumElems);
+			}
+		}
+	}
+
+	SetStatusText();
+}
+
+
+FText SCollectionSpreadSheetWidget::GetNoOutputText()
+{
+	return LOCTEXT("NoOutput", "No Output(s)");
+}
+
+
+FText SCollectionSpreadSheetWidget::GetNoGroupText()
+{
+	return LOCTEXT("NoOutput", "No Group(s)");
+}
+
+
+void SCollectionSpreadSheetWidget::Construct(const FArguments& InArgs)
+{
+	TSharedPtr<SScrollBar> SpreadSheetHorizontalScrollBar = SNew(SScrollBar)
+		.Orientation(Orient_Horizontal)
+		.Thickness(FVector2D(12.0f, 12.0f));
+
+	ChildSlot
+	[
+		SNew(SVerticalBox)
+
+		+ SVerticalBox::Slot()
+		.AutoHeight()
+		[
+			// Node: [ MakeBox_0                 ] [O]
+			SNew(SHorizontalBox)
+
+			+ SHorizontalBox::Slot()
+			.AutoWidth()
+			.Padding(10.0f, 10.0f, 0.0f, 5.0f)
+			[
+				SNew(STextBlock)
+				.Text(FText::FromString("Node: "))
+			]
+
+			+ SHorizontalBox::Slot()
+			.FillWidth(1.f)
+			.Padding(10.0f, 10.0f, 10.0f, 5.0f)
+			[
+				SAssignNew(NodeNameTextBlock, STextBlock)
+			]
+
+			+ SHorizontalBox::Slot()
+			.AutoWidth()
+			.Padding(10.0f, 10.0f, 4.0f, 5.0f)
+			[
+				SNew(SButton)
+				.ToolTipText(LOCTEXT("PinDownButtonToolTip", "The button pins down the panel. When it pinned down it doesn't react to node selection change."))
+				.ButtonStyle(FAppStyle::Get(), "Button")
+				.HAlign(HAlign_Center)
+				.VAlign(VAlign_Center)
+				.Content()
+				[
+					SNew(STextBlock)
+					.Text_Lambda([this]
+					{
+						if (!NodeNameTextBlock->GetText().IsEmpty())
+						{
+							return FText::FromString(bIsPinnedDown ? " X " : " O ");
+						}
+						else
+						{
+							return FText::FromString(" O ");
+						}
+					})
+				]
+				.OnClicked_Lambda([this]
+				{
+					if (!NodeNameTextBlock->GetText().IsEmpty())
+					{
+						bIsPinnedDown = !bIsPinnedDown;
+		
+						OnPinnedDownChangedDelegate.Broadcast(bIsPinnedDown);
+					}
+
+					return FReply::Handled();
+				})
+			]
+
+			+ SHorizontalBox::Slot()
+			.AutoWidth()
+			.Padding(4.0f, 10.0f, 10.0f, 5.0f)
+			[
+				SNew(SButton)
+				.ToolTipText(LOCTEXT("LockRefreshButtonToolTip", "The button locks the refresh of the values in the panel."))
+				.ButtonStyle(FAppStyle::Get(), "Button")
+				.HAlign(HAlign_Center)
+				.VAlign(VAlign_Center)
+				.Content()
+				[
+					SNew(STextBlock)
+					.Text_Lambda([this]
+					{
+						if (!NodeNameTextBlock->GetText().IsEmpty())
+						{
+							return FText::FromString(bIsRefreshLocked ? " L " : " U ");
+						}
+						else
+						{
+							return FText::FromString(" U ");
+						}
+					})
+				]
+				.OnClicked_Lambda([this]
+				{
+					if (!NodeNameTextBlock->GetText().IsEmpty())
+					{
+						bIsRefreshLocked = !bIsRefreshLocked;
+		
+						OnRefreshLockedChangedDelegate.Broadcast(bIsRefreshLocked);
+					}
+		
+					return FReply::Handled();
+				})
+			]
+
+		]
+
+		+ SVerticalBox::Slot()
+		.AutoHeight()
+		[
+			// Output: [ TransformSelection        |V|]
+			SNew(SHorizontalBox)
+
+			+ SHorizontalBox::Slot()
+			.AutoWidth()
+			.Padding(10.0f, 0.0f, 0.0f, 10.0f)
+			[
+				SNew(STextBlock)
+				.Text(FText::FromString("Output: "))
+			]
+
+			+ SHorizontalBox::Slot()
+			.FillWidth(1.0f)
+			.Padding(10.0f, 0.0f, 10.0f, 10.0f)
+			[
+				SAssignNew(NodeOutputsComboBox, SComboBox<FName>)
+				.ToolTipText(LOCTEXT("NodeOutputsToolTip", "Select a node output to see the output's data"))
+				.OptionsSource(&NodeOutputs)
+				.OnGenerateWidget(SComboBox<FName>::FOnGenerateWidget::CreateLambda([](FName Output)->TSharedRef<SWidget>
+				{
+					return SNew(STextBlock)
+						.Text(FText::FromName(Output));
+				}))
+				.OnSelectionChanged(this, &SCollectionSpreadSheetWidget::NodeOutputsComboBoxSelectionChanged)
+				[
+					SAssignNew(NodeOutputsComboBoxLabel, STextBlock)
+					.Text(GetNoOutputText())
+				]
+			]
+		]
+
+		+ SVerticalBox::Slot()
+		.AutoHeight()
+		[
+			// Group: [ TransformGroup        |V|]
+			SNew(SHorizontalBox)
+
+			+ SHorizontalBox::Slot()
+			.AutoWidth()
+			.Padding(10.0f, 0.0f, 0.0f, 2.0f)
+			[
+				SNew(STextBlock)
+				.Text(FText::FromString("Group: "))
+			]
+
+			+ SHorizontalBox::Slot()
+			.FillWidth(1.0f)
+			.Padding(10.0f, 0.0f, 10.0f, 2.0f)
+			[
+				SAssignNew(CollectionGroupsComboBox, SComboBox<FName>)
+				.ToolTipText(LOCTEXT("CollectionGroupsToolTip", "Select a group see the corresponding data"))
+				.OptionsSource(&CollectionGroups)
+				.OnGenerateWidget(SComboBox<FName>::FOnGenerateWidget::CreateLambda([](FName Group)->TSharedRef<SWidget>
+				{
+					return SNew(STextBlock)
+						.Text(FText::FromName(Group));
+				}))
+				.OnSelectionChanged(this, &SCollectionSpreadSheetWidget::CollectionGroupsComboBoxSelectionChanged)
+				[
+					SAssignNew(CollectionGroupsComboBoxLabel, STextBlock)
+						.Text(GetNoGroupText())
+				]
+			]
+		]
+
+		+SVerticalBox::Slot()
+		[
+			SNew(SHorizontalBox)
+
+			+ SHorizontalBox::Slot()
+			[
+				SNew(SScrollBox)
+				.Orientation(Orient_Horizontal)
+				.ExternalScrollbar(SpreadSheetHorizontalScrollBar)
+
+				+ SScrollBox::Slot()
+				[
+					SAssignNew(CollectionTable, SCollectionSpreadSheet)
+				]
+			]
+
+			+ SHorizontalBox::Slot()
+			.AutoWidth()
+			[
+				CollectionTable->SpreadSheetVerticalScrollBar.ToSharedRef()
+			]
+		]
+
+		+ SVerticalBox::Slot()
+		.AutoHeight()
+		[
+			SpreadSheetHorizontalScrollBar.ToSharedRef()
+		]
+
+		+ SVerticalBox::Slot()
+		.AutoHeight()
+		[
+			SNew(SHorizontalBox)
+
+			+ SHorizontalBox::Slot()
+			.AutoWidth()
+			.Padding(10.0f, 10.0f, 0.0f, 5.0f)
+			[
+				SAssignNew(StatusTextBlock, STextBlock)
+			]
+		]
+	];
+}
+
+
+void SCollectionSpreadSheetWidget::SetData(const FString& InNodeName)
+{
+	NodeName = InNodeName;
+
+	NodeOutputs.Empty();
+
+	if (!NodeName.IsEmpty())
+	{
+		if (CollectionTable->GetCollectionInfoMap().Num() > 0)
+		{
+			for (auto& Info : CollectionTable->GetCollectionInfoMap())
+			{
+				NodeOutputs.Add(FName(*Info.Key));
+			}
+		}
+	}
+
+	CollectionGroups.Empty();
+}
+
+
+void SCollectionSpreadSheetWidget::RefreshWidget()
+{
+	NodeNameTextBlock->SetText(FText::FromString(NodeName));
+
+	NodeOutputsComboBox->RefreshOptions();
+	NodeOutputsComboBox->ClearSelection();
+
+	if (NodeOutputs.Num() > 0)
+	{
+		NodeOutputsComboBox->SetSelectedItem(NodeOutputs[0]);
+	}
+	else
+	{
+		NodeOutputsComboBoxLabel->SetText(GetNoOutputText());
+	}
+
+	CollectionGroupsComboBox->RefreshOptions();
+	CollectionGroupsComboBox->ClearSelection();
+
+	if (NodeOutputs.Num() > 0)
+	{
+		UpdateCollectionGroups(NodeOutputs[0]);
+	}
+
+	if (CollectionGroups.Num() > 0)
+	{
+		CollectionGroupsComboBox->SetSelectedItem(CollectionGroups[0]);
+	}
+	else
+	{
+		CollectionGroupsComboBoxLabel->SetText(GetNoGroupText());
+	}
+}
+
+
+void SCollectionSpreadSheetWidget::SetStatusText()
+{
+	if (!NodeName.IsEmpty())
+	{
+		FString Str = FString::Printf(TEXT("Group has %d elements"), CollectionTable->GetNumItems());
+		StatusTextBlock->SetText(FText::FromString(Str));
+	}
+	else
+	{
+		StatusTextBlock->SetText(FText::FromString(" "));
+	}
+}
+
+
+void SCollectionSpreadSheetWidget::UpdateCollectionGroups(const FName& InOutputName)
+{
+	if (!InOutputName.IsNone())
+	{
+		CollectionGroups.Empty();
+
+		for (FName Group : GetCollectionTable()->GetCollectionInfoMap()[InOutputName.ToString()].Collection.GroupNames())
+		{
+			CollectionGroups.Add(Group);
+		}
+	}
+}
+
+#undef LOCTEXT_NAMESPACE
diff --git a/Engine/Plugins/Experimental/Dataflow/Source/DataflowEditor/Private/Dataflow/DataflowCollectionSpreadSheet.cpp b/Engine/Plugins/Experimental/Dataflow/Source/DataflowEditor/Private/Dataflow/DataflowCollectionSpreadSheet.cpp
new file mode 100644
--- /dev/null
+++ b/Engine/Plugins/Experimental/Dataflow/Source/DataflowEditor/Private/Dataflow/DataflowCollectionSpreadSheet.cpp
@@ -0,0 +1,79 @@
+// Copyright Epic Games, Inc. All Rights Reserved.
+
+#include "Dataflow/DataflowCollectionSpreadSheet.h"
+#include "Templates/EnableIf.h"
+#include "Dataflow/DataflowSelection.h"
+
+
+void FDataflowCollectionSpreadSheet::SetSupportedOutputTypes()
+{
+	GetSupportedOutputTypes().Empty();
+
+	GetSupportedOutputTypes().Add("FManagedArrayCollection");
+}
+
+void FDataflowCollectionSpreadSheet::UpdateViewData()
+{
+	if (CollectionSpreadSheet)
+	{
+		CollectionSpreadSheet->GetCollectionTable()->GetCollectionInfoMap().Empty();
+
+		if (GetSelectedNode())
+		{
+			if (GetSelectedNode()->IsBound())
+			{
+				if (TSharedPtr<FDataflowNode> DataflowNode = GetSelectedNode()->DataflowGraph->FindBaseNode(GetSelectedNode()->DataflowNodeGuid))
+				{
+					TArray<FDataflowOutput*> Outputs = DataflowNode->GetOutputs();
+
+					for (FDataflowOutput* Output : Outputs)
+					{
+						FName Name = Output->GetName();
+						FName Type = Output->GetType();
+
+						if (Output->GetType() == "FManagedArrayCollection")
+						{
+							const FManagedArrayCollection& Value = Output->GetValue<FManagedArrayCollection>(*GetContext(), FManagedArrayCollection());
+
+							CollectionSpreadSheet->GetCollectionTable()->GetCollectionInfoMap().Add(Name.ToString(), { Value });
+						}
+					}
+				}
+			}
+
+			CollectionSpreadSheet->SetData(GetSelectedNode()->GetName());
+		}
+		else
+		{
+			CollectionSpreadSheet->SetData(FString());
+		}
+
+		CollectionSpreadSheet->RefreshWidget();
+	}
+}
+
+
+void FDataflowCollectionSpreadSheet::SetCollectionSpreadSheet(TSharedPtr<SCollectionSpreadSheetWidget>& InCollectionSpreadSheet)
+{
+	ensure(!CollectionSpreadSheet);
+
+	CollectionSpreadSheet = InCollectionSpreadSheet;
+
+	if (CollectionSpreadSheet)
+	{
+		OnPinnedDownChangedDelegateHandle = CollectionSpreadSheet->GetOnPinnedDownChangedDelegate().AddRaw(this, &FDataflowCollectionSpreadSheet::OnPinnedDownChanged);
+		OnRefreshLockedChangedDelegateHandle = CollectionSpreadSheet->GetOnRefreshLockedChangedDelegate().AddRaw(this, &FDataflowCollectionSpreadSheet::OnRefreshLockedChanged);
+	}
+}
+
+
+FDataflowCollectionSpreadSheet::~FDataflowCollectionSpreadSheet()
+{
+	if (CollectionSpreadSheet)
+	{
+		CollectionSpreadSheet->GetOnPinnedDownChangedDelegate().Remove(OnPinnedDownChangedDelegateHandle);
+	}
+}
+
+
+#undef LOCTEXT_NAMESPACE
diff --git a/Engine/Plugins/Experimental/Dataflow/Source/DataflowEditor/Private/Dataflow/DataflowEditorToolkit.cpp b/Engine/Plugins/Experimental/Dataflow/Source/DataflowEditor/Private/Dataflow/DataflowEditorToolkit.cpp
--- a/Engine/Plugins/Experimental/Dataflow/Source/DataflowEditor/Private/Dataflow/DataflowEditorToolkit.cpp
+++ b/Engine/Plugins/Experimental/Dataflow/Source/DataflowEditor/Private/Dataflow/DataflowEditorToolkit.cpp
@@ -38,8 +38,14 @@ const FName FDataflowEditorToolkit::GraphCanvasTabId(TEXT("DataflowEditor_GraphC
 const FName FDataflowEditorToolkit::AssetDetailsTabId(TEXT("DataflowEditor_AssetDetails"));
 const FName FDataflowEditorToolkit::NodeDetailsTabId(TEXT("DataflowEditor_NodeDetails"));
 const FName FDataflowEditorToolkit::SkeletalTabId(TEXT("DataflowEditor_Skeletal"));
-const FName FDataflowEditorToolkit::SelectionViewTabId(TEXT("DataflowEditor_SelectionView"));
-
+const FName FDataflowEditorToolkit::SelectionViewTabId_1(TEXT("DataflowEditor_SelectionView_1"));
+const FName FDataflowEditorToolkit::SelectionViewTabId_2(TEXT("DataflowEditor_SelectionView_2"));
+const FName FDataflowEditorToolkit::SelectionViewTabId_3(TEXT("DataflowEditor_SelectionView_3"));
+const FName FDataflowEditorToolkit::SelectionViewTabId_4(TEXT("DataflowEditor_SelectionView_4"));
+const FName FDataflowEditorToolkit::CollectionSpreadSheetTabId_1(TEXT("DataflowEditor_CollectionSpreadSheet_1"));
+const FName FDataflowEditorToolkit::CollectionSpreadSheetTabId_2(TEXT("DataflowEditor_CollectionSpreadSheet_2"));
+const FName FDataflowEditorToolkit::CollectionSpreadSheetTabId_3(TEXT("DataflowEditor_CollectionSpreadSheet_3"));
+const FName FDataflowEditorToolkit::CollectionSpreadSheetTabId_4(TEXT("DataflowEditor_CollectionSpreadSheet_4"));
 
 UDataflow* GetDataflowFrom(UObject* InObject)
 {
@@ -463,10 +469,40 @@ TSharedRef<SDockTab> FDataflowEditorToolkit::SpawnTab_Skeletal(const FSpawnTabAr
 
 TSharedRef<SDockTab> FDataflowEditorToolkit::SpawnTab_SelectionView(const FSpawnTabArgs& Args)
 {
-	check(Args.GetTabId() == SelectionViewTabId);
+//	check(Args.GetTabId().TabType == SelectionViewTabId_1);
 
-	DataflowSelectionView = MakeShared<FDataflowSelectionView>(FDataflowSelectionView());
-	ViewListeners.Add(DataflowSelectionView.Get());
+	if (Args.GetTabId() == SelectionViewTabId_1)
+	{
+		DataflowSelectionView_1 = MakeShared<FDataflowSelectionView>(FDataflowSelectionView());
+		if (DataflowSelectionView_1.IsValid())
+		{
+			ViewListeners.Add(DataflowSelectionView_1.Get());
+		}
+	}
+	else if (Args.GetTabId() == SelectionViewTabId_2)
+	{
+		DataflowSelectionView_2 = MakeShared<FDataflowSelectionView>(FDataflowSelectionView());
+		if (DataflowSelectionView_2.IsValid())
+		{
+			ViewListeners.Add(DataflowSelectionView_2.Get());
+		}
+	}
+	else if (Args.GetTabId() == SelectionViewTabId_3)
+	{
+		DataflowSelectionView_3 = MakeShared<FDataflowSelectionView>(FDataflowSelectionView());
+		if (DataflowSelectionView_3.IsValid())
+		{
+			ViewListeners.Add(DataflowSelectionView_3.Get());
+		}
+	}
+	else if (Args.GetTabId() == SelectionViewTabId_4)
+	{
+		DataflowSelectionView_4 = MakeShared<FDataflowSelectionView>(FDataflowSelectionView());
+		if (DataflowSelectionView_4.IsValid())
+		{
+			ViewListeners.Add(DataflowSelectionView_4.Get());
+		}
+	}
 
 	TSharedPtr<SSelectionViewWidget> SelectionViewWidget;
 
@@ -477,21 +513,151 @@ TSharedRef<SDockTab> FDataflowEditorToolkit::SpawnTab_SelectionView(const FSpawn
 
 	if (SelectionViewWidget)
 	{
-		DataflowSelectionView->SetSelectionView(SelectionViewWidget);
+		if (Args.GetTabId() == SelectionViewTabId_1)
+		{
+			DataflowSelectionView_1->SetSelectionView(SelectionViewWidget);
+
+			// Set the Context on the interface
+			if (TSharedPtr<Dataflow::FContext> CurrentContext = this->GetContext())
+			{
+				DataflowSelectionView_1->SetContext(CurrentContext);
+			}
+		}
+		else if (Args.GetTabId() == SelectionViewTabId_2)
+		{
+			DataflowSelectionView_2->SetSelectionView(SelectionViewWidget);
+
+			// Set the Context on the interface
+			if (TSharedPtr<Dataflow::FContext> CurrentContext = this->GetContext())
+			{
+				DataflowSelectionView_2->SetContext(CurrentContext);
+			}
+		}
+		else if (Args.GetTabId() == SelectionViewTabId_3)
+		{
+			DataflowSelectionView_3->SetSelectionView(SelectionViewWidget);
+
+			// Set the Context on the interface
+			if (TSharedPtr<Dataflow::FContext> CurrentContext = this->GetContext())
+			{
+				DataflowSelectionView_3->SetContext(CurrentContext);
+			}
+		}
+		else if (Args.GetTabId() == SelectionViewTabId_4)
+		{
+			DataflowSelectionView_4->SetSelectionView(SelectionViewWidget);
+
+			// Set the Context on the interface
+			if (TSharedPtr<Dataflow::FContext> CurrentContext = this->GetContext())
+			{
+				DataflowSelectionView_4->SetContext(CurrentContext);
+			}
+		}
+	}
+
+	DockableTab->SetOnTabClosed(SDockTab::FOnTabClosedCallback::CreateRaw(this, &FDataflowEditorToolkit::OnTabClosed));
+
+	return DockableTab;
+}
+
+TSharedRef<SDockTab> FDataflowEditorToolkit::SpawnTab_CollectionSpreadSheet(const FSpawnTabArgs& Args)
+{
+//	check(Args.GetTabId() == CollectionSpreadSheetTabId_1);
+
+	if (Args.GetTabId() == CollectionSpreadSheetTabId_1)
+	{
+		DataflowCollectionSpreadSheet_1 = MakeShared<FDataflowCollectionSpreadSheet>(FDataflowCollectionSpreadSheet());
+		if (DataflowCollectionSpreadSheet_1.IsValid())
+		{
+			ViewListeners.Add(DataflowCollectionSpreadSheet_1.Get());
+		}
+	}
+	else if (Args.GetTabId() == CollectionSpreadSheetTabId_2)
+	{
+		DataflowCollectionSpreadSheet_2 = MakeShared<FDataflowCollectionSpreadSheet>(FDataflowCollectionSpreadSheet());
+		if (DataflowCollectionSpreadSheet_2.IsValid())
+		{
+			ViewListeners.Add(DataflowCollectionSpreadSheet_2.Get());
+		}
+	}
+	else if (Args.GetTabId() == CollectionSpreadSheetTabId_3)
+	{
+		DataflowCollectionSpreadSheet_3 = MakeShared<FDataflowCollectionSpreadSheet>(FDataflowCollectionSpreadSheet());
+		if (DataflowCollectionSpreadSheet_3.IsValid())
+		{
+			ViewListeners.Add(DataflowCollectionSpreadSheet_3.Get());
+		}
+	}
+	else if (Args.GetTabId() == CollectionSpreadSheetTabId_4)
+	{
+		DataflowCollectionSpreadSheet_4 = MakeShared<FDataflowCollectionSpreadSheet>(FDataflowCollectionSpreadSheet());
+		if (DataflowCollectionSpreadSheet_4.IsValid())
+		{
+			ViewListeners.Add(DataflowCollectionSpreadSheet_4.Get());
+		}
+	}
+
+	TSharedPtr<SCollectionSpreadSheetWidget> CollectionSpreadSheetWidget;
+
+	TSharedRef<SDockTab> DockableTab = SNew(SDockTab)
+	[
+		SAssignNew(CollectionSpreadSheetWidget, SCollectionSpreadSheetWidget)
+	];
+
+	if (CollectionSpreadSheetWidget)
+	{
+		if (Args.GetTabId() == CollectionSpreadSheetTabId_1)
+		{
+			DataflowCollectionSpreadSheet_1->SetCollectionSpreadSheet(CollectionSpreadSheetWidget);
 
-		// Set the Context on the interface
-		if (TSharedPtr<Dataflow::FContext> CurrentContext = this->GetContext())
+			// Set the Context on the interface
+			if (TSharedPtr<Dataflow::FContext> CurrentContext = this->GetContext())
+			{
+				DataflowCollectionSpreadSheet_1->SetContext(CurrentContext);
+			}
+		}
+		else if (Args.GetTabId() == CollectionSpreadSheetTabId_2)
+		{
+			DataflowCollectionSpreadSheet_2->SetCollectionSpreadSheet(CollectionSpreadSheetWidget);
+
+			// Set the Context on the interface
+			if (TSharedPtr<Dataflow::FContext> CurrentContext = this->GetContext())
+			{
+				DataflowCollectionSpreadSheet_2->SetContext(CurrentContext);
+			}
+		}
+		else if (Args.GetTabId() == CollectionSpreadSheetTabId_3)
 		{
-			DataflowSelectionView->SetContext(CurrentContext);
+			DataflowCollectionSpreadSheet_3->SetCollectionSpreadSheet(CollectionSpreadSheetWidget);
+
+			// Set the Context on the interface
+			if (TSharedPtr<Dataflow::FContext> CurrentContext = this->GetContext())
+			{
+				DataflowCollectionSpreadSheet_3->SetContext(CurrentContext);
+			}
+		}
+		else if (Args.GetTabId() == CollectionSpreadSheetTabId_4)
+		{
+			DataflowCollectionSpreadSheet_4->SetCollectionSpreadSheet(CollectionSpreadSheetWidget);
+
+			// Set the Context on the interface
+			if (TSharedPtr<Dataflow::FContext> CurrentContext = this->GetContext())
+			{
+				DataflowCollectionSpreadSheet_4->SetContext(CurrentContext);
+			}
 		}
 	}
 
+	DockableTab->SetOnTabClosed(SDockTab::FOnTabClosedCallback::CreateRaw(this, &FDataflowEditorToolkit::OnTabClosed));
+
 	return DockableTab;
 }
 
 void FDataflowEditorToolkit::RegisterTabSpawners(const TSharedRef<FTabManager>& InTabManager)
 {
 	TSharedRef<FWorkspaceItem> WorkspaceMenuCategoryRef = InTabManager->AddLocalWorkspaceMenuCategory(LOCTEXT("WorkspaceMenu_DataflowEditor", "Dataflow Editor"));
+	TSharedRef<FWorkspaceItem> SelectionViewWorkspaceMenuCategoryRef = WorkspaceMenuCategoryRef->AddGroup(LOCTEXT("WorkspaceMenu_SelectionView", "Selection View"), FSlateIcon(FAppStyle::GetAppStyleSetName(), "LevelEditor.Tabs.Outliner"));
+	TSharedRef<FWorkspaceItem> CollectionSpreadSheetWorkspaceMenuCategoryRef = WorkspaceMenuCategoryRef->AddGroup(LOCTEXT("WorkspaceMenu_CollectionSpreadSheet", "Collection SpreadSheet"), FSlateIcon(FAppStyle::GetAppStyleSetName(), "LevelEditor.Tabs.Outliner"));
 
 	InTabManager->RegisterTabSpawner(ViewportTabId, FOnSpawnTab::CreateSP(this, &FDataflowEditorToolkit::SpawnTab_Viewport))
 		.SetDisplayName(LOCTEXT("DataflowViewportTab", "Dataflow Viewport"))
@@ -518,14 +684,104 @@ void FDataflowEditorToolkit::RegisterTabSpawners(const TSharedRef<FTabManager>&
 		.SetGroup(WorkspaceMenuCategoryRef)
 		.SetIcon(FSlateIcon(FAppStyle::GetAppStyleSetName(), "LevelEditor.Tabs.SkeletalHierarchy"));
 
-	InTabManager->RegisterTabSpawner(SelectionViewTabId, FOnSpawnTab::CreateSP(this, &FDataflowEditorToolkit::SpawnTab_SelectionView))
-		.SetDisplayName(LOCTEXT("DataflowSelectionViewTab", "Selection View"))
-		.SetGroup(WorkspaceMenuCategoryRef)
-		.SetIcon(FSlateIcon(FAppStyle::GetAppStyleSetName(), "LevelEditor.Tabs.SelectionView"));
+	InTabManager->RegisterTabSpawner(SelectionViewTabId_1, FOnSpawnTab::CreateSP(this, &FDataflowEditorToolkit::SpawnTab_SelectionView))
+		.SetDisplayName(LOCTEXT("DataflowSelectionViewTab1", "Selection View 1"))
+		.SetGroup(SelectionViewWorkspaceMenuCategoryRef)
+		.SetIcon(FSlateIcon(FAppStyle::GetAppStyleSetName(), "LevelEditor.Tabs.Outliner"));
+
+	InTabManager->RegisterTabSpawner(SelectionViewTabId_2, FOnSpawnTab::CreateSP(this, &FDataflowEditorToolkit::SpawnTab_SelectionView))
+		.SetDisplayName(LOCTEXT("DataflowSelectionViewTab2", "Selection View 2"))
+		.SetGroup(SelectionViewWorkspaceMenuCategoryRef)
+		.SetIcon(FSlateIcon(FAppStyle::GetAppStyleSetName(), "LevelEditor.Tabs.Outliner"));
+
+	InTabManager->RegisterTabSpawner(SelectionViewTabId_3, FOnSpawnTab::CreateSP(this, &FDataflowEditorToolkit::SpawnTab_SelectionView))
+		.SetDisplayName(LOCTEXT("DataflowSelectionViewTab3", "Selection View 3"))
+		.SetGroup(SelectionViewWorkspaceMenuCategoryRef)
+		.SetIcon(FSlateIcon(FAppStyle::GetAppStyleSetName(), "LevelEditor.Tabs.Outliner"));
+
+	InTabManager->RegisterTabSpawner(SelectionViewTabId_4, FOnSpawnTab::CreateSP(this, &FDataflowEditorToolkit::SpawnTab_SelectionView))
+		.SetDisplayName(LOCTEXT("DataflowSelectionViewTab4", "Selection View 4"))
+		.SetGroup(SelectionViewWorkspaceMenuCategoryRef)
+		.SetIcon(FSlateIcon(FAppStyle::GetAppStyleSetName(), "LevelEditor.Tabs.Outliner"));
+
+	InTabManager->RegisterTabSpawner(CollectionSpreadSheetTabId_1, FOnSpawnTab::CreateSP(this, &FDataflowEditorToolkit::SpawnTab_CollectionSpreadSheet))
+		.SetDisplayName(LOCTEXT("DataflowCollectionSpreadSheetTab1", "Collection SpreadSheet 1"))
+		.SetGroup(CollectionSpreadSheetWorkspaceMenuCategoryRef)
+		.SetIcon(FSlateIcon(FAppStyle::GetAppStyleSetName(), "LevelEditor.Tabs.Outliner"));
+
+	InTabManager->RegisterTabSpawner(CollectionSpreadSheetTabId_2, FOnSpawnTab::CreateSP(this, &FDataflowEditorToolkit::SpawnTab_CollectionSpreadSheet))
+		.SetDisplayName(LOCTEXT("DataflowCollectionSpreadSheetTab2", "Collection SpreadSheet 2"))
+		.SetGroup(CollectionSpreadSheetWorkspaceMenuCategoryRef)
+		.SetIcon(FSlateIcon(FAppStyle::GetAppStyleSetName(), "LevelEditor.Tabs.Outliner"));
+
+	InTabManager->RegisterTabSpawner(CollectionSpreadSheetTabId_3, FOnSpawnTab::CreateSP(this, &FDataflowEditorToolkit::SpawnTab_CollectionSpreadSheet))
+		.SetDisplayName(LOCTEXT("DataflowCollectionSpreadSheetTab3", "Collection SpreadSheet 3"))
+		.SetGroup(CollectionSpreadSheetWorkspaceMenuCategoryRef)
+		.SetIcon(FSlateIcon(FAppStyle::GetAppStyleSetName(), "LevelEditor.Tabs.Outliner"));
+
+	InTabManager->RegisterTabSpawner(CollectionSpreadSheetTabId_4, FOnSpawnTab::CreateSP(this, &FDataflowEditorToolkit::SpawnTab_CollectionSpreadSheet))
+		.SetDisplayName(LOCTEXT("DataflowCollectionSpreadSheetTab4", "Collection SpreadSheet 4"))
+		.SetGroup(CollectionSpreadSheetWorkspaceMenuCategoryRef)
+		.SetIcon(FSlateIcon(FAppStyle::GetAppStyleSetName(), "LevelEditor.Tabs.Outliner"));
 
 	FAssetEditorToolkit::RegisterTabSpawners(InTabManager);
 }
 
+void FDataflowEditorToolkit::UnregisterTabSpawners(const TSharedRef<class FTabManager>& InTabManager)
+{
+	FAssetEditorToolkit::UnregisterTabSpawners(InTabManager);
+
+	InTabManager->UnregisterTabSpawner(ViewportTabId);
+	InTabManager->UnregisterTabSpawner(GraphCanvasTabId);
+	InTabManager->UnregisterTabSpawner(AssetDetailsTabId);
+	InTabManager->UnregisterTabSpawner(NodeDetailsTabId);
+	InTabManager->UnregisterTabSpawner(SkeletalTabId);
+	InTabManager->UnregisterTabSpawner(SelectionViewTabId_1);
+	InTabManager->UnregisterTabSpawner(SelectionViewTabId_2);
+	InTabManager->UnregisterTabSpawner(SelectionViewTabId_3);
+	InTabManager->UnregisterTabSpawner(SelectionViewTabId_4);
+	InTabManager->UnregisterTabSpawner(CollectionSpreadSheetTabId_1);
+	InTabManager->UnregisterTabSpawner(CollectionSpreadSheetTabId_2);
+	InTabManager->UnregisterTabSpawner(CollectionSpreadSheetTabId_3);
+	InTabManager->UnregisterTabSpawner(CollectionSpreadSheetTabId_4);
+}
+
+void FDataflowEditorToolkit::OnTabClosed(TSharedRef<SDockTab> Tab)
+{
+	if (Tab->GetTabLabel().EqualTo(FText::FromString("Selection View 1")))
+	{
+		ViewListeners.Remove(DataflowSelectionView_1.Get());
+	}
+	else if (Tab->GetTabLabel().EqualTo(FText::FromString("Selection View 2")))
+	{
+		ViewListeners.Remove(DataflowSelectionView_2.Get());
+	}
+	else if (Tab->GetTabLabel().EqualTo(FText::FromString("Selection View 3")))
+	{
+		ViewListeners.Remove(DataflowSelectionView_3.Get());
+	}
+	else if (Tab->GetTabLabel().EqualTo(FText::FromString("Selection View 4")))
+	{
+		ViewListeners.Remove(DataflowSelectionView_4.Get());
+	}
+	else if (Tab->GetTabLabel().EqualTo(FText::FromString("Collection SpreadSheet 1")))
+	{
+		ViewListeners.Remove(DataflowCollectionSpreadSheet_1.Get());
+	}
+	else if (Tab->GetTabLabel().EqualTo(FText::FromString("Collection SpreadSheet 2")))
+	{
+		ViewListeners.Remove(DataflowCollectionSpreadSheet_2.Get());
+	}
+	else if (Tab->GetTabLabel().EqualTo(FText::FromString("Collection SpreadSheet 3")))
+	{
+		ViewListeners.Remove(DataflowCollectionSpreadSheet_3.Get());
+	}
+	else if (Tab->GetTabLabel().EqualTo(FText::FromString("Collection SpreadSheet 4")))
+	{
+		ViewListeners.Remove(DataflowCollectionSpreadSheet_4.Get());
+	}
+}
+
 FName FDataflowEditorToolkit::GetToolkitFName() const
 {
 	return FName("DataflowEditor");
diff --git a/Engine/Plugins/Experimental/Dataflow/Source/DataflowEditor/Private/Dataflow/SelectionViewWidget.cpp b/Engine/Plugins/Experimental/Dataflow/Source/DataflowEditor/Private/Dataflow/SelectionViewWidget.cpp
--- a/Engine/Plugins/Experimental/Dataflow/Source/DataflowEditor/Private/Dataflow/SelectionViewWidget.cpp
+++ b/Engine/Plugins/Experimental/Dataflow/Source/DataflowEditor/Private/Dataflow/SelectionViewWidget.cpp
@@ -5,7 +5,7 @@
 #include "Styling/StarshipCoreStyle.h"
 
 
-#define LOCTEXT_NAMESPACE "DataflowGraphEditor"
+#define LOCTEXT_NAMESPACE "SelectionViewWidget"
 
 
 const FName FSelectionViewHeader::IndexColumnNameTransform = FName("Transform Index");
@@ -44,7 +44,6 @@ TSharedRef<SWidget> SSelectionViewRow::GenerateWidgetForColumn(const FName& Colu
 				.Text(FText::FromString(AttrValue))
 				.ShadowColorAndOpacity(FLinearColor(0.1f, 0.1f, 0.1f, 1.f))
 				.Visibility(EVisibility::Visible)
-				.ToolTipText(FText::FromString("This is an example tooltip."))
 				.Font_Lambda([this]() -> FSlateFontInfo
 				{
 					if (Item->Values[1] == FSelectionViewItem::SelectedName.ToString())
@@ -64,9 +63,7 @@ TSharedRef<SWidget> SSelectionViewRow::GenerateWidgetForColumn(const FName& Colu
 // ----------------------------------------------------------------------------
 //
 
-void SSelectionView::Construct(
-	const FArguments& InArgs
-)
+void SSelectionView::Construct(const FArguments& InArgs)
 {
 	SelectedOutput = InArgs._SelectedOutput;
 
@@ -90,7 +87,7 @@ void SSelectionView::Construct(
 			+ SOverlay::Slot()
 			[
 				SAssignNew(ListView, SListView<TSharedPtr<const FSelectionViewItem>>)
-				.SelectionMode(ESelectionMode::None)
+				.SelectionMode(ESelectionMode::Multi)
 				.ListItemsSource(&ListItems)
 				.OnGenerateRow(this, &SSelectionView::GenerateRow)
 				.HeaderRow(HeaderRowWidget)
@@ -153,6 +150,9 @@ void SSelectionView::RegenerateHeader()
 			SHeaderRow::Column(ColumnName)
 			.DefaultLabel(FText::FromName(ColumnName))
 			.FillWidth(CustomFillWidth)
+			.HAlignCell(HAlign_Center)
+			.HAlignHeader(HAlign_Center)
+			.VAlignCell(VAlign_Center)
 		);
 	}
 }
@@ -218,6 +218,21 @@ TSharedRef<ITableRow> SSelectionView::GenerateRow(TSharedPtr<const FSelectionVie
 // ----------------------------------------------------------------------------
 //
 
+void SSelectionViewWidget::NodeOutputsComboBoxSelectionChanged(FName InSelectedOutput, ESelectInfo::Type InSelectInfo)
+{
+	if (SelectionTable)
+	{
+		if (SelectionTable->GetSelectedOutput() != InSelectedOutput)
+		{
+			SelectionTable->SetSelectedOutput(InSelectedOutput);
+
+			NodeOutputsComboBoxLabel->SetText(FText::FromName(SelectionTable->GetSelectedOutput()));
+
+			SetStatusText();
+		}
+	}
+}
+
 
 void SSelectionViewWidget::Construct(const FArguments& InArgs)
 {
@@ -341,36 +356,17 @@ void SSelectionViewWidget::Construct(const FArguments& InArgs)
 			.Padding(10.0f, 0.0f, 10.0f, 10.0f)
 			[
 				SAssignNew(NodeOutputsComboBox, SComboBox<FName>)
+				.ToolTipText(LOCTEXT("NodeOutputsToolTip", "Select a node output to see the output's data"))
 				.OptionsSource(&NodeOutputs)
-				.OnSelectionChanged(SComboBox<FName>::FOnSelectionChanged::CreateLambda(
-					[this](FName SelectedName, ESelectInfo::Type)
-					{
-						if (SelectionTable)
-						{
-							SelectionTable->SetSelectedOutput(SelectedName);
-						}
-
-						SetStatusText();
-					}))
-				.OnGenerateWidget(SComboBox<FName>::FOnGenerateWidget::CreateLambda(
-					[](FName Item)
-					{
-						return SNew(STextBlock)
-							.Text(FText::FromName(Item));
-					}))
+				.OnGenerateWidget(SComboBox<FName>::FOnGenerateWidget::CreateLambda([](FName Item)->TSharedRef<SWidget>
+				{
+					return SNew(STextBlock)
+						.Text(FText::FromName(Item));
+				}))
+				.OnSelectionChanged(this, &SSelectionViewWidget::NodeOutputsComboBoxSelectionChanged)
 				[
-					SNew(STextBlock)
-					.Text_Lambda([this]()
-					{
-						if (SelectionTable)
-						{
-							return FText::FromName(SelectionTable->GetSelectedOutput());
-						}
-						else
-						{
-							return FText::FromString("");
-						}
-					})
+					SAssignNew(NodeOutputsComboBoxLabel, STextBlock)
+					.Text(GetNoOutputText())
 				]
 			]
 		]
@@ -426,6 +422,10 @@ void SSelectionViewWidget::RefreshWidget()
 	{
 		NodeOutputsComboBox->SetSelectedItem(NodeOutputs[0]);
 	}
+	else
+	{
+		NodeOutputsComboBoxLabel->SetText(GetNoOutputText());
+	}
 }
 
 
@@ -442,5 +442,9 @@ void SSelectionViewWidget::SetStatusText()
 	}
 }
 
+FText SSelectionViewWidget::GetNoOutputText()
+{
+	return LOCTEXT("NoOutput", "No Output(s)");
+}
 
 #undef LOCTEXT_NAMESPACE
diff --git a/Engine/Plugins/Experimental/Dataflow/Source/DataflowEditor/Public/Dataflow/CollectionSpreadSheetWidget.h b/Engine/Plugins/Experimental/Dataflow/Source/DataflowEditor/Public/Dataflow/CollectionSpreadSheetWidget.h
new file mode 100644
--- /dev/null
+++ b/Engine/Plugins/Experimental/Dataflow/Source/DataflowEditor/Public/Dataflow/CollectionSpreadSheetWidget.h
@@ -0,0 +1,171 @@
+// Copyright Epic Games, Inc. All Rights Reserved.
+
+#pragma once
+
+#include "CoreMinimal.h"
+#include "Dataflow/DataflowSelection.h"
+#include "DetailLayoutBuilder.h"
+#include "Widgets/Views/STableRow.h"
+#include "Widgets/Text/STextBlock.h"
+#include "Widgets/Input/SComboBox.h"
+#include "GeometryCollection/ManagedArrayCollection.h"
+
+/**
+* Struct to hold OutputType/Selection data for the outputs
+* Data is stored in a map using the OutputName as key: TMap<FString, SCollectionInfo>
+*/
+struct SCollectionInfo
+{
+//	FString OutputType;
+	const FManagedArrayCollection Collection;
+};
+
+
+/** 
+* Header
+* 1st column possible values: Transform Index/Face Index/Vertex Index based on the OutputType
+* 2nd column is SelectionStatus
+*/
+struct FCollectionSpreadSheetHeader
+{
+	static const FName IndexColumnName;
+
+	TArray<FName> ColumnNames;
+};
+
+
+/** 
+* Representing a row in the table 
+* Index/SelectionStatus
+*/
+struct FCollectionSpreadSheetItem
+{
+	TArray<FString> Values;
+};
+
+
+/** 
+* 
+*/
+class SCollectionSpreadSheetRow : public SMultiColumnTableRow<TSharedPtr<const FCollectionSpreadSheetItem>>
+{
+public:
+	SLATE_BEGIN_ARGS(SCollectionSpreadSheetRow) 
+	{}
+		SLATE_ARGUMENT(TSharedPtr<const FCollectionSpreadSheetHeader>, Header)
+		SLATE_ARGUMENT(TSharedPtr<const FCollectionSpreadSheetItem>, Item)
+	SLATE_END_ARGS()
+
+	void Construct(const FArguments& InArgs, TSharedRef<STableViewBase> OwnerTableView, const TSharedPtr<const FCollectionSpreadSheetHeader>& InHeader, const TSharedPtr<const FCollectionSpreadSheetItem>& InItem);
+
+	virtual TSharedRef<SWidget> GenerateWidgetForColumn(const FName& ColumnName) override;
+
+private:
+	TSharedPtr<const FCollectionSpreadSheetHeader> Header;
+	TSharedPtr<const FCollectionSpreadSheetItem> Item;
+};
+
+/** 
+* 2xn grid to display Collection data
+*/
+class SCollectionSpreadSheet : public SCompoundWidget
+{
+public:
+	SLATE_BEGIN_ARGS(SCollectionSpreadSheet) {}
+		SLATE_ARGUMENT(FName, SelectedOutput)
+	SLATE_END_ARGS()
+
+	void Construct(const FArguments& InArgs);
+
+	TMap<FString, SCollectionInfo>& GetCollectionInfoMap() { return CollectionInfoMap; }
+
+	const FName& GetSelectedOutput() const;
+	void SetSelectedOutput(const FName& InSelectedOutput);
+	
+	const FName& GetSelectedGroup() const;
+	void SetSelectedGroup(const FName& InSelectedGroup);
+
+	int32 GetNumItems() const { return NumItems; }
+	void SetNumItems(int32 InNumItems) { NumItems = InNumItems; }
+
+	TSharedRef<ITableRow> GenerateRow(TSharedPtr<const FCollectionSpreadSheetItem> InItem, const TSharedRef<STableViewBase>& OwnerTable);
+
+	TSharedPtr<SScrollBar> SpreadSheetVerticalScrollBar;
+
+	TMap<FString, int32> AttrTypeWidthMap;
+
+private:
+	FName SelectedOutput = FName(TEXT(""));
+	FName SelectedGroup = FName(TEXT(""));
+	TMap<FString, SCollectionInfo> CollectionInfoMap;
+
+	TSharedPtr<SListView<TSharedPtr<const FCollectionSpreadSheetItem>>> ListView;
+//	TSharedPtr<STreeView<TSharedPtr<const FCollectionSpreadSheetItem>>> ListView;
+	TArray<TSharedPtr<const FCollectionSpreadSheetItem>> ListItems;
+
+	TSharedPtr<FCollectionSpreadSheetHeader> Header;
+	TSharedPtr<SHeaderRow> HeaderRowWidget;
+
+	int32 NumItems = 0;
+	void RegenerateHeader();
+	void RepopulateListView();
+};
+
+/**
+* Widget for the CollectionSpreadSheet panel
+*/
+class SCollectionSpreadSheetWidget : public SCompoundWidget
+{
+public:
+	SLATE_BEGIN_ARGS(SCollectionSpreadSheetWidget) {}
+	SLATE_END_ARGS()
+
+	/**
+	* Slate construction function
+	*
+	* @param InArgs - Arguments passed from Slate
+	*
+	*/
+	void Construct(const FArguments& InArgs);
+
+	void SetData(const FString& InNodeName);
+	void RefreshWidget();
+	TSharedPtr<SCollectionSpreadSheet> GetCollectionTable() { return CollectionTable; }
+	void SetStatusText();
+
+	void UpdateCollectionGroups(const FName& InOutputName);
+
+	/** Gets a multicast delegate which is called whenever the PinnedDown button clicked */
+	DECLARE_MULTICAST_DELEGATE_OneParam(FOnPinnedDownChanged, const bool);
+	FOnPinnedDownChanged& GetOnPinnedDownChangedDelegate() { return OnPinnedDownChangedDelegate; }
+
+	/** Gets a multicast delegate which is called whenever the LockRefresh button clicked */
+	DECLARE_MULTICAST_DELEGATE_OneParam(FOnRefreshLockedChanged, const bool);
+	FOnRefreshLockedChanged& GetOnRefreshLockedChangedDelegate() { return OnRefreshLockedChangedDelegate; }
+
+	// UI callbacks
+	void NodeOutputsComboBoxSelectionChanged(FName InSelectedOutput, ESelectInfo::Type InSelectInfo);
+	void CollectionGroupsComboBoxSelectionChanged(FName InSelectedOutput, ESelectInfo::Type InSelectInfo);
+	FText GetNoOutputText();
+	FText GetNoGroupText();
+
+private:
+	TSharedPtr<STextBlock> NodeNameTextBlock;
+	TSharedPtr<SComboBox<FName>> NodeOutputsComboBox;
+	TSharedPtr<STextBlock> NodeOutputsComboBoxLabel;
+	TSharedPtr<SComboBox<FName>> CollectionGroupsComboBox;
+	TSharedPtr<STextBlock> CollectionGroupsComboBoxLabel;
+	TSharedPtr<SCollectionSpreadSheet> CollectionTable;
+	TSharedPtr<STextBlock> StatusTextBlock;
+
+	FString NodeName;
+	TArray<FName> NodeOutputs;
+	TArray<FName> CollectionGroups;
+	bool bIsPinnedDown = false;
+	bool bIsRefreshLocked = false;
+
+protected:
+	FOnPinnedDownChanged OnPinnedDownChangedDelegate;
+	FOnRefreshLockedChanged OnRefreshLockedChangedDelegate;
+
+};
\ No newline at end of file
diff --git a/Engine/Plugins/Experimental/Dataflow/Source/DataflowEditor/Public/Dataflow/DataflowCollectionSpreadSheet.h b/Engine/Plugins/Experimental/Dataflow/Source/DataflowEditor/Public/Dataflow/DataflowCollectionSpreadSheet.h
new file mode 100644
--- /dev/null
+++ b/Engine/Plugins/Experimental/Dataflow/Source/DataflowEditor/Public/Dataflow/DataflowCollectionSpreadSheet.h
@@ -0,0 +1,34 @@
+// Copyright Epic Games, Inc. All Rights Reserved.
+
+#pragma once
+
+#include "Widgets/SCompoundWidget.h"
+#include "Widgets/Views/SListView.h"
+#include "Dataflow/DataflowEdNode.h"
+#include "Dataflow/CollectionSpreadSheetWidget.h"
+#include "Dataflow/DataflowSelection.h"
+#include "Dataflow/DataflowView.h"
+
+
+/**
+*
+* Class to handle the SelectionView widget
+*
+*/
+class FDataflowCollectionSpreadSheet : public FDataflowNodeView
+{
+public:
+	~FDataflowCollectionSpreadSheet();
+
+	virtual void SetSupportedOutputTypes() override;
+	virtual void UpdateViewData() override;
+
+	void SetCollectionSpreadSheet(TSharedPtr<SCollectionSpreadSheetWidget>& InCollectionSpreadSheet);
+
+private:
+	TSharedPtr<SCollectionSpreadSheetWidget> CollectionSpreadSheet;
+
+	FDelegateHandle OnPinnedDownChangedDelegateHandle;
+	FDelegateHandle OnRefreshLockedChangedDelegateHandle;
+};
+
diff --git a/Engine/Plugins/Experimental/Dataflow/Source/DataflowEditor/Public/Dataflow/DataflowEditorToolkit.h b/Engine/Plugins/Experimental/Dataflow/Source/DataflowEditor/Public/Dataflow/DataflowEditorToolkit.h
--- a/Engine/Plugins/Experimental/Dataflow/Source/DataflowEditor/Public/Dataflow/DataflowEditorToolkit.h
+++ b/Engine/Plugins/Experimental/Dataflow/Source/DataflowEditor/Public/Dataflow/DataflowEditorToolkit.h
@@ -10,6 +10,8 @@
 #include "TickableEditorObject.h"
 #include "Dataflow/DataflowSelectionView.h"
 #include "Dataflow/SelectionViewWidget.h"
+#include "Dataflow/DataflowCollectionSpreadSheet.h"
+#include "Dataflow/CollectionSpreadSheetWidget.h"
 
 class FEditorViewportTabContent;
 class IDetailsView;
@@ -61,13 +63,18 @@ public:
 
 	// Tab spawners 
 	virtual void RegisterTabSpawners(const TSharedRef<FTabManager>& TabManager) override;
+	virtual void UnregisterTabSpawners(const TSharedRef<class FTabManager>& TabManager) override;
+
 	TSharedRef<SDockTab> SpawnTab_Viewport(const FSpawnTabArgs& Args);
 	TSharedRef<SDockTab> SpawnTab_GraphCanvas(const FSpawnTabArgs& Args);
 	TSharedRef<SDockTab> SpawnTab_AssetDetails(const FSpawnTabArgs& Args);
 	TSharedRef<SDockTab> SpawnTab_NodeDetails(const FSpawnTabArgs& Args);
 	TSharedRef<SDockTab> SpawnTab_Skeletal(const FSpawnTabArgs& Args);
 	TSharedRef<SDockTab> SpawnTab_SelectionView(const FSpawnTabArgs& Args);
+	TSharedRef<SDockTab> SpawnTab_CollectionSpreadSheet(const FSpawnTabArgs& Args);
 
+	// Callbacks for Tab
+	void OnTabClosed(TSharedRef<SDockTab> Tab);
 
 	// Member Access
 	UObject* GetAsset() { return Asset; }
@@ -126,13 +133,32 @@ private:
 	TSharedPtr<class ISkeletonTree> SkeletalEditor;
 	TSharedPtr<ISkeletonTree> CreateSkeletalEditorWidget(UObject* ObjectToEdit);
 
-	static const FName SelectionViewTabId;
-
 	TSet<UObject*> PrevNodeSelection;
-	Dataflow::FTimestamp LastContextRefreshTimestamp = Dataflow::FTimestamp::Invalid;
-	int32 ContextRefreshCounter = 0;
 
-	TSharedPtr<FDataflowSelectionView> DataflowSelectionView;
+	static const FName SelectionViewTabId_1;
+	TSharedPtr<FDataflowSelectionView> DataflowSelectionView_1;
+
+	static const FName SelectionViewTabId_2;
+	TSharedPtr<FDataflowSelectionView> DataflowSelectionView_2;
+
+	static const FName SelectionViewTabId_3;
+	TSharedPtr<FDataflowSelectionView> DataflowSelectionView_3;
+
+	static const FName SelectionViewTabId_4;
+	TSharedPtr<FDataflowSelectionView> DataflowSelectionView_4;
+
+	static const FName CollectionSpreadSheetTabId_1;
+	TSharedPtr<FDataflowCollectionSpreadSheet> DataflowCollectionSpreadSheet_1;
+
+	static const FName CollectionSpreadSheetTabId_2;
+	TSharedPtr<FDataflowCollectionSpreadSheet> DataflowCollectionSpreadSheet_2;
+
+	static const FName CollectionSpreadSheetTabId_3;
+	TSharedPtr<FDataflowCollectionSpreadSheet> DataflowCollectionSpreadSheet_3;
+
+	static const FName CollectionSpreadSheetTabId_4;
+	TSharedPtr<FDataflowCollectionSpreadSheet> DataflowCollectionSpreadSheet_4;
+
 	TArray<IDataflowViewListener*> ViewListeners;
 
 	TSharedPtr<Dataflow::FEngineContext> Context;
diff --git a/Engine/Plugins/Experimental/Dataflow/Source/DataflowEditor/Public/Dataflow/SelectionViewWidget.h b/Engine/Plugins/Experimental/Dataflow/Source/DataflowEditor/Public/Dataflow/SelectionViewWidget.h
--- a/Engine/Plugins/Experimental/Dataflow/Source/DataflowEditor/Public/Dataflow/SelectionViewWidget.h
+++ b/Engine/Plugins/Experimental/Dataflow/Source/DataflowEditor/Public/Dataflow/SelectionViewWidget.h
@@ -16,7 +16,7 @@
 struct SSelectionInfo
 {
 	FString OutputType;
-	TBitArray<> SelectionArray;
+	const TBitArray<> SelectionArray;
 };
 
 
@@ -130,6 +130,10 @@ public:
 	TSharedPtr<SSelectionView> GetSelectionTable() { return SelectionTable; }
 	void SetStatusText();
 
+	// UI callbacks
+	void NodeOutputsComboBoxSelectionChanged(FName InSelectedOutput, ESelectInfo::Type InSelectInfo);
+	FText GetNoOutputText();
+
 	/** Gets a multicast delegate which is called whenever the PinnedDown button clicked */
 	DECLARE_MULTICAST_DELEGATE_OneParam(FOnPinnedDownChanged, const bool);
 	FOnPinnedDownChanged& GetOnPinnedDownChangedDelegate() { return OnPinnedDownChangedDelegate; }
@@ -141,6 +145,7 @@ public:
 private:
 	TSharedPtr<STextBlock> NodeNameTextBlock;
 	TSharedPtr<SComboBox<FName>> NodeOutputsComboBox;
+	TSharedPtr<STextBlock> NodeOutputsComboBoxLabel;
 	TSharedPtr<SSelectionView> SelectionTable;
 	TSharedPtr<STextBlock> StatusTextBlock;
 
