Dataflow:
- fixed issue with LOCTEXT

diff --git a/Engine/Plugins/Experimental/Dataflow/Source/DataflowEditor/Private/Dataflow/CollectionSpreadSheetWidget.cpp b/Engine/Plugins/Experimental/Dataflow/Source/DataflowEditor/Private/Dataflow/CollectionSpreadSheetWidget.cpp
--- a/Engine/Plugins/Experimental/Dataflow/Source/DataflowEditor/Private/Dataflow/CollectionSpreadSheetWidget.cpp
+++ b/Engine/Plugins/Experimental/Dataflow/Source/DataflowEditor/Private/Dataflow/CollectionSpreadSheetWidget.cpp
@@ -6,9 +6,6 @@
 #include "Widgets/Layout/SScrollBox.h"
 
 
-#define LOCTEXT_NAMESPACE "CollectionSpreadSheetWidget"
-
-
 const FName FCollectionSpreadSheetHeader::IndexColumnName = FName("Index");
 
 namespace CollectionSpreadSheetHelpers
@@ -492,13 +489,13 @@ void SCollectionSpreadSheetWidget::CollectionGroupsComboBoxSelectionChanged(FNam
 
 FText SCollectionSpreadSheetWidget::GetNoOutputText()
 {
-	return LOCTEXT("NoOutput", "No Output(s)");
+	return FText::FromString("No Output(s)");
 }
 
 
 FText SCollectionSpreadSheetWidget::GetNoGroupText()
 {
-	return LOCTEXT("NoOutput", "No Group(s)");
+	return FText::FromString("No Group(s)");
 }
 
 
@@ -538,7 +535,7 @@ void SCollectionSpreadSheetWidget::Construct(const FArguments& InArgs)
 			.Padding(10.0f, 10.0f, 4.0f, 5.0f)
 			[
 				SNew(SButton)
-				.ToolTipText(LOCTEXT("PinDownButtonToolTip", "The button pins down the panel. When it pinned down it doesn't react to node selection change."))
+				.ToolTipText(FText::FromString("The button pins down the panel. When it pinned down it doesn't react to node selection change."))
 				.ButtonStyle(FAppStyle::Get(), "Button")
 				.HAlign(HAlign_Center)
 				.VAlign(VAlign_Center)
@@ -575,7 +572,7 @@ void SCollectionSpreadSheetWidget::Construct(const FArguments& InArgs)
 			.Padding(4.0f, 10.0f, 10.0f, 5.0f)
 			[
 				SNew(SButton)
-				.ToolTipText(LOCTEXT("LockRefreshButtonToolTip", "The button locks the refresh of the values in the panel."))
+				.ToolTipText(FText::FromString("The button locks the refresh of the values in the panel."))
 				.ButtonStyle(FAppStyle::Get(), "Button")
 				.HAlign(HAlign_Center)
 				.VAlign(VAlign_Center)
@@ -628,7 +625,7 @@ void SCollectionSpreadSheetWidget::Construct(const FArguments& InArgs)
 			.Padding(10.0f, 0.0f, 10.0f, 10.0f)
 			[
 				SAssignNew(NodeOutputsComboBox, SComboBox<FName>)
-				.ToolTipText(LOCTEXT("NodeOutputsToolTip", "Select a node output to see the output's data"))
+				.ToolTipText(FText::FromString("Select a node output to see the output's data"))
 				.OptionsSource(&NodeOutputs)
 				.OnGenerateWidget(SComboBox<FName>::FOnGenerateWidget::CreateLambda([](FName Output)->TSharedRef<SWidget>
 				{
@@ -662,7 +659,7 @@ void SCollectionSpreadSheetWidget::Construct(const FArguments& InArgs)
 			.Padding(10.0f, 0.0f, 10.0f, 2.0f)
 			[
 				SAssignNew(CollectionGroupsComboBox, SComboBox<FName>)
-				.ToolTipText(LOCTEXT("CollectionGroupsToolTip", "Select a group see the corresponding data"))
+				.ToolTipText(FText::FromString("Select a group see the corresponding data"))
 				.OptionsSource(&CollectionGroups)
 				.OnGenerateWidget(SComboBox<FName>::FOnGenerateWidget::CreateLambda([](FName Group)->TSharedRef<SWidget>
 				{
@@ -805,4 +802,3 @@ void SCollectionSpreadSheetWidget::UpdateCollectionGroups(const FName& InOutputN
 	}
 }
 
-#undef LOCTEXT_NAMESPACE
