Fix localization warnings
#rnx

diff --git a/Engine/Plugins/Experimental/Dataflow/Source/DataflowEditor/Private/Dataflow/DataflowCollectionSpreadSheet.cpp b/Engine/Plugins/Experimental/Dataflow/Source/DataflowEditor/Private/Dataflow/DataflowCollectionSpreadSheet.cpp
--- a/Engine/Plugins/Experimental/Dataflow/Source/DataflowEditor/Private/Dataflow/DataflowCollectionSpreadSheet.cpp
+++ b/Engine/Plugins/Experimental/Dataflow/Source/DataflowEditor/Private/Dataflow/DataflowCollectionSpreadSheet.cpp
@@ -74,6 +74,3 @@ FDataflowCollectionSpreadSheet::~FDataflowCollectionSpreadSheet()
 		CollectionSpreadSheet->GetOnPinnedDownChangedDelegate().Remove(OnPinnedDownChangedDelegateHandle);
 	}
 }
-
-
-#undef LOCTEXT_NAMESPACE
