Dataflow : Stop evaluation when in pie or simulate

diff --git a/Engine/Plugins/Experimental/Dataflow/Source/DataflowEditor/Private/Dataflow/DataflowEditorCommands.cpp b/Engine/Plugins/Experimental/Dataflow/Source/DataflowEditor/Private/Dataflow/DataflowEditorCommands.cpp
--- a/Engine/Plugins/Experimental/Dataflow/Source/DataflowEditor/Private/Dataflow/DataflowEditorCommands.cpp
+++ b/Engine/Plugins/Experimental/Dataflow/Source/DataflowEditor/Private/Dataflow/DataflowEditorCommands.cpp
@@ -8,6 +8,7 @@
 #include "EdGraph/EdGraphNode.h"
 #include "IStructureDetailsView.h"
 #include "EdGraphNode_Comment.h"
+#include "Editor.h"
 
 #define LOCTEXT_NAMESPACE "DataflowEditorCommands"
 
@@ -61,26 +62,30 @@ void FDataflowEditorCommands::Unregister()
 
 void FDataflowEditorCommands::EvaluateSelectedNodes(const FGraphPanelSelectionSet& SelectedNodes, FDataflowEditorCommands::FGraphEvaluationCallback Evaluate)
 {
-	for (UObject* Ode : SelectedNodes)
+	const bool bIsInPIEOrSimulate = GEditor->PlayWorld != NULL || GEditor->bIsSimulatingInEditor;
+	if (!bIsInPIEOrSimulate)
 	{
-		if (UDataflowEdNode* EdNode = dynamic_cast<UDataflowEdNode*>(Ode))
+		for (UObject* Ode : SelectedNodes)
 		{
-			if (const TSharedPtr<Dataflow::FGraph> DataflowGraph = EdNode->GetDataflowGraph())
+			if (UDataflowEdNode* EdNode = dynamic_cast<UDataflowEdNode*>(Ode))
 			{
-				if (const TSharedPtr<FDataflowNode> DataflowNode = DataflowGraph->FindBaseNode(EdNode->GetDataflowNodeGuid()))
+				if (const TSharedPtr<Dataflow::FGraph> DataflowGraph = EdNode->GetDataflowGraph())
 				{
-					if (DataflowNode->bActive)
+					if (const TSharedPtr<FDataflowNode> DataflowNode = DataflowGraph->FindBaseNode(EdNode->GetDataflowNodeGuid()))
 					{
-						if (DataflowNode->GetOutputs().Num())
+						if (DataflowNode->bActive)
 						{
-							for (FDataflowConnection* NodeOutput : DataflowNode->GetOutputs())
+							if (DataflowNode->GetOutputs().Num())
 							{
-								Evaluate(DataflowNode.Get(), (FDataflowOutput*)NodeOutput);
+								for (FDataflowConnection* NodeOutput : DataflowNode->GetOutputs())
+								{
+									Evaluate(DataflowNode.Get(), (FDataflowOutput*)NodeOutput);
+								}
+							}
+							else
+							{
+								Evaluate(DataflowNode.Get(), nullptr);
 							}
-						}
-						else
-						{
-							Evaluate(DataflowNode.Get(), nullptr);
 						}
 					}
 				}
@@ -92,35 +97,39 @@ void FDataflowEditorCommands::EvaluateSelectedNodes(const FGraphPanelSelectionSe
 void FDataflowEditorCommands::EvaluateNode(Dataflow::FContext& Context, Dataflow::FTimestamp& OutLastNodeTimestamp,
 	const UDataflow* Dataflow, const FDataflowNode* InNode, const FDataflowOutput* Output, FString NodeName)
 {
-	if (Dataflow)
+	const bool bIsInPIEOrSimulate = GEditor->PlayWorld != NULL || GEditor->bIsSimulatingInEditor;
+	if (!bIsInPIEOrSimulate)
 	{
-		const FDataflowNode* Node = InNode;
-		if (Node == nullptr)
+		if (Dataflow)
 		{
-			if (const TSharedPtr<Dataflow::FGraph> Graph = Dataflow->GetDataflow())
+			const FDataflowNode* Node = InNode;
+			if (Node == nullptr)
 			{
-				if (TSharedPtr<const FDataflowNode> GraphNode = Graph->FindBaseNode(FName(NodeName)))
+				if (const TSharedPtr<Dataflow::FGraph> Graph = Dataflow->GetDataflow())
 				{
-					Node = GraphNode.Get();
+					if (TSharedPtr<const FDataflowNode> GraphNode = Graph->FindBaseNode(FName(NodeName)))
+					{
+						Node = GraphNode.Get();
+					}
 				}
 			}
-		}
 
-		if (Node != nullptr)
-		{
-			if (Output == nullptr)
+			if (Node != nullptr)
 			{
-				if (Node->GetTimestamp() >= OutLastNodeTimestamp)
+				if (Output == nullptr)
 				{
-					Context.Evaluate(Node, nullptr);
-					OutLastNodeTimestamp = Context.GetTimestamp();
+					if (Node->GetTimestamp() >= OutLastNodeTimestamp)
+					{
+						Context.Evaluate(Node, nullptr);
+						OutLastNodeTimestamp = Context.GetTimestamp();
+					}
 				}
-			}
-			else // Output != nullptr
-			{
-				if (!Context.HasData(Output->CacheKey(), Context.GetTimestamp()))
+				else // Output != nullptr
 				{
-					Context.Evaluate(Node, Output);
+					if (!Context.HasData(Output->CacheKey(), Context.GetTimestamp()))
+					{
+						Context.Evaluate(Node, Output);
+					}
 				}
 			}
 		}
@@ -130,43 +139,47 @@ void FDataflowEditorCommands::EvaluateNode(Dataflow::FContext& Context, Dataflow
 void FDataflowEditorCommands::EvaluateTerminalNode(Dataflow::FContext& Context, Dataflow::FTimestamp& OutLastNodeTimestamp,
 	const UDataflow* Dataflow, const FDataflowNode* InNode, const FDataflowOutput* Output, UObject* InAsset, FString NodeName)
 {
-	if (Dataflow)
+	const bool bIsInPIEOrSimulate = GEditor->PlayWorld != NULL || GEditor->bIsSimulatingInEditor;
+	if (!bIsInPIEOrSimulate)
 	{
-		const FDataflowNode* Node = InNode;
-		if (Node == nullptr)
+		if (Dataflow)
 		{
-			if (const TSharedPtr<Dataflow::FGraph> Graph = Dataflow->GetDataflow())
+			const FDataflowNode* Node = InNode;
+			if (Node == nullptr)
 			{
-				if (TSharedPtr<const FDataflowNode> GraphNode = Graph->FindBaseNode(FName(NodeName)))
+				if (const TSharedPtr<Dataflow::FGraph> Graph = Dataflow->GetDataflow())
 				{
-					Node = GraphNode.Get();
+					if (TSharedPtr<const FDataflowNode> GraphNode = Graph->FindBaseNode(FName(NodeName)))
+					{
+						Node = GraphNode.Get();
+					}
 				}
 			}
-		}
 
-		if (Node != nullptr)
-		{
-			if (Output == nullptr)
+			if (Node != nullptr)
 			{
-				if (Node->GetTimestamp() >= OutLastNodeTimestamp)
+				if (Output == nullptr)
 				{
-					Context.Evaluate(Node, nullptr);
-					OutLastNodeTimestamp = Context.GetTimestamp();
-
-					if (const FDataflowTerminalNode* TerminalNode = Node->AsType<const FDataflowTerminalNode>())
+					if (Node->GetTimestamp() >= OutLastNodeTimestamp)
 					{
-						if (InAsset)
+						Context.Evaluate(Node, nullptr);
+						OutLastNodeTimestamp = Context.GetTimestamp();
+
+						if (const FDataflowTerminalNode* TerminalNode = Node->AsType<const FDataflowTerminalNode>())
 						{
-							TerminalNode->SetAssetValue(InAsset, Context);
+							if (InAsset)
+							{
+								TerminalNode->SetAssetValue(InAsset, Context);
+							}
 						}
 					}
 				}
-			}
-			else // Output != nullptr
-			{
-				if (!Context.HasData(Output->CacheKey(), Context.GetTimestamp()))
+				else // Output != nullptr
 				{
-					Context.Evaluate(Node, Output);
+					if (!Context.HasData(Output->CacheKey(), Context.GetTimestamp()))
+					{
+						Context.Evaluate(Node, Output);
+					}
 				}
 			}
 		}
