Dataflow - Fixed a misplaced const on a returned shared pointer of the Dataflow graph and added a new non-const method for continuity.

#rnx

diff --git a/Engine/Plugins/Experimental/Dataflow/Source/DataflowEditor/Private/Dataflow/DataflowEditorCommands.cpp b/Engine/Plugins/Experimental/Dataflow/Source/DataflowEditor/Private/Dataflow/DataflowEditorCommands.cpp
--- a/Engine/Plugins/Experimental/Dataflow/Source/DataflowEditor/Private/Dataflow/DataflowEditorCommands.cpp
+++ b/Engine/Plugins/Experimental/Dataflow/Source/DataflowEditor/Private/Dataflow/DataflowEditorCommands.cpp
@@ -105,7 +105,7 @@ void FDataflowEditorCommands::EvaluateNode(Dataflow::FContext& Context, Dataflow
 			const FDataflowNode* Node = InNode;
 			if (Node == nullptr)
 			{
-				if (const TSharedPtr<Dataflow::FGraph> Graph = Dataflow->GetDataflow())
+				if (const TSharedPtr<const Dataflow::FGraph> Graph = Dataflow->GetDataflow())
 				{
 					if (TSharedPtr<const FDataflowNode> GraphNode = Graph->FindBaseNode(FName(NodeName)))
 					{
@@ -147,7 +147,7 @@ void FDataflowEditorCommands::EvaluateTerminalNode(Dataflow::FContext& Context,
 			const FDataflowNode* Node = InNode;
 			if (Node == nullptr)
 			{
-				if (const TSharedPtr<Dataflow::FGraph> Graph = Dataflow->GetDataflow())
+				if (const TSharedPtr<const Dataflow::FGraph> Graph = Dataflow->GetDataflow())
 				{
 					if (TSharedPtr<const FDataflowNode> GraphNode = Graph->FindBaseNode(FName(NodeName)))
 					{
diff --git a/Engine/Source/Runtime/Experimental/Dataflow/Engine/Public/Dataflow/DataflowObject.h b/Engine/Source/Runtime/Experimental/Dataflow/Engine/Public/Dataflow/DataflowObject.h
--- a/Engine/Source/Runtime/Experimental/Dataflow/Engine/Public/Dataflow/DataflowObject.h
+++ b/Engine/Source/Runtime/Experimental/Dataflow/Engine/Public/Dataflow/DataflowObject.h
@@ -98,7 +98,8 @@ public:
 	void Serialize(FArchive& Ar);
 
 	/** Accessors for internal geometry collection */
-	const TSharedPtr<Dataflow::FGraph, ESPMode::ThreadSafe> GetDataflow() const { return Dataflow; }
+	TSharedPtr<const Dataflow::FGraph, ESPMode::ThreadSafe> GetDataflow() const { return Dataflow; }
+	TSharedPtr<Dataflow::FGraph, ESPMode::ThreadSafe> GetDataflow() { return Dataflow; }
 
 	/**Editing the collection should only be through the edit object.*/
 	FDataflowAssetEdit EditDataflow() const {
