Dataflow - Added AddOptionPin/RemoveOptionPin support.

#rnx

diff --git a/Engine/Plugins/Experimental/Dataflow/Source/DataflowEditor/Private/Dataflow/DataflowEditorCommands.cpp b/Engine/Plugins/Experimental/Dataflow/Source/DataflowEditor/Private/Dataflow/DataflowEditorCommands.cpp
--- a/Engine/Plugins/Experimental/Dataflow/Source/DataflowEditor/Private/Dataflow/DataflowEditorCommands.cpp
+++ b/Engine/Plugins/Experimental/Dataflow/Source/DataflowEditor/Private/Dataflow/DataflowEditorCommands.cpp
@@ -20,6 +20,8 @@ void FDataflowEditorCommandsImpl::RegisterCommands()
 	UI_COMMAND(ToggleObjectSelection, "ToggleObjectSelection", "Enable object selection in editor.", EUserInterfaceActionType::ToggleButton, FInputChord());
 	UI_COMMAND(ToggleFaceSelection, "ToggleFaceSelection", "Enable face selection in editor.", EUserInterfaceActionType::ToggleButton, FInputChord());
 	UI_COMMAND(ToggleVertexSelection, "ToggleVertexSelection", "Enable vertex selection in editor.", EUserInterfaceActionType::ToggleButton, FInputChord());
+	UI_COMMAND(AddOptionPin, "AddOptionPin", "Add an option pin to the selected nodes.", EUserInterfaceActionType::Button, FInputChord());
+	UI_COMMAND(RemoveOptionPin, "RemoveOptionPin", "Remove the last option pin from the selected nodes.", EUserInterfaceActionType::Button, FInputChord());
 
 	if (Dataflow::FNodeFactory* Factory = Dataflow::FNodeFactory::GetInstance())
 	{
diff --git a/Engine/Plugins/Experimental/Dataflow/Source/DataflowEditor/Private/Dataflow/DataflowGraphEditor.cpp b/Engine/Plugins/Experimental/Dataflow/Source/DataflowEditor/Private/Dataflow/DataflowGraphEditor.cpp
--- a/Engine/Plugins/Experimental/Dataflow/Source/DataflowEditor/Private/Dataflow/DataflowGraphEditor.cpp
+++ b/Engine/Plugins/Experimental/Dataflow/Source/DataflowEditor/Private/Dataflow/DataflowGraphEditor.cpp
@@ -85,7 +85,17 @@ void SDataflowGraphEditor::Construct(const FArguments& InArgs, UObject* InAssetO
 			);
 			GraphEditorCommands->MapAction(
 				FDataflowEditorCommands::Get().ToggleEnabledState,
-				FExecuteAction::CreateRaw(this, &SDataflowGraphEditor::ToggleEnabledState)
+				FExecuteAction::CreateSP(this, &SDataflowGraphEditor::ToggleEnabledState)
+			);
+			GraphEditorCommands->MapAction(
+				FDataflowEditorCommands::Get().AddOptionPin,
+				FExecuteAction::CreateSP(this, &SDataflowGraphEditor::OnAddOptionPin),
+				FCanExecuteAction::CreateSP(this, &SDataflowGraphEditor::CanAddOptionPin)
+			);
+			GraphEditorCommands->MapAction(
+				FDataflowEditorCommands::Get().RemoveOptionPin,
+				FExecuteAction::CreateSP(this, &SDataflowGraphEditor::OnRemoveOptionPin),
+				FCanExecuteAction::CreateSP(this, &SDataflowGraphEditor::CanRemoveOptionPin)
 			);
 		}
 	}
@@ -336,4 +346,130 @@ void SDataflowGraphEditor::ToggleEnabledState()
 	FDataflowEditorCommands::ToggleEnabledState(DataflowAsset.Get());
 }
 
+void SDataflowGraphEditor::OnAddOptionPin()
+{
+	UDataflow* const Graph = DataflowAsset.Get();
+	FDataflowAssetEdit Edit = Graph->EditDataflow();
+	if (Dataflow::FGraph* const DataflowGraph = Edit.GetGraph())
+	{
+		const FGraphPanelSelectionSet& SelectedNodes = GetSelectedNodes();
+
+		// Iterate over all nodes, and add the pin
+		for (FGraphPanelSelectionSet::TConstIterator It(SelectedNodes); It; ++It)
+		{
+			UDataflowEdNode* const EdNode = CastChecked<UDataflowEdNode>(*It);
+
+			if (const TSharedPtr<FDataflowNode> Node = DataflowGraph->FindBaseNode(EdNode->DataflowNodeGuid))
+			{
+				if (Node->CanAddPin())
+				{
+					const FScopedTransaction Transaction(LOCTEXT("AddOptionPin", "Add Option Pin"));
+					EdNode->Modify();
+
+					EdNode->AddOptionPin();
+
+					const UDataflowSchema* const Schema = CastChecked<UDataflowSchema>(Graph->GetSchema());
+					Schema->ReconstructNode(*EdNode);
+				}
+			}
+		}
+	}
+}
+
+bool SDataflowGraphEditor::CanAddOptionPin() const
+{
+	bool bCanAddOptionPin = false;
+
+	const UDataflow* const Graph = DataflowAsset.Get();
+	if (const Dataflow::FGraph* const DataflowGraph = Graph->GetDataflow().Get())
+	{
+		const FGraphPanelSelectionSet& SelectedNodes = GetSelectedNodes();
+
+		// Iterate over all nodes, and add the pin
+		for (FGraphPanelSelectionSet::TConstIterator It(SelectedNodes); It; ++It)
+		{
+			const UDataflowEdNode* const EdNode = CastChecked<UDataflowEdNode>(*It);
+
+			if (const TSharedPtr<const FDataflowNode> Node = DataflowGraph->FindBaseNode(EdNode->DataflowNodeGuid))
+			{
+				bCanAddOptionPin = Node->CanAddPin();
+			}
+			else
+			{
+				bCanAddOptionPin = false;
+			}
+
+			if (!bCanAddOptionPin)
+			{
+				break;  // One bad node is good enough to return false
+			}
+		}
+	}
+
+	return bCanAddOptionPin;
+}
+
+void SDataflowGraphEditor::OnRemoveOptionPin()
+{
+	UDataflow* const Graph = DataflowAsset.Get();
+	FDataflowAssetEdit Edit = Graph->EditDataflow();
+	if (Dataflow::FGraph* const DataflowGraph = Edit.GetGraph())
+	{
+		const FGraphPanelSelectionSet& SelectedNodes = GetSelectedNodes();
+
+		// Iterate over all nodes, and remove a pin
+		for (FGraphPanelSelectionSet::TConstIterator It(SelectedNodes); It; ++It)
+		{
+			UDataflowEdNode* const EdNode = CastChecked<UDataflowEdNode>(*It);
+
+			if (const TSharedPtr<FDataflowNode> Node = DataflowGraph->FindBaseNode(EdNode->DataflowNodeGuid))
+			{
+				if (Node->CanRemovePin())
+				{
+					const FScopedTransaction Transaction(LOCTEXT("RemoveOptionPin", "Remove Option Pin"));
+					EdNode->Modify();
+
+					EdNode->RemoveOptionPin();
+
+					const UDataflowSchema* const Schema = CastChecked<UDataflowSchema>(Graph->GetSchema());
+					Schema->ReconstructNode(*EdNode);
+				}
+			}
+		}
+	}
+}
+
+bool SDataflowGraphEditor::CanRemoveOptionPin() const
+{
+	bool bCanRemoveOptionPin = false;
+
+	const UDataflow* const Graph = DataflowAsset.Get();
+	if (const Dataflow::FGraph* const DataflowGraph = Graph->GetDataflow().Get())
+	{
+		const FGraphPanelSelectionSet& SelectedNodes = GetSelectedNodes();
+
+		// Iterate over all nodes, and add the pin
+		for (FGraphPanelSelectionSet::TConstIterator It(SelectedNodes); It; ++It)
+		{
+			const UDataflowEdNode* const EdNode = CastChecked<UDataflowEdNode>(*It);
+
+			if (const TSharedPtr<const FDataflowNode> Node = DataflowGraph->FindBaseNode(EdNode->DataflowNodeGuid))
+			{
+				bCanRemoveOptionPin = Node->CanRemovePin();
+			}
+			else
+			{
+				bCanRemoveOptionPin = false;
+			}
+
+			if (!bCanRemoveOptionPin)
+			{
+				break;  // One bad node is good enough to return false
+			}
+		}
+	}
+
+	return bCanRemoveOptionPin;
+}
+
 #undef LOCTEXT_NAMESPACE
diff --git a/Engine/Plugins/Experimental/Dataflow/Source/DataflowEditor/Private/Dataflow/DataflowSchema.cpp b/Engine/Plugins/Experimental/Dataflow/Source/DataflowEditor/Private/Dataflow/DataflowSchema.cpp
--- a/Engine/Plugins/Experimental/Dataflow/Source/DataflowEditor/Private/Dataflow/DataflowSchema.cpp
+++ b/Engine/Plugins/Experimental/Dataflow/Source/DataflowEditor/Private/Dataflow/DataflowSchema.cpp
@@ -38,6 +38,8 @@ void UDataflowSchema::GetContextMenuActions(class UToolMenu* Menu, class UGraphN
 				Section.AddMenuEntry(FGenericCommands::Get().Duplicate);
 				Section.AddMenuEntry(FDataflowEditorCommands::Get().ToggleEnabledState, FText::FromString("Toggle Enabled State"));
 				Section.AddMenuEntry(FGraphEditorCommands::Get().BreakNodeLinks);
+				Section.AddMenuEntry(FDataflowEditorCommands::Get().AddOptionPin);
+				Section.AddMenuEntry(FDataflowEditorCommands::Get().RemoveOptionPin);
 				Section.AddMenuEntry(FDataflowEditorCommands::Get().EvaluateNode);
 			}
 		}
diff --git a/Engine/Plugins/Experimental/Dataflow/Source/DataflowEditor/Public/Dataflow/DataflowEditorCommands.h b/Engine/Plugins/Experimental/Dataflow/Source/DataflowEditor/Public/Dataflow/DataflowEditorCommands.h
--- a/Engine/Plugins/Experimental/Dataflow/Source/DataflowEditor/Public/Dataflow/DataflowEditorCommands.h
+++ b/Engine/Plugins/Experimental/Dataflow/Source/DataflowEditor/Public/Dataflow/DataflowEditorCommands.h
@@ -43,6 +43,8 @@ public:
 	TSharedPtr<FUICommandInfo> ToggleObjectSelection;
 	TSharedPtr<FUICommandInfo> ToggleFaceSelection;
 	TSharedPtr<FUICommandInfo> ToggleVertexSelection;
+	TSharedPtr< FUICommandInfo > AddOptionPin;
+	TSharedPtr< FUICommandInfo > RemoveOptionPin;
 
 	TMap< FName, TSharedPtr<FUICommandInfo> > CreateNodesMap;
 };
diff --git a/Engine/Plugins/Experimental/Dataflow/Source/DataflowEditor/Public/Dataflow/DataflowGraphEditor.h b/Engine/Plugins/Experimental/Dataflow/Source/DataflowEditor/Public/Dataflow/DataflowGraphEditor.h
--- a/Engine/Plugins/Experimental/Dataflow/Source/DataflowEditor/Public/Dataflow/DataflowGraphEditor.h
+++ b/Engine/Plugins/Experimental/Dataflow/Source/DataflowEditor/Public/Dataflow/DataflowGraphEditor.h
@@ -117,6 +117,16 @@ public:
 	SGraphEditor* GetGraphEditor() { return (SGraphEditor*)this; }
 
 private:
+	/** Add an additional option pin to all selected Dataflow nodes for those that overrides the AddPin function. */
+	void OnAddOptionPin();
+	/** Return whether all currently selected Dataflow nodes can execute the AddPin function. */
+	bool CanAddOptionPin() const;
+
+	/** Remove an option pin from all selected Dataflow nodes for those that overrides the RemovePin function. */
+	void OnRemoveOptionPin();
+	/** Return whether all currently selected Dataflow nodes can execute the RemovePin function. */
+	bool CanRemoveOptionPin() const;
+
 	FDataflowEditorCommands::FOnDragDropEventCallback OnDragDropEventCallback;
 	FDataflowEditorCommands::FGraphEvaluationCallback EvaluateGraphCallback;
 
diff --git a/Engine/Source/Runtime/Experimental/Dataflow/Core/Private/Dataflow/DataflowNode.cpp b/Engine/Source/Runtime/Experimental/Dataflow/Core/Private/Dataflow/DataflowNode.cpp
--- a/Engine/Source/Runtime/Experimental/Dataflow/Core/Private/Dataflow/DataflowNode.cpp
+++ b/Engine/Source/Runtime/Experimental/Dataflow/Core/Private/Dataflow/DataflowNode.cpp
@@ -252,6 +252,34 @@ void FDataflowNode::RegisterInputConnection(const void* InProperty)
 	}
 }
 
+void FDataflowNode::UnregisterInputConnection(const void* InProperty)
+{
+	if (TUniquePtr<FStructOnScope> ScriptOnStruct = TUniquePtr<FStructOnScope>(NewStructOnScope()))
+	{
+		if (const UStruct* const Struct = ScriptOnStruct->GetStruct())
+		{
+			for (TFieldIterator<FProperty> PropertyIt(Struct); PropertyIt; ++PropertyIt)
+			{
+				const FProperty* const Property = *PropertyIt;
+				const int32 Offset = Property->GetOffset_ForInternal();
+				const size_t RealAddress = (size_t)this + (size_t)Offset;
+				if (RealAddress == (size_t)InProperty)
+				{
+					if (FDataflowInput* const* const Input = Inputs.Find(Offset))
+					{
+						Inputs.Remove(Offset);
+						delete *Input;
+
+						// Invalidate graph as this input might have had connections
+						Invalidate();
+					}
+					break;
+				}
+			}
+		}
+	}
+}
+
 void FDataflowNode::RegisterOutputConnection(const void* InProperty, const void* Passthrough)
 {
 	if (TUniquePtr<FStructOnScope> ScriptOnStruct = TUniquePtr<FStructOnScope>(NewStructOnScope()))
diff --git a/Engine/Source/Runtime/Experimental/Dataflow/Core/Public/Dataflow/DataflowNode.h b/Engine/Source/Runtime/Experimental/Dataflow/Core/Public/Dataflow/DataflowNode.h
--- a/Engine/Source/Runtime/Experimental/Dataflow/Core/Public/Dataflow/DataflowNode.h
+++ b/Engine/Source/Runtime/Experimental/Dataflow/Core/Public/Dataflow/DataflowNode.h
@@ -86,6 +86,14 @@ struct DATAFLOWCORE_API FDataflowNode
 
 	TArray<Dataflow::FPin> GetPins() const;
 
+	/** Override this function to add the AddOptionPin functionality to the node's context menu. */
+	virtual Dataflow::FPin AddPin() { return { Dataflow::FPin::EDirection::NONE, NAME_None, NAME_None }; }
+	/** Override this function to add the AddOptionPin functionality to the node's context menu. */
+	virtual bool CanAddPin() const { return false; }
+	/** Override this function to add the RemoveOPtionPin functionality to the node's context menu. */
+	virtual Dataflow::FPin RemovePin() { return { Dataflow::FPin::EDirection::NONE, NAME_None, NAME_None }; }
+	/** Override this function to add the RemoveOPtionPin functionality to the node's context menu. */
+	virtual bool CanRemovePin() const { return false; }
 
 	virtual void AddInput(FDataflowInput* InPtr);
 	TArray< FDataflowInput* > GetInputs() const;
@@ -128,6 +136,8 @@ struct DATAFLOWCORE_API FDataflowNode
 	/** Register the Input and Outputs after the creation in the factory */
 	void RegisterInputConnection(const void* Property);
 	void RegisterOutputConnection(const void* Property, const void* Passthrough = nullptr);
+	/** Unregister the input connection if one exists matching this property, and then invalidate the graph. */
+	void UnregisterInputConnection(const void* Property);
 
 	//
 	// Evaluation
diff --git a/Engine/Source/Runtime/Experimental/Dataflow/Engine/Private/Dataflow/DataflowEdNode.cpp b/Engine/Source/Runtime/Experimental/Dataflow/Engine/Private/Dataflow/DataflowEdNode.cpp
--- a/Engine/Source/Runtime/Experimental/Dataflow/Engine/Private/Dataflow/DataflowEdNode.cpp
+++ b/Engine/Source/Runtime/Experimental/Dataflow/Engine/Private/Dataflow/DataflowEdNode.cpp
@@ -94,6 +94,81 @@ void UDataflowEdNode::AllocateDefaultPins()
 #endif // WITH_EDITOR && !UE_BUILD_SHIPPING
 }
 
+void UDataflowEdNode::AddOptionPin()
+{
+#if WITH_EDITOR && !(UE_BUILD_SHIPPING || UE_BUILD_TEST)
+	if (DataflowGraph && DataflowNodeGuid.IsValid())
+	{
+		// Modify();  // TODO: How do we modify a DataflowNode
+
+		if (const TSharedPtr<FDataflowNode> DataflowNode = DataflowGraph->FindBaseNode(DataflowNodeGuid))
+		{
+			const Dataflow::FPin Pin = DataflowNode->AddPin();
+			switch (Pin.Direction)
+			{
+			case Dataflow::FPin::EDirection::INPUT:
+				CreatePin(EEdGraphPinDirection::EGPD_Input, Pin.Type, Pin.Name);
+				ReconstructNode();
+				break;
+			case Dataflow::FPin::EDirection::OUTPUT:
+				CreatePin(EEdGraphPinDirection::EGPD_Output, Pin.Type, Pin.Name);
+				ReconstructNode();
+				break;
+			default:
+				break;  // Add pin isn't implemented on this node
+			}
+		}
+
+		// Refresh the current graph, so the pins can be updated
+		if (UEdGraph* const ParentGraph = GetGraph())
+		{
+			ParentGraph->NotifyGraphChanged();
+		}
+	}
+#endif // WITH_EDITOR && !UE_BUILD_SHIPPING
+}
+
+void UDataflowEdNode::RemoveOptionPin()
+{
+#if WITH_EDITOR && !(UE_BUILD_SHIPPING || UE_BUILD_TEST)
+	if (DataflowGraph && DataflowNodeGuid.IsValid())
+	{
+		// Modify();  // TODO: How do we modify a DataflowNode
+
+		if (const TSharedPtr<FDataflowNode> DataflowNode = DataflowGraph->FindBaseNode(DataflowNodeGuid))
+		{
+			const Dataflow::FPin Pin = DataflowNode->RemovePin();
+			switch (Pin.Direction)
+			{
+			case Dataflow::FPin::EDirection::INPUT:
+				if (UEdGraphPin* const EdPin = FindPin(Pin.Name, EEdGraphPinDirection::EGPD_Input))
+				{
+					EdPin->BreakAllPinLinks();
+					RemovePin(EdPin);
+					ReconstructNode();
+				}
+				break;
+			case Dataflow::FPin::EDirection::OUTPUT:
+				if (UEdGraphPin* const EdPin = FindPin(Pin.Name, EEdGraphPinDirection::EGPD_Output))
+				{
+					EdPin->BreakAllPinLinks();
+					RemovePin(EdPin);
+					ReconstructNode();
+				}
+				break;
+			default:
+				break;  // Add pin isn't implemented on this node
+			}
+		}
+
+		// Refresh the current graph, so the pins can be updated
+		if (UEdGraph* const ParentGraph = GetGraph())
+		{
+			ParentGraph->NotifyGraphChanged();
+		}
+	}
+#endif // WITH_EDITOR && !UE_BUILD_SHIPPING
+}
 
 FText UDataflowEdNode::GetNodeTitle(ENodeTitleType::Type TitleType) const
 {
diff --git a/Engine/Source/Runtime/Experimental/Dataflow/Engine/Public/Dataflow/DataflowEdNode.h b/Engine/Source/Runtime/Experimental/Dataflow/Engine/Public/Dataflow/DataflowEdNode.h
--- a/Engine/Source/Runtime/Experimental/Dataflow/Engine/Public/Dataflow/DataflowEdNode.h
+++ b/Engine/Source/Runtime/Experimental/Dataflow/Engine/Public/Dataflow/DataflowEdNode.h
@@ -63,6 +63,11 @@ public:
 	TSharedPtr<FDataflowNode> GetDataflowNode();
 	TSharedPtr<const FDataflowNode> GetDataflowNode() const;
 
+	/** Add a new option pin if the underlying Dataflow node AddPin member is overriden. */
+	void AddOptionPin();
+	/** Remove an option pin if the underlying Dataflow node RemovePin member is overriden. */
+	void RemoveOptionPin();
+
 	//
 	// Node Rendering
 	//
