CIS - Fixed non-unity build error C2079: 'Transaction' uses undefined class 'FScopedTransaction'.

#rnx

diff --git a/Engine/Plugins/Experimental/Dataflow/Source/DataflowEditor/Private/Dataflow/DataflowGraphEditor.cpp b/Engine/Plugins/Experimental/Dataflow/Source/DataflowEditor/Private/Dataflow/DataflowGraphEditor.cpp
--- a/Engine/Plugins/Experimental/Dataflow/Source/DataflowEditor/Private/Dataflow/DataflowGraphEditor.cpp
+++ b/Engine/Plugins/Experimental/Dataflow/Source/DataflowEditor/Private/Dataflow/DataflowGraphEditor.cpp
@@ -13,6 +13,7 @@
 #include "Dataflow/DataflowSelectionNodes.h"
 #include "Dataflow/DataflowSchema.h"
 #include "HAL/PlatformApplicationMisc.h"
+#include "ScopedTransaction.h"
 #include "IStructureDetailsView.h"
 
 #define LOCTEXT_NAMESPACE "DataflowGraphEditor"
