dataflow : remove ensure in dataflow spreadsheet widget when a type is unknown

diff --git a/Engine/Plugins/Experimental/Dataflow/Source/DataflowEditor/Private/Dataflow/CollectionSpreadSheetWidget.cpp b/Engine/Plugins/Experimental/Dataflow/Source/DataflowEditor/Private/Dataflow/CollectionSpreadSheetWidget.cpp
--- a/Engine/Plugins/Experimental/Dataflow/Source/DataflowEditor/Private/Dataflow/CollectionSpreadSheetWidget.cpp
+++ b/Engine/Plugins/Experimental/Dataflow/Source/DataflowEditor/Private/Dataflow/CollectionSpreadSheetWidget.cpp
@@ -187,7 +187,7 @@ namespace CollectionSpreadSheetHelpers
 			break;
 
 		default:
-			ensure(false);
+			//ensure(false);
 			ValueAsString = "<Unknown Data Type>";
 		}
 
