Dataflow : Fix collection issue in dataflow editor.

diff --git a/Engine/Plugins/Experimental/Dataflow/Source/DataflowEditor/Private/Dataflow/DataflowSNode.cpp b/Engine/Plugins/Experimental/Dataflow/Source/DataflowEditor/Private/Dataflow/DataflowSNode.cpp
--- a/Engine/Plugins/Experimental/Dataflow/Source/DataflowEditor/Private/Dataflow/DataflowSNode.cpp
+++ b/Engine/Plugins/Experimental/Dataflow/Source/DataflowEditor/Private/Dataflow/DataflowSNode.cpp
@@ -4,6 +4,7 @@
 
 #include "Dataflow/DataflowEdNode.h"
 #include "Dataflow/DataflowEditorStyle.h"
+#include "Dataflow/DataflowEngineUtil.h"
 #include "Dataflow/DataflowGraphEditor.h"
 #include "Dataflow/DataflowNodeFactory.h"
 #include "Dataflow/DataflowObject.h"
@@ -151,6 +152,16 @@ FReply SDataflowEdNode::OnMouseButtonDoubleClick(const FGeometry& InMyGeometry,
 	return Super::OnMouseButtonDoubleClick(InMyGeometry, InMouseEvent);
 }
 
+void SDataflowEdNode::AddReferencedObjects(FReferenceCollector& Collector)
+{
+	if (DataflowGraphNode)
+	{
+		if (TSharedPtr<FDataflowNode> DataflowNode = DataflowGraphNode->GetDataflowNode())
+		{
+			Collector.AddPropertyReferences(DataflowNode->TypedScriptStruct(), DataflowNode.Get());
+		}
+	}
+}
 
 
 //
@@ -233,5 +244,8 @@ UEdGraphNode* FAssetSchemaAction_Dataflow_CreateNode_DataflowEdNode::PerformActi
 //	Collector.AddReferencedObject(NodeTemplate);
 //}
 
+
+
+
 #undef LOCTEXT_NAMESPACE
 
diff --git a/Engine/Plugins/Experimental/Dataflow/Source/DataflowEditor/Public/Dataflow/DataflowSNode.h b/Engine/Plugins/Experimental/Dataflow/Source/DataflowEditor/Public/Dataflow/DataflowSNode.h
--- a/Engine/Plugins/Experimental/Dataflow/Source/DataflowEditor/Public/Dataflow/DataflowSNode.h
+++ b/Engine/Plugins/Experimental/Dataflow/Source/DataflowEditor/Public/Dataflow/DataflowSNode.h
@@ -4,6 +4,7 @@
 
 #include "EdGraphUtilities.h"
 #include "SGraphNode.h"
+#include "UObject/GCObject.h"
 
 #include "DataflowSNode.generated.h"
 
@@ -14,7 +15,7 @@ class SCheckBox;
 // SDataflowEdNode
 //
 
-class DATAFLOWEDITOR_API SDataflowEdNode : public SGraphNode
+class DATAFLOWEDITOR_API SDataflowEdNode : public SGraphNode , public FGCObject
 {
 	typedef SGraphNode Super;
 
@@ -34,6 +35,16 @@ public:
 	virtual FReply OnMouseButtonDoubleClick(const FGeometry& InMyGeometry, const FPointerEvent& InMouseEvent) override;
 	TArray<FOverlayWidgetInfo> GetOverlayWidgets(bool bSelected, const FVector2D& WidgetSize) const;
 
+
+
+	//~ Begin FGCObject interface
+	virtual void AddReferencedObjects(FReferenceCollector& Collector) override;
+	virtual FString GetReferencerName() const override
+	{
+		return TEXT("SDataflowEdNode");
+	}
+	//~ End FGCObject interface
+
 private:
 	UDataflowEdNode* DataflowGraphNode = nullptr;	
 
