Dataflow:
- Fixed crash when renaming a node with a name from a recently deleted node.
- Fixed typo in Node (Ode) variable name.

#rnx

diff --git a/Engine/Plugins/Experimental/Dataflow/Source/DataflowEditor/Private/Dataflow/DataflowEditorCommands.cpp b/Engine/Plugins/Experimental/Dataflow/Source/DataflowEditor/Private/Dataflow/DataflowEditorCommands.cpp
--- a/Engine/Plugins/Experimental/Dataflow/Source/DataflowEditor/Private/Dataflow/DataflowEditorCommands.cpp
+++ b/Engine/Plugins/Experimental/Dataflow/Source/DataflowEditor/Private/Dataflow/DataflowEditorCommands.cpp
@@ -68,9 +68,9 @@ void FDataflowEditorCommands::EvaluateSelectedNodes(const FGraphPanelSelectionSe
 	const bool bIsInPIEOrSimulate = GEditor->PlayWorld != NULL || GEditor->bIsSimulatingInEditor;
 	if (!bIsInPIEOrSimulate)
 	{
-		for (UObject* Ode : SelectedNodes)
+		for (UObject* Node : SelectedNodes)
 		{
-			if (UDataflowEdNode* EdNode = dynamic_cast<UDataflowEdNode*>(Ode))
+			if (UDataflowEdNode* EdNode = dynamic_cast<UDataflowEdNode*>(Node))
 			{
 				if (const TSharedPtr<Dataflow::FGraph> DataflowGraph = EdNode->GetDataflowGraph())
 				{
@@ -304,9 +304,9 @@ void FDataflowEditorCommands::DeleteNodes(UDataflow* Graph, const FGraphPanelSel
 {
 	if (ensureMsgf(Graph != nullptr, TEXT("Warning : Failed to find valid graph.")))
 	{
-		for (UObject* Ode : SelectedNodes)
+		for (UObject* Node : SelectedNodes)
 		{
-			if (UDataflowEdNode* EdNode = dynamic_cast<UDataflowEdNode*>(Ode))
+			if (UDataflowEdNode* EdNode = dynamic_cast<UDataflowEdNode*>(Node))
 			{
 				if (const TSharedPtr<Dataflow::FGraph> DataflowGraph = EdNode->GetDataflowGraph())
 				{
@@ -317,10 +317,13 @@ void FDataflowEditorCommands::DeleteNodes(UDataflow* Graph, const FGraphPanelSel
 					}
 				}
 			}
-			else if (UEdGraphNode_Comment* CommentNode = dynamic_cast<UEdGraphNode_Comment*>(Ode))
+			else if (UEdGraphNode_Comment* CommentNode = dynamic_cast<UEdGraphNode_Comment*>(Node))
 			{
 				Graph->RemoveNode(CommentNode);
 			}
+
+			// Auto-rename node so that its current name is made available until it is garbage collected
+			Node->Rename();
 		}
 	}
 }
