Dataflow:
- implemented Duplicate for nodes and comments
- bug fixes

diff --git a/Engine/Plugins/Experimental/Dataflow/Source/DataflowEditor/Private/Dataflow/DataflowEditorCommands.cpp b/Engine/Plugins/Experimental/Dataflow/Source/DataflowEditor/Private/Dataflow/DataflowEditorCommands.cpp
--- a/Engine/Plugins/Experimental/Dataflow/Source/DataflowEditor/Private/Dataflow/DataflowEditorCommands.cpp
+++ b/Engine/Plugins/Experimental/Dataflow/Source/DataflowEditor/Private/Dataflow/DataflowEditorCommands.cpp
@@ -10,6 +10,9 @@
 #include "IStructureDetailsView.h"
 #include "EdGraphNode_Comment.h"
 #include "Editor.h"
+#include "Dataflow/DataflowGraphEditor.h"
+#include "Dataflow/DataflowSCommentNode.h"
+#include "Dataflow/DataflowSNode.h"
 
 #define LOCTEXT_NAMESPACE "DataflowEditorCommands"
 
@@ -361,4 +364,181 @@ void FDataflowEditorCommands::ToggleEnabledState(UDataflow* Graph)
 {
 }
 
+static UEdGraphPin* GetPin(const UDataflowEdNode* Node, const EEdGraphPinDirection Direction, const FName Name)
+{
+	for (UEdGraphPin* Pin : Node->GetAllPins())
+	{
+		if (Pin->PinName == Name && Pin->Direction == Direction)
+		{
+			return Pin;
+		}
+	}
+
+	return nullptr;
+}
+
+void FDataflowEditorCommands::DuplicateNodes(UDataflow* Graph, const TSharedPtr<SDataflowGraphEditor>& DataflowGraphEditor, const FGraphPanelSelectionSet& SelectedNodes)
+{
+	if (ensureMsgf(Graph != nullptr, TEXT("Warning : Failed to find valid graph.")))
+	{
+		// Separate selected nodes into an array of UDataflowEdNodes and
+		// an array of UEdGraphNode_Comments
+		TSet<UDataflowEdNode*> SelectedEdNodes;
+		TSet<UEdGraphNode_Comment*> SelectedEdCommentNodes;
+
+		for (UObject* Node : SelectedNodes)
+		{
+			if (UDataflowEdNode* EdNode = dynamic_cast<UDataflowEdNode*>(Node))
+			{
+				SelectedEdNodes.Add(EdNode);
+			}
+			else if (UEdGraphNode_Comment* EdCommentNode = dynamic_cast<UEdGraphNode_Comment*>(Node))
+			{
+				SelectedEdCommentNodes.Add(EdCommentNode);
+			}
+		}
+
+		FDataflowAssetEdit Edit = Graph->EditDataflow();
+		if (Dataflow::FGraph* DataflowGraph = Edit.GetGraph())
+		{
+			// Process Dataflow nodes
+			TSet<UDataflowEdNode*> DuplicatedEdNodes;
+			FVector2D RefLocation;
+
+			if (SelectedEdNodes.Num() > 0)
+			{
+				// Store the location of the first selected node for recreating spatial relationships
+				RefLocation.X = SelectedEdNodes.Array()[0]->NodePosX; RefLocation.Y = SelectedEdNodes.Array()[0]->NodePosY;
+
+				//
+				// Create a map to record OriginalNode.Guid -> DuplicatedNode.Guid
+				//
+				TMap<FGuid, FGuid> NodeGuidMap;				// [OriginalDataflowNodeGuid -> DuplicatedDataflowNodeGuid]
+				TMap<FGuid, UDataflowEdNode*> EdNodeMap;	// [OriginalDataflowNodeGuid -> DuplicatedEdNode]
+
+				// Duplicate selected nodes
+				for (UDataflowEdNode* SelectedEdNode : SelectedEdNodes)
+				{
+					if (TSharedPtr<FDataflowNode> DataflowNode = DataflowGraph->FindBaseNode(SelectedEdNode->DataflowNodeGuid))
+					{
+						if (TSharedPtr<FAssetSchemaAction_Dataflow_DuplicateNode_DataflowEdNode> DuplicateNodeAction = FAssetSchemaAction_Dataflow_DuplicateNode_DataflowEdNode::CreateAction(Graph, DataflowNode->GetType()))
+						{
+							DuplicateNodeAction->DataflowNodeToDuplicate = DataflowNode;
+
+							FVector2D SelectedEdNodeLocation(SelectedEdNode->NodePosX, SelectedEdNode->NodePosY);
+							FVector2D DeltaLocation = SelectedEdNodeLocation - RefLocation;
+
+							if (UDataflowEdNode* NewEdNode = (UDataflowEdNode*)DuplicateNodeAction->PerformAction(Graph, nullptr, DataflowGraphEditor->GetPasteLocation() + DeltaLocation, false))
+							{
+								DuplicatedEdNodes.Add(NewEdNode);
+
+								// Record Guids
+								NodeGuidMap.Add(SelectedEdNode->DataflowNodeGuid, NewEdNode->DataflowNodeGuid);
+								EdNodeMap.Add(SelectedEdNode->DataflowNodeGuid, NewEdNode);
+							}
+						}
+					}
+				}
+
+				// Recreate connections between duplicated nodes
+				for (UDataflowEdNode* SelectedEdNode : SelectedEdNodes)
+				{
+					if (TSharedPtr<FDataflowNode> DataflowNode = DataflowGraph->FindBaseNode(SelectedEdNode->DataflowNodeGuid))
+					{
+						FGuid DataflowNodeAGuid = DataflowNode->GetGuid();
+
+						for (FDataflowOutput* Output : DataflowNode->GetOutputs())
+						{
+							for (FDataflowInput* Connection : Output->Connections)
+							{
+								const FName OutputputName = Connection->GetConnection()->GetName();
+
+								// Check if the node on the end of the conmnection was duplicated
+								FGuid DataflowNodeBGuid = Connection->GetOwningNode()->GetGuid();
+
+								if (NodeGuidMap.Contains(DataflowNodeBGuid))
+								{
+									const FName InputputName = Connection->GetName();
+
+									if (TSharedPtr<FDataflowNode> DuplicatedDataflowNodeA = DataflowGraph->FindBaseNode(NodeGuidMap[DataflowNodeAGuid]))
+									{
+										FDataflowOutput* OutputConnection = DuplicatedDataflowNodeA->FindOutput(OutputputName);
+
+										if (TSharedPtr<FDataflowNode> DuplicatedDataflowNodeB = DataflowGraph->FindBaseNode(NodeGuidMap[DataflowNodeBGuid]))
+										{
+											FDataflowInput* InputConnection = DuplicatedDataflowNodeB->FindInput(InputputName);
+
+											DataflowGraph->Connect(OutputConnection, InputConnection);
+
+											// Connect the UDataflowEdNode FPins as well
+											if (UEdGraphPin* OutputPin = GetPin(EdNodeMap[DataflowNodeAGuid], EEdGraphPinDirection::EGPD_Output, OutputputName))
+											{
+												if (UEdGraphPin* InputPin = GetPin(EdNodeMap[DataflowNodeBGuid], EEdGraphPinDirection::EGPD_Input, InputputName))
+												{
+													OutputPin->MakeLinkTo(InputPin);
+												}
+											}
+										}
+									}
+								}
+							}
+						}
+					}
+				}
+			}
+
+			// Process Comment nodes
+			TSet<UEdGraphNode*> DuplicatedEdCommentNodes;
+
+			if (SelectedEdCommentNodes.Num() > 0)
+			{
+				if (SelectedEdNodes.Num() == 0)
+				{
+					RefLocation.X = SelectedEdCommentNodes.Array()[0]->NodePosX; RefLocation.Y = SelectedEdCommentNodes.Array()[0]->NodePosY;
+				}
+
+				const TSharedPtr<SGraphEditor>& InGraphEditor = (TSharedPtr<SGraphEditor>)DataflowGraphEditor;
+
+				for (UEdGraphNode_Comment* SelectedCommentNode : SelectedEdCommentNodes)
+				{
+					if (TSharedPtr<FAssetSchemaAction_Dataflow_DuplicateCommentNode_DataflowEdNode> DuplicateCommentNodeAction = FAssetSchemaAction_Dataflow_DuplicateCommentNode_DataflowEdNode::CreateAction(Graph, InGraphEditor))
+					{
+						DuplicateCommentNodeAction->CommentNodeToDuplicate = SelectedCommentNode;
+
+						FVector2D SelectedCommentNodeLocation(SelectedCommentNode->NodePosX, SelectedCommentNode->NodePosY);
+						FVector2D DeltaLocation = SelectedCommentNodeLocation - RefLocation;
+
+						if (UEdGraphNode* NewCommentNode = DuplicateCommentNodeAction->PerformAction(Graph, nullptr, DataflowGraphEditor->GetPasteLocation() + DeltaLocation, false))
+						{
+							DuplicatedEdCommentNodes.Add(NewCommentNode);
+						}
+					}
+				}
+			}
+
+			// Update the selection in the Editor
+			if (SelectedEdNodes.Num() > 0 || SelectedEdCommentNodes.Num() > 0)
+			{
+				DataflowGraphEditor->ClearSelectionSet();
+
+				if (SelectedEdNodes.Num() > 0)
+				{
+					for (UDataflowEdNode* Node : DuplicatedEdNodes)
+					{
+						DataflowGraphEditor->SetNodeSelection(Node, true);
+					}
+				}
+
+				if (SelectedEdCommentNodes.Num() > 0)
+				{
+					for (UEdGraphNode* Node : DuplicatedEdCommentNodes)
+					{
+						DataflowGraphEditor->SetNodeSelection(Node, true);
+					}
+				}
+			}
+		}
+	}
+}
+
 #undef LOCTEXT_NAMESPACE
diff --git a/Engine/Plugins/Experimental/Dataflow/Source/DataflowEditor/Private/Dataflow/DataflowGraphEditor.cpp b/Engine/Plugins/Experimental/Dataflow/Source/DataflowEditor/Private/Dataflow/DataflowGraphEditor.cpp
--- a/Engine/Plugins/Experimental/Dataflow/Source/DataflowEditor/Private/Dataflow/DataflowGraphEditor.cpp
+++ b/Engine/Plugins/Experimental/Dataflow/Source/DataflowEditor/Private/Dataflow/DataflowGraphEditor.cpp
@@ -98,6 +98,10 @@ void SDataflowGraphEditor::Construct(const FArguments& InArgs, UObject* InAssetO
 				FExecuteAction::CreateSP(this, &SDataflowGraphEditor::OnRemoveOptionPin),
 				FCanExecuteAction::CreateSP(this, &SDataflowGraphEditor::CanRemoveOptionPin)
 			);
+			GraphEditorCommands->MapAction(
+				FGenericCommands::Get().Duplicate,
+				FExecuteAction::CreateSP(this, &SDataflowGraphEditor::DuplicateSelectedNodes)
+			);
 		}
 	}
 
@@ -473,4 +477,18 @@ bool SDataflowGraphEditor::CanRemoveOptionPin() const
 	return bCanRemoveOptionPin;
 }
 
+void SDataflowGraphEditor::DuplicateSelectedNodes()
+{
+	if (UDataflow* Graph = DataflowAsset.Get())
+	{
+		const TSharedPtr<SDataflowGraphEditor>& DataflowGraphEditor = SharedThis(this);
+		const FGraphPanelSelectionSet& SelectedNodes = GetSelectedNodes();
+
+		if (SelectedNodes.Num() > 0)
+		{
+			FDataflowEditorCommands::DuplicateNodes(Graph, DataflowGraphEditor, SelectedNodes);
+		}
+	}
+}
+
 #undef LOCTEXT_NAMESPACE
diff --git a/Engine/Plugins/Experimental/Dataflow/Source/DataflowEditor/Private/Dataflow/DataflowSCommentNode.cpp b/Engine/Plugins/Experimental/Dataflow/Source/DataflowEditor/Private/Dataflow/DataflowSCommentNode.cpp
--- a/Engine/Plugins/Experimental/Dataflow/Source/DataflowEditor/Private/Dataflow/DataflowSCommentNode.cpp
+++ b/Engine/Plugins/Experimental/Dataflow/Source/DataflowEditor/Private/Dataflow/DataflowSCommentNode.cpp
@@ -72,4 +72,51 @@ UEdGraphNode* FAssetSchemaAction_Dataflow_CreateCommentNode_DataflowEdNode::Perf
 //	Collector.AddReferencedObject(NodeTemplate);
 //}
 
+TSharedPtr<FAssetSchemaAction_Dataflow_DuplicateCommentNode_DataflowEdNode> FAssetSchemaAction_Dataflow_DuplicateCommentNode_DataflowEdNode::CreateAction(UEdGraph* ParentGraph, const TSharedPtr<SGraphEditor>& GraphEditor)
+{
+	return MakeShared<FAssetSchemaAction_Dataflow_DuplicateCommentNode_DataflowEdNode>(GraphEditor);
+}
+
+//
+//  Duplicate comment node
+//
+UEdGraphNode* FAssetSchemaAction_Dataflow_DuplicateCommentNode_DataflowEdNode::PerformAction(class UEdGraph* ParentGraph, UEdGraphPin* FromPin, const FVector2D Location, bool bSelectNewNode)
+{
+	FString NewNodeName = CommentNodeToDuplicate->GetFName().ToString();
+	NewNodeName.Append("_copy");
+
+	UEdGraphNode_Comment* CommentTemplate = NewObject<UEdGraphNode_Comment>(ParentGraph, FName(*NewNodeName));
+
+	CommentTemplate->bCommentBubbleVisible_InDetailsPanel = false;
+	CommentTemplate->bCommentBubbleVisible = false;
+	CommentTemplate->bCommentBubblePinned = false;
+
+	const FScopedTransaction Transaction(NSLOCTEXT("UnrealEd", "DataflowEditorNewNode", "Dataflow Editor: New Comment Node (by using Duplicate)"));
+
+	ParentGraph->Modify();
+
+	CommentTemplate->SetFlags(RF_Transactional);
+
+	// set outer to be the graph so it doesn't go away
+	CommentTemplate->Rename(NULL, ParentGraph, REN_NonTransactional);
+	ParentGraph->AddNode(CommentTemplate, true, bSelectNewNode);
+
+	CommentTemplate->CreateNewGuid();
+	CommentTemplate->PostPlacedNewNode();
+	CommentTemplate->AllocateDefaultPins();
+	CommentTemplate->AutowireNewNode(FromPin);
+
+	CommentTemplate->NodePosX = Location.X;
+	CommentTemplate->NodePosY = Location.Y;
+	CommentTemplate->NodeWidth = CommentNodeToDuplicate->NodeWidth;
+	CommentTemplate->NodeHeight = CommentNodeToDuplicate->NodeHeight;
+	CommentTemplate->SnapToGrid(GetDefault<UEditorStyleSettings>()->GridSnapSize);
+
+	CommentTemplate->NodeComment = CommentNodeToDuplicate->NodeComment;
+
+	ParentGraph->NotifyGraphChanged();
+
+	return CommentTemplate;
+}
+
 #undef LOCTEXT_NAMESPACE
diff --git a/Engine/Plugins/Experimental/Dataflow/Source/DataflowEditor/Private/Dataflow/DataflowSNode.cpp b/Engine/Plugins/Experimental/Dataflow/Source/DataflowEditor/Private/Dataflow/DataflowSNode.cpp
--- a/Engine/Plugins/Experimental/Dataflow/Source/DataflowEditor/Private/Dataflow/DataflowSNode.cpp
+++ b/Engine/Plugins/Experimental/Dataflow/Source/DataflowEditor/Private/Dataflow/DataflowSNode.cpp
@@ -8,7 +8,6 @@
 #include "Dataflow/DataflowGraphEditor.h"
 #include "Dataflow/DataflowNodeFactory.h"
 #include "Dataflow/DataflowObject.h"
-#include "Dataflow/DataflowCore.h"
 #include "Logging/LogMacros.h"
 #include "SourceCodeNavigation.h"
 #include "Styling/SlateTypes.h"
@@ -176,16 +175,60 @@ TSharedPtr<FAssetSchemaAction_Dataflow_CreateNode_DataflowEdNode> FAssetSchemaAc
 		{
 			const FText ToolTip = FText::FromString(Param.ToolTip.IsEmpty() ? FString("Add a Dataflow node.") : Param.ToolTip);
 			const FText NodeName = FText::FromString(Param.DisplayName.ToString());
-			const FText Catagory = FText::FromString(Param.Category.ToString().IsEmpty() ? FString("Dataflow") : Param.Category.ToString());
+			const FText Category = FText::FromString(Param.Category.ToString().IsEmpty() ? FString("Dataflow") : Param.Category.ToString());
 			const FText Tags = FText::FromString(Param.Tags);
 			TSharedPtr<FAssetSchemaAction_Dataflow_CreateNode_DataflowEdNode> NewNodeAction(
-				new FAssetSchemaAction_Dataflow_CreateNode_DataflowEdNode(InNodeTypeName, Catagory, NodeName, ToolTip, Tags));
+				new FAssetSchemaAction_Dataflow_CreateNode_DataflowEdNode(InNodeTypeName, Category, NodeName, ToolTip, Tags));
 			return NewNodeAction;
 		}
 	}
 	return TSharedPtr<FAssetSchemaAction_Dataflow_CreateNode_DataflowEdNode>(nullptr);
 }
 
+static UDataflowEdNode* CreateNode(UDataflow* Dataflow, UEdGraphPin* FromPin, const FVector2D Location, bool bSelectNewNode, const FName NodeUniqueName, const FName NodeTypeName, TSharedPtr<FDataflowNode> DataflowNodeToDuplicate, bool bCopySettings = false)
+{
+	if (Dataflow::FNodeFactory* Factory = Dataflow::FNodeFactory::GetInstance())
+	{
+		if (TSharedPtr<FDataflowNode> DataflowNode = Factory->NewNodeFromRegisteredType(*Dataflow->GetDataflow(), { FGuid::NewGuid(), NodeTypeName, NodeUniqueName }))
+		{
+			if (UDataflowEdNode* EdNode = NewObject<UDataflowEdNode>(Dataflow, UDataflowEdNode::StaticClass(), NodeUniqueName))
+			{
+				Dataflow->Modify();
+				if (FromPin != nullptr)
+				{
+					FromPin->Modify();
+				}
+
+				Dataflow->AddNode(EdNode, true, bSelectNewNode);
+
+				// Copy properties from DataflowNodeToDuplicate to DataflowNode
+				if (bCopySettings)
+				{
+					DataflowNode->CopyNodeProperties(DataflowNodeToDuplicate);
+				}
+
+				EdNode->CreateNewGuid();
+				EdNode->PostPlacedNewNode();
+
+				EdNode->SetDataflowGraph(Dataflow->GetDataflow());
+				EdNode->SetDataflowNodeGuid(DataflowNode->GetGuid());
+				EdNode->AllocateDefaultPins();
+
+				EdNode->AutowireNewNode(FromPin);
+
+				EdNode->NodePosX = Location.X;
+				EdNode->NodePosY = Location.Y;
+
+				EdNode->SetFlags(RF_Transactional);
+
+				return EdNode;
+			}
+		}
+	}
+
+	return nullptr;
+}
+
 //
 //  Created the EdGraph node and bind the guids to the Dataflow's node. 
 //
@@ -204,37 +247,9 @@ UEdGraphNode* FAssetSchemaAction_Dataflow_CreateNode_DataflowEdNode::PerformActi
 			NameIndex++;
 		}
 
-		if (Dataflow::FNodeFactory* Factory = Dataflow::FNodeFactory::GetInstance())
-		{
-			if (TSharedPtr<FDataflowNode> DataflowNode = Factory->NewNodeFromRegisteredType(*Dataflow->GetDataflow(), { FGuid::NewGuid(), NodeTypeName, NodeUniqueName }))
-			{
-				if (UDataflowEdNode* EdNode = NewObject<UDataflowEdNode>(Dataflow, UDataflowEdNode::StaticClass(), NodeUniqueName))
-				{
-					Dataflow->Modify();
-					if (FromPin != nullptr)
-						FromPin->Modify();
-		
-					Dataflow->AddNode(EdNode, true, bSelectNewNode);
-		
-					EdNode->CreateNewGuid();
-					EdNode->PostPlacedNewNode();
-		
-					EdNode->SetDataflowGraph(Dataflow->GetDataflow());
-					EdNode->SetDataflowNodeGuid(DataflowNode->GetGuid());
-					EdNode->AllocateDefaultPins();
-		
-					EdNode->AutowireNewNode(FromPin);
-		
-					EdNode->NodePosX = Location.X;
-					EdNode->NodePosY = Location.Y;
-		
-					EdNode->SetFlags(RF_Transactional);
-		
-					return EdNode;
-				}
-			}
-		}
+		return CreateNode(Dataflow, FromPin, Location, bSelectNewNode, NodeUniqueName, NodeTypeName, nullptr);
 	}
+
 	return nullptr;
 }
 
@@ -244,8 +259,57 @@ UEdGraphNode* FAssetSchemaAction_Dataflow_CreateNode_DataflowEdNode::PerformActi
 //	Collector.AddReferencedObject(NodeTemplate);
 //}
 
+//
+// 
+//
+TSharedPtr<FAssetSchemaAction_Dataflow_DuplicateNode_DataflowEdNode> FAssetSchemaAction_Dataflow_DuplicateNode_DataflowEdNode::CreateAction(UEdGraph* ParentGraph, const FName& InNodeTypeName)
+{
+	if (Dataflow::FNodeFactory* Factory = Dataflow::FNodeFactory::GetInstance())
+	{
+		const Dataflow::FFactoryParameters& Param = Factory->GetParameters(InNodeTypeName);
+		if (Param.IsValid())
+		{
+			const FText ToolTip = FText::FromString(Param.ToolTip.IsEmpty() ? FString("Add a Dataflow node.") : Param.ToolTip);
+			const FText NodeName = FText::FromString(Param.DisplayName.ToString());
+			const FText Category = FText::FromString(Param.Category.ToString().IsEmpty() ? FString("Dataflow") : Param.Category.ToString());
+			const FText Tags = FText::FromString(Param.Tags);
+			TSharedPtr<FAssetSchemaAction_Dataflow_DuplicateNode_DataflowEdNode> NewNodeAction(
+				new FAssetSchemaAction_Dataflow_DuplicateNode_DataflowEdNode(InNodeTypeName, Category, NodeName, ToolTip, Tags));
+			return NewNodeAction;
+		}
+	}
+	return TSharedPtr<FAssetSchemaAction_Dataflow_DuplicateNode_DataflowEdNode>(nullptr);
+}
+
+//
+//   
+//
+UEdGraphNode* FAssetSchemaAction_Dataflow_DuplicateNode_DataflowEdNode::PerformAction(class UEdGraph* ParentGraph, UEdGraphPin* FromPin, const FVector2D Location, bool bSelectNewNode)
+{
+	if (UDataflow* Dataflow = Cast<UDataflow>(ParentGraph))
+	{
+		// Append "_copy" to selected node's name if it doesn't have it
+		FString NodeToDuplicateName = DataflowNodeToDuplicate->GetName().ToString();
+		if (!NodeToDuplicateName.Contains("_copy"))
+		{
+			NodeToDuplicateName.Append("_copy");
+		}
+
+		// Check if that is unique, if not then make it unique with an index postfix
+		const FString NodeBaseName = NodeToDuplicateName;
+		FName NodeUniqueName{ NodeBaseName };
+		int32 NameIndex = 0;
+		while (Dataflow->GetDataflow()->FindBaseNode(FName(NodeUniqueName)) != nullptr)
+		{
+			NodeUniqueName = FName(NodeBaseName + FString::Printf(TEXT("_%d"), NameIndex));
+			NameIndex++;
+		}
 
+		return CreateNode(Dataflow, FromPin, Location, bSelectNewNode, NodeUniqueName, NodeTypeName, DataflowNodeToDuplicate, /*bCopySettings=*/true);
+	}
 
+	return nullptr;
+}
 
 #undef LOCTEXT_NAMESPACE
 
diff --git a/Engine/Plugins/Experimental/Dataflow/Source/DataflowEditor/Public/Dataflow/DataflowEditorCommands.h b/Engine/Plugins/Experimental/Dataflow/Source/DataflowEditor/Public/Dataflow/DataflowEditorCommands.h
--- a/Engine/Plugins/Experimental/Dataflow/Source/DataflowEditor/Public/Dataflow/DataflowEditorCommands.h
+++ b/Engine/Plugins/Experimental/Dataflow/Source/DataflowEditor/Public/Dataflow/DataflowEditorCommands.h
@@ -15,6 +15,7 @@ class UDataflow;
 class UDataflowEdNode;
 struct FDataflowNode;
 class UEdGraphNode;
+class SDataflowGraphEditor;
 
 typedef TSet<class UObject*> FGraphPanelSelectionSet;
 
@@ -107,6 +108,11 @@ public:
 	*/
 	static void ToggleEnabledState(UDataflow* Graph);
 
+	/*
+	*  DuplicateNodes
+	*/
+	static void DuplicateNodes(UDataflow* Graph, const TSharedPtr<SDataflowGraphEditor>& DataflowGraphEditor, const FGraphPanelSelectionSet& SelectedNodes);
+
 
 
 };
diff --git a/Engine/Plugins/Experimental/Dataflow/Source/DataflowEditor/Public/Dataflow/DataflowGraphEditor.h b/Engine/Plugins/Experimental/Dataflow/Source/DataflowEditor/Public/Dataflow/DataflowGraphEditor.h
--- a/Engine/Plugins/Experimental/Dataflow/Source/DataflowEditor/Public/Dataflow/DataflowGraphEditor.h
+++ b/Engine/Plugins/Experimental/Dataflow/Source/DataflowEditor/Public/Dataflow/DataflowGraphEditor.h
@@ -114,6 +114,9 @@ public:
 	/** */
 	void ToggleEnabledState();
 
+	/** */
+	void DuplicateSelectedNodes();
+
 	SGraphEditor* GetGraphEditor() { return (SGraphEditor*)this; }
 
 private:
diff --git a/Engine/Plugins/Experimental/Dataflow/Source/DataflowEditor/Public/Dataflow/DataflowSCommentNode.h b/Engine/Plugins/Experimental/Dataflow/Source/DataflowEditor/Public/Dataflow/DataflowSCommentNode.h
--- a/Engine/Plugins/Experimental/Dataflow/Source/DataflowEditor/Public/Dataflow/DataflowSCommentNode.h
+++ b/Engine/Plugins/Experimental/Dataflow/Source/DataflowEditor/Public/Dataflow/DataflowSCommentNode.h
@@ -45,3 +45,26 @@ public:
 
 	TSharedPtr<SGraphEditor> GraphEditor;
 };
+
+//
+// Action to duplicate a comment node in the graph
+//
+USTRUCT()
+struct DATAFLOWEDITOR_API FAssetSchemaAction_Dataflow_DuplicateCommentNode_DataflowEdNode : public FEdGraphSchemaAction
+{
+	GENERATED_USTRUCT_BODY();
+
+public:
+	FAssetSchemaAction_Dataflow_DuplicateCommentNode_DataflowEdNode(const TSharedPtr<SGraphEditor>& InGraphEditor) : FEdGraphSchemaAction()
+		, GraphEditor(InGraphEditor)
+	{}
+
+	FAssetSchemaAction_Dataflow_DuplicateCommentNode_DataflowEdNode() {}
+
+	static TSharedPtr<FAssetSchemaAction_Dataflow_DuplicateCommentNode_DataflowEdNode> CreateAction(UEdGraph* ParentGraph, const TSharedPtr<SGraphEditor>& GraphEditor);
+
+	virtual UEdGraphNode* PerformAction(class UEdGraph* ParentGraph, UEdGraphPin* FromPin, const FVector2D Location, bool bSelectNewNode = true) override;
+
+	TSharedPtr<SGraphEditor> GraphEditor;
+	UEdGraphNode_Comment* CommentNodeToDuplicate;
+};
diff --git a/Engine/Plugins/Experimental/Dataflow/Source/DataflowEditor/Public/Dataflow/DataflowSNode.h b/Engine/Plugins/Experimental/Dataflow/Source/DataflowEditor/Public/Dataflow/DataflowSNode.h
--- a/Engine/Plugins/Experimental/Dataflow/Source/DataflowEditor/Public/Dataflow/DataflowSNode.h
+++ b/Engine/Plugins/Experimental/Dataflow/Source/DataflowEditor/Public/Dataflow/DataflowSNode.h
@@ -5,6 +5,7 @@
 #include "EdGraphUtilities.h"
 #include "SGraphNode.h"
 #include "UObject/GCObject.h"
+#include "Dataflow/DataflowCore.h"
 
 #include "DataflowSNode.generated.h"
 
@@ -76,3 +77,25 @@ public:
 
 	FName NodeTypeName;
 };
+
+//
+// Action to duplicate a set of nodes in the graph
+//
+USTRUCT()
+struct DATAFLOWEDITOR_API FAssetSchemaAction_Dataflow_DuplicateNode_DataflowEdNode : public FEdGraphSchemaAction
+{
+	GENERATED_USTRUCT_BODY();
+
+public:
+	FAssetSchemaAction_Dataflow_DuplicateNode_DataflowEdNode() : FEdGraphSchemaAction() {}
+
+	FAssetSchemaAction_Dataflow_DuplicateNode_DataflowEdNode(const FName& InType, FText InNodeCategory, FText InMenuDesc, FText InToolTip, FText InKeywords)
+		: FEdGraphSchemaAction(InNodeCategory, InMenuDesc, InToolTip, 0, InKeywords), NodeTypeName(InType) {}
+
+	static TSharedPtr<FAssetSchemaAction_Dataflow_DuplicateNode_DataflowEdNode> CreateAction(UEdGraph* ParentGraph, const FName& NodeTypeName);
+
+	virtual UEdGraphNode* PerformAction(class UEdGraph* ParentGraph, UEdGraphPin* FromPin, const FVector2D Location, bool bSelectNewNode = true) override;
+
+	FName NodeTypeName;
+	TSharedPtr<FDataflowNode> DataflowNodeToDuplicate;
+};
diff --git a/Engine/Plugins/Experimental/GeometryCollectionPlugin/Source/GeometryCollectionNodes/Private/Dataflow/GeometryCollectionSelectionNodes.cpp b/Engine/Plugins/Experimental/GeometryCollectionPlugin/Source/GeometryCollectionNodes/Private/Dataflow/GeometryCollectionSelectionNodes.cpp
--- a/Engine/Plugins/Experimental/GeometryCollectionPlugin/Source/GeometryCollectionNodes/Private/Dataflow/GeometryCollectionSelectionNodes.cpp
+++ b/Engine/Plugins/Experimental/GeometryCollectionPlugin/Source/GeometryCollectionNodes/Private/Dataflow/GeometryCollectionSelectionNodes.cpp
@@ -953,7 +953,7 @@ void FCollectionSelectionConvertDataflowNode::Evaluate(Dataflow::FContext& Conte
 		{
 			// Passthrough
 			const FDataflowFaceSelection& InFaceSelection = GetValue<FDataflowFaceSelection>(Context, &FaceSelection);
-			SetValue<FDataflowFaceSelection>(Context, InFaceSelection, &InFaceSelection);
+			SetValue<FDataflowFaceSelection>(Context, InFaceSelection, &FaceSelection);
 		}
 		else if (IsConnected<FDataflowTransformSelection>(&TransformSelection))
 		{
diff --git a/Engine/Source/Runtime/Experimental/Dataflow/Core/Private/Dataflow/DataflowNode.cpp b/Engine/Source/Runtime/Experimental/Dataflow/Core/Private/Dataflow/DataflowNode.cpp
--- a/Engine/Source/Runtime/Experimental/Dataflow/Core/Private/Dataflow/DataflowNode.cpp
+++ b/Engine/Source/Runtime/Experimental/Dataflow/Core/Private/Dataflow/DataflowNode.cpp
@@ -5,6 +5,8 @@
 #include "ChaosLog.h"
 #include "Dataflow/DataflowInputOutput.h"
 #include "Dataflow/DataflowArchive.h"
+#include "Serialization/MemoryWriter.h"
+#include "Serialization/MemoryReader.h"
 
 #include UE_INLINE_GENERATED_CPP_BY_NAME(DataflowNode)
 
@@ -573,6 +575,18 @@ TArray<FString> FDataflowNode::GetPinMetaData(const FName& PropertyName)
 	return TArray<FString>();
 }
 
+void FDataflowNode::CopyNodeProperties(const TSharedPtr<FDataflowNode> CopyFromDataflowNode)
+{
+	TArray<uint8> NodeData;
+
+	FMemoryWriter ArWriter(NodeData);
+	CopyFromDataflowNode->SerializeInternal(ArWriter);
+
+	FMemoryReader ArReader(NodeData);
+	this->SerializeInternal(ArReader);
+}
+
+
 
 
 
diff --git a/Engine/Source/Runtime/Experimental/Dataflow/Core/Public/Dataflow/DataflowNode.h b/Engine/Source/Runtime/Experimental/Dataflow/Core/Public/Dataflow/DataflowNode.h
--- a/Engine/Source/Runtime/Experimental/Dataflow/Core/Public/Dataflow/DataflowNode.h
+++ b/Engine/Source/Runtime/Experimental/Dataflow/Core/Public/Dataflow/DataflowNode.h
@@ -80,6 +80,8 @@ struct DATAFLOWCORE_API FDataflowNode
 	FText GetPinDisplayName(const FName& PropertyName);
 	TArray<FString> GetPinMetaData(const FName& PropertyName);
 	virtual TArray<Dataflow::FRenderingParameter> GetRenderParameters() const { return GetRenderParametersImpl(); }
+	// Copy node property values from another node
+	void CopyNodeProperties(const TSharedPtr<FDataflowNode> CopyFromDataflowNode);
 
 	//
 	// Connections
