Dataflow:
- Fixed Issue 316405

diff --git a/Engine/Plugins/Experimental/Dataflow/Source/DataflowEditor/Private/Dataflow/DataflowSCommentNode.cpp b/Engine/Plugins/Experimental/Dataflow/Source/DataflowEditor/Private/Dataflow/DataflowSCommentNode.cpp
--- a/Engine/Plugins/Experimental/Dataflow/Source/DataflowEditor/Private/Dataflow/DataflowSCommentNode.cpp
+++ b/Engine/Plugins/Experimental/Dataflow/Source/DataflowEditor/Private/Dataflow/DataflowSCommentNode.cpp
@@ -91,7 +91,7 @@ UEdGraphNode* FAssetSchemaAction_Dataflow_DuplicateCommentNode_DataflowEdNode::P
 	CommentTemplate->bCommentBubbleVisible = false;
 	CommentTemplate->bCommentBubblePinned = false;
 
-	const FScopedTransaction Transaction(NSLOCTEXT("UnrealEd", "DataflowEditorNewNode", "Dataflow Editor: New Comment Node (by using Duplicate)"));
+//	const FScopedTransaction Transaction(NSLOCTEXT("UnrealEd", "DataflowEditorDuplicatedNode", "Dataflow Editor: New Comment Node (by using Duplicate)"));
 
 	ParentGraph->Modify();
 
