DataflowGraphEditor: zoom to fit the selected nodes or entire graph when the user presses the 'F' key

diff --git a/Engine/Plugins/Experimental/Dataflow/Source/DataflowEditor/Private/Dataflow/DataflowEditorCommands.cpp b/Engine/Plugins/Experimental/Dataflow/Source/DataflowEditor/Private/Dataflow/DataflowEditorCommands.cpp
--- a/Engine/Plugins/Experimental/Dataflow/Source/DataflowEditor/Private/Dataflow/DataflowEditorCommands.cpp
+++ b/Engine/Plugins/Experimental/Dataflow/Source/DataflowEditor/Private/Dataflow/DataflowEditorCommands.cpp
@@ -26,6 +26,8 @@ void FDataflowEditorCommandsImpl::RegisterCommands()
 	UI_COMMAND(ToggleVertexSelection, "ToggleVertexSelection", "Enable vertex selection in editor.", EUserInterfaceActionType::ToggleButton, FInputChord());
 	UI_COMMAND(AddOptionPin, "AddOptionPin", "Add an option pin to the selected nodes.", EUserInterfaceActionType::Button, FInputChord());
 	UI_COMMAND(RemoveOptionPin, "RemoveOptionPin", "Remove the last option pin from the selected nodes.", EUserInterfaceActionType::Button, FInputChord());
+	UI_COMMAND(ZoomToFitGraph, "ZoomToFitGraph", "Fit the graph in the graph editor viewport.", EUserInterfaceActionType::None, FInputChord(EKeys::F));
+	
 
 	if (Dataflow::FNodeFactory* Factory = Dataflow::FNodeFactory::GetInstance())
 	{
diff --git a/Engine/Plugins/Experimental/Dataflow/Source/DataflowEditor/Private/Dataflow/DataflowGraphEditor.cpp b/Engine/Plugins/Experimental/Dataflow/Source/DataflowEditor/Private/Dataflow/DataflowGraphEditor.cpp
--- a/Engine/Plugins/Experimental/Dataflow/Source/DataflowEditor/Private/Dataflow/DataflowGraphEditor.cpp
+++ b/Engine/Plugins/Experimental/Dataflow/Source/DataflowEditor/Private/Dataflow/DataflowGraphEditor.cpp
@@ -102,6 +102,10 @@ void SDataflowGraphEditor::Construct(const FArguments& InArgs, UObject* InAssetO
 				FGenericCommands::Get().Duplicate,
 				FExecuteAction::CreateSP(this, &SDataflowGraphEditor::DuplicateSelectedNodes)
 			);
+			GraphEditorCommands->MapAction(
+				FDataflowEditorCommands::Get().ZoomToFitGraph,
+				FExecuteAction::CreateSP(this, &SDataflowGraphEditor::ZoomToFitGraph)
+			);
 		}
 	}
 
@@ -491,4 +495,10 @@ void SDataflowGraphEditor::DuplicateSelectedNodes()
 	}
 }
 
+void SDataflowGraphEditor::ZoomToFitGraph()
+{
+	constexpr bool bOnlySelection = true;	// This will focus on the selected nodes, if any. If no nodes are selected, it will focus the whole graph.
+	ZoomToFit(bOnlySelection);
+}
+
 #undef LOCTEXT_NAMESPACE
diff --git a/Engine/Plugins/Experimental/Dataflow/Source/DataflowEditor/Public/Dataflow/DataflowEditorCommands.h b/Engine/Plugins/Experimental/Dataflow/Source/DataflowEditor/Public/Dataflow/DataflowEditorCommands.h
--- a/Engine/Plugins/Experimental/Dataflow/Source/DataflowEditor/Public/Dataflow/DataflowEditorCommands.h
+++ b/Engine/Plugins/Experimental/Dataflow/Source/DataflowEditor/Public/Dataflow/DataflowEditorCommands.h
@@ -46,6 +46,7 @@ public:
 	TSharedPtr<FUICommandInfo> ToggleVertexSelection;
 	TSharedPtr< FUICommandInfo > AddOptionPin;
 	TSharedPtr< FUICommandInfo > RemoveOptionPin;
+	TSharedPtr< FUICommandInfo > ZoomToFitGraph;
 
 	TMap< FName, TSharedPtr<FUICommandInfo> > CreateNodesMap;
 };
diff --git a/Engine/Plugins/Experimental/Dataflow/Source/DataflowEditor/Public/Dataflow/DataflowGraphEditor.h b/Engine/Plugins/Experimental/Dataflow/Source/DataflowEditor/Public/Dataflow/DataflowGraphEditor.h
--- a/Engine/Plugins/Experimental/Dataflow/Source/DataflowEditor/Public/Dataflow/DataflowGraphEditor.h
+++ b/Engine/Plugins/Experimental/Dataflow/Source/DataflowEditor/Public/Dataflow/DataflowGraphEditor.h
@@ -117,6 +117,9 @@ public:
 	/** */
 	void DuplicateSelectedNodes();
 
+	/** */
+	void ZoomToFitGraph();
+
 	SGraphEditor* GetGraphEditor() { return (SGraphEditor*)this; }
 
 private:
