Dataflow : Module file renamed.

diff --git a/Engine/Plugins/Experimental/ChaosFlesh/Source/ChaosFleshEditor/Private/ChaosFlesh/Asset/AssetTypeActions_FleshAsset.cpp b/Engine/Plugins/Experimental/ChaosFlesh/Source/ChaosFleshEditor/Private/ChaosFlesh/Asset/AssetTypeActions_FleshAsset.cpp
--- a/Engine/Plugins/Experimental/ChaosFlesh/Source/ChaosFleshEditor/Private/ChaosFlesh/Asset/AssetTypeActions_FleshAsset.cpp
+++ b/Engine/Plugins/Experimental/ChaosFlesh/Source/ChaosFleshEditor/Private/ChaosFlesh/Asset/AssetTypeActions_FleshAsset.cpp
@@ -3,7 +3,7 @@
 #include "ChaosFlesh/Asset/AssetTypeActions_FleshAsset.h"
 
 #include "ChaosFlesh/FleshAsset.h"
-#include "Dataflow/DataflowEditorPlugin.h"
+#include "Dataflow/DataflowEditorModule.h"
 #include "Dataflow/DataflowEditorToolkit.h"
 #include "ThumbnailRendering/SceneThumbnailInfo.h"
 
@@ -43,7 +43,7 @@ void FAssetTypeActions_FleshAsset::OpenAssetEditor(const TArray<UObject*>& InObj
 			if (FDataflowEditorToolkit::CanOpenDataflowEditor(Object))
 			{
 				bNeedsBaseEditor = false;
-				IDataflowEditorPlugin* DataflowEditorPlugin = &FModuleManager::LoadModuleChecked<IDataflowEditorPlugin>("DataflowEditor");
+				FDataflowEditorModule* DataflowEditorPlugin = &FModuleManager::LoadModuleChecked<FDataflowEditorModule>("DataflowEditor");
 				DataflowEditorPlugin->CreateDataflowAssetEditor(Mode, EditWithinLevelEditor, Object);
 			}
 		}
diff --git a/Engine/Plugins/Experimental/Dataflow/Source/DataflowEditor/Private/Dataflow/DataflowAssetActions.cpp b/Engine/Plugins/Experimental/Dataflow/Source/DataflowEditor/Private/Dataflow/DataflowAssetActions.cpp
--- a/Engine/Plugins/Experimental/Dataflow/Source/DataflowEditor/Private/Dataflow/DataflowAssetActions.cpp
+++ b/Engine/Plugins/Experimental/Dataflow/Source/DataflowEditor/Private/Dataflow/DataflowAssetActions.cpp
@@ -1,7 +1,7 @@
 // Copyright Epic Games, Inc. All Rights Reserved.
 
 #include "Dataflow/DataflowAssetActions.h"
-#include "Dataflow/DataflowEditorPlugin.h"
+#include "Dataflow/DataflowEditorModule.h"
 
 #include "Dataflow/DataflowObject.h"
 
@@ -37,7 +37,7 @@ void FDataflowAssetActions::OpenAssetEditor(
 	{
 		if (auto Object = Cast<UDataflow>(*ObjIt))
 		{
-			IDataflowEditorPlugin* DataflowEditorPlugin = &FModuleManager::LoadModuleChecked<IDataflowEditorPlugin>("DataflowEditor");
+			FDataflowEditorModule* DataflowEditorPlugin = &FModuleManager::LoadModuleChecked<FDataflowEditorModule>("DataflowEditor");
 			DataflowEditorPlugin->CreateDataflowAssetEditor(Mode, EditWithinLevelEditor, Object);
 		}
 	}
diff --git a/Engine/Plugins/Experimental/Dataflow/Source/DataflowEditor/Private/Dataflow/DataflowEditorModule.cpp b/Engine/Plugins/Experimental/Dataflow/Source/DataflowEditor/Private/Dataflow/DataflowEditorModule.cpp
new file mode 100644
--- /dev/null
+++ b/Engine/Plugins/Experimental/Dataflow/Source/DataflowEditor/Private/Dataflow/DataflowEditorModule.cpp
@@ -0,0 +1,60 @@
+// Copyright Epic Games, Inc. All Rights Reserved.
+
+#include "Dataflow/DataflowEditorModule.h"
+#include "Dataflow/DataflowEditorStyle.h"
+
+#include "AssetToolsModule.h"
+#include "CoreMinimal.h"
+#include "EdGraphUtilities.h"
+#include "Dataflow/DataflowEditorToolkit.h"
+#include "Dataflow/DataflowNodeFactory.h"
+#include "Dataflow/DataflowAssetActions.h"
+#include "Dataflow/DataflowSNodeFactories.h"
+#include "PropertyEditorModule.h"
+
+#define LOCTEXT_NAMESPACE "DataflowEditor"
+
+//#define BOX_BRUSH(StyleSet, RelativePath, ...) FSlateBoxBrush(StyleSet->RootToContentDir(RelativePath, TEXT(".png")), __VA_ARGS__)
+//#define DEFAULT_FONT(...) FCoreStyle::GetDefaultFontStyle(__VA_ARGS__)
+
+void FDataflowEditorModule::StartupModule()
+{
+	FDataflowEditorStyle::Get();
+
+	DataflowAssetActions = new FDataflowAssetActions();
+
+	FAssetToolsModule& AssetToolsModule = FAssetToolsModule::GetModule();
+	IAssetTools& AssetTools = AssetToolsModule.Get();
+	AssetTools.RegisterAssetTypeActions(MakeShareable(DataflowAssetActions));
+
+	FPropertyEditorModule& PropertyEditorModule = FModuleManager::LoadModuleChecked<FPropertyEditorModule>(TEXT("PropertyEditor"));
+
+	DataflowSNodeFactory = MakeShareable(new FDataflowSNodeFactory());
+	FEdGraphUtilities::RegisterVisualNodeFactory(DataflowSNodeFactory);
+}
+
+void FDataflowEditorModule::ShutdownModule()
+{
+	if (UObjectInitialized())
+	{
+		FAssetToolsModule& AssetToolsModule = FAssetToolsModule::GetModule();
+		IAssetTools& AssetTools = AssetToolsModule.Get();
+
+		AssetTools.UnregisterAssetTypeActions(DataflowAssetActions->AsShared());
+
+		FEdGraphUtilities::UnregisterVisualNodeFactory(DataflowSNodeFactory);
+	}
+}
+
+TSharedRef<FAssetEditorToolkit> FDataflowEditorModule::CreateDataflowAssetEditor(const EToolkitMode::Type Mode, const TSharedPtr<IToolkitHost>& InitToolkitHost, UObject* Asset)
+{
+	TSharedPtr<FDataflowEditorToolkit> NewDataflowAssetEditor = MakeShared<FDataflowEditorToolkit>();
+	NewDataflowAssetEditor->InitializeEditor(Mode, InitToolkitHost, Asset);
+	return StaticCastSharedPtr<FAssetEditorToolkit>(NewDataflowAssetEditor).ToSharedRef();
+}
+
+
+IMPLEMENT_MODULE(FDataflowEditorModule, DataflowEditor)
+
+
+#undef LOCTEXT_NAMESPACE
diff --git a/Engine/Plugins/Experimental/Dataflow/Source/DataflowEditor/Private/Dataflow/DataflowEditorPlugin.cpp b/Engine/Plugins/Experimental/Dataflow/Source/DataflowEditor/Private/Dataflow/DataflowEditorPlugin.cpp
deleted file mode 100644
diff --git a/Engine/Plugins/Experimental/Dataflow/Source/DataflowEditor/Public/Dataflow/DataflowEditor.h b/Engine/Plugins/Experimental/Dataflow/Source/DataflowEditor/Public/Dataflow/DataflowEditor.h
--- a/Engine/Plugins/Experimental/Dataflow/Source/DataflowEditor/Public/Dataflow/DataflowEditor.h
+++ b/Engine/Plugins/Experimental/Dataflow/Source/DataflowEditor/Public/Dataflow/DataflowEditor.h
@@ -6,7 +6,7 @@
 #include "Dataflow/DataflowAssetActions.h"
 #include "Dataflow/DataflowAssetFactory.h"
 #include "Dataflow/DataflowEditorCommands.h"
-#include "Dataflow/DataflowEditorPlugin.h"
+#include "Dataflow/DataflowEditorModule.h"
 #include "Dataflow/DataflowSchema.h"
 #include "Dataflow/DataflowSNode.h"
 #include "Dataflow/DataflowSNodeFactories.h"
diff --git a/Engine/Plugins/Experimental/Dataflow/Source/DataflowEditor/Public/Dataflow/DataflowEditorModule.h b/Engine/Plugins/Experimental/Dataflow/Source/DataflowEditor/Public/Dataflow/DataflowEditorModule.h
new file mode 100644
--- /dev/null
+++ b/Engine/Plugins/Experimental/Dataflow/Source/DataflowEditor/Public/Dataflow/DataflowEditorModule.h
@@ -0,0 +1,64 @@
+// Copyright Epic Games, Inc. All Rights Reserved.
+
+#pragma once
+
+#include "CoreMinimal.h"
+#include "Modules/ModuleInterface.h"
+#include "Modules/ModuleManager.h"
+//#include "Toolkits/IToolkit.h"
+#include "Toolkits/AssetEditorToolkit.h"
+
+//class ISlateStyle;
+//class FSlateStyleSet;
+class FDataflowSNodeFactory;
+class FDataflowAssetActions;
+
+/**
+ * The public interface to this module
+ */
+class DATAFLOWEDITOR_API FDataflowEditorModule : public IModuleInterface
+{
+public:
+	virtual void StartupModule();
+	virtual void ShutdownModule();
+
+	TSharedRef<FAssetEditorToolkit> CreateDataflowAssetEditor(const EToolkitMode::Type Mode, const TSharedPtr< IToolkitHost >& InitToolkitHost, UObject* Dataflow);
+
+
+	/**
+	 * Singleton-like access to this module's interface.  This is just for convenience!
+	 * Beware of calling this during the shutdown phase, though.  Your module might have been unloaded already.
+	 *
+	 * @return Returns singleton instance, loading the module on demand if needed
+	 */
+	static inline FDataflowEditorModule& Get()
+	{
+		return FModuleManager::LoadModuleChecked< FDataflowEditorModule >( "DataflowEditor" );
+	}
+
+	/**
+	 * Checks to see if this module is loaded and ready.  It is only valid to call Get() if IsAvailable() returns true.
+	 *
+	 * @return True if the module is loaded and ready to use
+	 */
+	static inline bool IsAvailable()
+	{
+		return FModuleManager::Get().IsModuleLoaded( "DataflowEditor" );
+	}
+
+	//TSharedPtr<FSlateStyleSet> GetStyleSet() { return StyleSet; }
+	//static FName GetEditorStyleName();
+	//static const ISlateStyle* GetEditorStyle();
+
+private:
+	void RegisterMenus();
+
+private:
+
+	FDataflowAssetActions* DataflowAssetActions;
+
+	//TSharedPtr<FSlateStyleSet> StyleSet;
+	TSharedPtr<FDataflowSNodeFactory> DataflowSNodeFactory;
+
+};
+
diff --git a/Engine/Plugins/Experimental/Dataflow/Source/DataflowEditor/Public/Dataflow/DataflowEditorPlugin.h b/Engine/Plugins/Experimental/Dataflow/Source/DataflowEditor/Public/Dataflow/DataflowEditorPlugin.h
deleted file mode 100644
diff --git a/Engine/Plugins/Experimental/GeometryCollectionPlugin/Source/GeometryCollectionEditor/Private/GeometryCollection/AssetTypeActions_GeometryCollection.cpp b/Engine/Plugins/Experimental/GeometryCollectionPlugin/Source/GeometryCollectionEditor/Private/GeometryCollection/AssetTypeActions_GeometryCollection.cpp
--- a/Engine/Plugins/Experimental/GeometryCollectionPlugin/Source/GeometryCollectionEditor/Private/GeometryCollection/AssetTypeActions_GeometryCollection.cpp
+++ b/Engine/Plugins/Experimental/GeometryCollectionPlugin/Source/GeometryCollectionEditor/Private/GeometryCollection/AssetTypeActions_GeometryCollection.cpp
@@ -2,7 +2,7 @@
 
 #include "GeometryCollection/AssetTypeActions_GeometryCollection.h"
 
-#include "Dataflow/DataflowEditorPlugin.h"
+#include "Dataflow/DataflowEditorModule.h"
 #include "Dataflow/DataflowEditorToolkit.h"
 #include "ThumbnailRendering/SceneThumbnailInfo.h"
 #include "GeometryCollection/GeometryCollectionObject.h"
@@ -40,7 +40,7 @@ void FAssetTypeActions_GeometryCollection::OpenAssetEditor(const TArray<UObject*
 			if (FDataflowEditorToolkit::CanOpenDataflowEditor(Object))
 			{
 				bNeedsBaseEditor = false;
-				IDataflowEditorPlugin* DataflowEditorPlugin = &FModuleManager::LoadModuleChecked<IDataflowEditorPlugin>("DataflowEditor");
+				FDataflowEditorModule* DataflowEditorPlugin = &FModuleManager::LoadModuleChecked<FDataflowEditorModule>("DataflowEditor");
 				DataflowEditorPlugin->CreateDataflowAssetEditor(Mode, EditWithinLevelEditor, Object);
 			}
 		}
