DataflowEditor : Cycle check on connect

diff --git a/Engine/Plugins/Experimental/Dataflow/Source/DataflowEditor/Private/Dataflow/DataflowSchema.cpp b/Engine/Plugins/Experimental/Dataflow/Source/DataflowEditor/Private/Dataflow/DataflowSchema.cpp
--- a/Engine/Plugins/Experimental/Dataflow/Source/DataflowEditor/Private/Dataflow/DataflowSchema.cpp
+++ b/Engine/Plugins/Experimental/Dataflow/Source/DataflowEditor/Private/Dataflow/DataflowSchema.cpp
@@ -92,6 +92,43 @@ void UDataflowSchema::GetGraphContextActions(FGraphContextMenuBuilder& ContextMe
 	}
 }
 
+
+bool HasLoop(const UEdGraphNode* DownstreamNode, const UEdGraphNode* UpstreamNode)
+{
+	if (DownstreamNode == UpstreamNode)
+	{
+		return true;
+	}
+
+	TArray< const UEdGraphNode*> ConnectedUpstreamNodes;
+	for (UEdGraphPin* Pin : DownstreamNode->GetAllPins())
+	{
+		if (Pin->Direction == EEdGraphPinDirection::EGPD_Input)
+		{
+			if (Pin->HasAnyConnections())
+			{
+				if (ensure(Pin->LinkedTo.Num() == 1))
+				{
+					if (Pin->LinkedTo[0]->GetOwningNode())
+					{
+						ConnectedUpstreamNodes.Add(Pin->LinkedTo[0]->GetOwningNode());
+					}
+				}
+			}
+		}
+	}
+
+	for (const UEdGraphNode* ConnectedNode : ConnectedUpstreamNodes)
+	{
+		if (HasLoop(ConnectedNode, UpstreamNode))
+		{
+			return true;
+		}
+	}
+
+	return false;
+}
+
 const FPinConnectionResponse UDataflowSchema::CanCreateConnection(const UEdGraphPin* InPinA, const UEdGraphPin* InPinB) const
 {
 	bool bSwapped = false;
@@ -112,18 +149,25 @@ const FPinConnectionResponse UDataflowSchema::CanCreateConnection(const UEdGraph
 			// Make sure the pins are not on the same node
 			if (PinA->GetOwningNode() != PinB->GetOwningNode())
 			{
-				// Make sure types match. 
-				if (PinA->PinType == PinB->PinType)
+				if (!HasLoop(PinA->GetOwningNode(), PinB->GetOwningNode()))
 				{
-					if (PinB->LinkedTo.Num())
+					// Make sure types match. 
+					if (PinA->PinType == PinB->PinType)
 					{
-						return (bSwapped) ?
-							FPinConnectionResponse(CONNECT_RESPONSE_BREAK_OTHERS_A, LOCTEXT("PinSteal", "Disconnect existing input and connect new input."))
-							:
-							FPinConnectionResponse(CONNECT_RESPONSE_BREAK_OTHERS_B, LOCTEXT("PinSteal", "Disconnect existing input and connect new input."));
+						if (PinB->LinkedTo.Num())
+						{
+							return (bSwapped) ?
+								FPinConnectionResponse(CONNECT_RESPONSE_BREAK_OTHERS_A, LOCTEXT("PinSteal", "Disconnect existing input and connect new input."))
+								:
+								FPinConnectionResponse(CONNECT_RESPONSE_BREAK_OTHERS_B, LOCTEXT("PinSteal", "Disconnect existing input and connect new input."));
 
+						}
+						return FPinConnectionResponse(CONNECT_RESPONSE_MAKE, LOCTEXT("PinConnect", "Connect input to output."));
 					}
-					return FPinConnectionResponse(CONNECT_RESPONSE_MAKE, LOCTEXT("PinConnect", "Connect input to output."));
+				}
+				else
+				{
+					return FPinConnectionResponse(CONNECT_RESPONSE_DISALLOW, LOCTEXT("PinErrorSameNode_Loop", "Loop"));
 				}
 			}
 		}
