Dataflow : Cycle check after type check.

diff --git a/Engine/Plugins/Experimental/Dataflow/Source/DataflowEditor/Private/Dataflow/DataflowSchema.cpp b/Engine/Plugins/Experimental/Dataflow/Source/DataflowEditor/Private/Dataflow/DataflowSchema.cpp
--- a/Engine/Plugins/Experimental/Dataflow/Source/DataflowEditor/Private/Dataflow/DataflowSchema.cpp
+++ b/Engine/Plugins/Experimental/Dataflow/Source/DataflowEditor/Private/Dataflow/DataflowSchema.cpp
@@ -149,11 +149,13 @@ const FPinConnectionResponse UDataflowSchema::CanCreateConnection(const UEdGraph
 			// Make sure the pins are not on the same node
 			if (PinA->GetOwningNode() != PinB->GetOwningNode())
 			{
-				if (!HasLoop(PinA->GetOwningNode(), PinB->GetOwningNode()))
+				// Make sure types match. 
+				if (PinA->PinType == PinB->PinType)
 				{
-					// Make sure types match. 
-					if (PinA->PinType == PinB->PinType)
+					// cycle checking on connect
+					if (!HasLoop(PinA->GetOwningNode(), PinB->GetOwningNode()))
 					{
+
 						if (PinB->LinkedTo.Num())
 						{
 							return (bSwapped) ?
@@ -162,12 +164,17 @@ const FPinConnectionResponse UDataflowSchema::CanCreateConnection(const UEdGraph
 								FPinConnectionResponse(CONNECT_RESPONSE_BREAK_OTHERS_B, LOCTEXT("PinSteal", "Disconnect existing input and connect new input."));
 
 						}
+
 						return FPinConnectionResponse(CONNECT_RESPONSE_MAKE, LOCTEXT("PinConnect", "Connect input to output."));
 					}
+					else
+					{
+						return FPinConnectionResponse(CONNECT_RESPONSE_DISALLOW, LOCTEXT("PinError_Loop", "Graph Cycle"));
+					}
 				}
 				else
 				{
-					return FPinConnectionResponse(CONNECT_RESPONSE_DISALLOW, LOCTEXT("PinErrorSameNode_Loop", "Loop"));
+					return FPinConnectionResponse(CONNECT_RESPONSE_DISALLOW, LOCTEXT("PinError_Type mismatch", "Type Mismatch"));
 				}
 			}
 		}
