Dataflow : imporved loop detection when connection nodes in dataflow
- Faster traversal, garanteed to only visit each node once at most
- robust to existing loops

diff --git a/Engine/Plugins/Experimental/Dataflow/Source/DataflowEditor/Private/Dataflow/DataflowSchema.cpp b/Engine/Plugins/Experimental/Dataflow/Source/DataflowEditor/Private/Dataflow/DataflowSchema.cpp
--- a/Engine/Plugins/Experimental/Dataflow/Source/DataflowEditor/Private/Dataflow/DataflowSchema.cpp
+++ b/Engine/Plugins/Experimental/Dataflow/Source/DataflowEditor/Private/Dataflow/DataflowSchema.cpp
@@ -92,38 +92,51 @@ void UDataflowSchema::GetGraphContextActions(FGraphContextMenuBuilder& ContextMe
 	}
 }
 
-
-bool HasLoop(const UEdGraphNode* DownstreamNode, const UEdGraphNode* UpstreamNode)
+bool HasLoopIfConnected(const UEdGraphNode* FromNode, const UEdGraphNode* ToNode)
 {
-	if (DownstreamNode == UpstreamNode)
+	if (ToNode == FromNode)
 	{
 		return true;
 	}
 
-	TArray< const UEdGraphNode*> ConnectedUpstreamNodes;
-	for (UEdGraphPin* Pin : DownstreamNode->GetAllPins())
+	// We only need to process from the FromNode and test if anything in the feeding nodes contains ToNode
+	TArray<const UEdGraphNode*> NodesToProcess;
+	NodesToProcess.Push(FromNode);
+
+	// to speed things up, we do not revisit branches we have already look at  
+	TSet<const UEdGraphNode*> VisitedNodes;
+
+	while (NodesToProcess.Num() > 0)
 	{
-		if (Pin->Direction == EEdGraphPinDirection::EGPD_Input)
+		const UEdGraphNode* NodeToProcess = NodesToProcess.Pop();
+		if (!VisitedNodes.Contains(NodeToProcess))
 		{
-			if (Pin->HasAnyConnections())
+			VisitedNodes.Add(NodeToProcess);
+
+			int32 NumConnectedInputPins = 0;
+			for (UEdGraphPin* Pin : NodeToProcess->GetAllPins())
 			{
-				if (ensure(Pin->LinkedTo.Num() == 1))
+				if (Pin->Direction == EEdGraphPinDirection::EGPD_Input)
 				{
-					if (Pin->LinkedTo[0]->GetOwningNode())
+					if (Pin->HasAnyConnections())
 					{
-						ConnectedUpstreamNodes.Add(Pin->LinkedTo[0]->GetOwningNode());
+						NumConnectedInputPins++;
+						if (ensure(Pin->LinkedTo.Num() == 1))
+						{
+							if (const UEdGraphNode* OwningNode = Pin->LinkedTo[0]->GetOwningNode())
+							{
+								if (OwningNode == ToNode)
+								{
+									return true;
+								}
+								NodesToProcess.Push(OwningNode);
+							}
+						}
 					}
 				}
 			}
 		}
-	}
 
-	for (const UEdGraphNode* ConnectedNode : ConnectedUpstreamNodes)
-	{
-		if (HasLoop(ConnectedNode, UpstreamNode))
-		{
-			return true;
-		}
 	}
 
 	return false;
@@ -153,7 +166,7 @@ const FPinConnectionResponse UDataflowSchema::CanCreateConnection(const UEdGraph
 				if (PinA->PinType == PinB->PinType)
 				{
 					// cycle checking on connect
-					if (!HasLoop(PinA->GetOwningNode(), PinB->GetOwningNode()))
+					if (!HasLoopIfConnected(PinA->GetOwningNode(), PinB->GetOwningNode()))
 					{
 
 						if (PinB->LinkedTo.Num())
