Dataflow graph editor: implement show/hide connected pins commands.

diff --git a/Engine/Plugins/Experimental/Dataflow/Source/DataflowEditor/Private/Dataflow/DataflowGraphEditor.cpp b/Engine/Plugins/Experimental/Dataflow/Source/DataflowEditor/Private/Dataflow/DataflowGraphEditor.cpp
--- a/Engine/Plugins/Experimental/Dataflow/Source/DataflowEditor/Private/Dataflow/DataflowGraphEditor.cpp
+++ b/Engine/Plugins/Experimental/Dataflow/Source/DataflowEditor/Private/Dataflow/DataflowGraphEditor.cpp
@@ -15,6 +15,7 @@
 #include "HAL/PlatformApplicationMisc.h"
 #include "ScopedTransaction.h"
 #include "IStructureDetailsView.h"
+#include "SGraphPanel.h"
 
 #define LOCTEXT_NAMESPACE "DataflowGraphEditor"
 
@@ -106,6 +107,24 @@ void SDataflowGraphEditor::Construct(const FArguments& InArgs, UObject* InAssetO
 				FDataflowEditorCommands::Get().ZoomToFitGraph,
 				FExecuteAction::CreateSP(this, &SDataflowGraphEditor::ZoomToFitGraph)
 			);
+			GraphEditorCommands->MapAction(
+				FGraphEditorCommands::Get().ShowAllPins,
+				FExecuteAction::CreateSP(this, &SDataflowGraphEditor::SetPinVisibility, SGraphEditor::Pin_Show),
+				FCanExecuteAction(),
+				FIsActionChecked::CreateSP(this, &SDataflowGraphEditor::GetPinVisibility, SGraphEditor::Pin_Show)
+			);
+			GraphEditorCommands->MapAction(
+				FGraphEditorCommands::Get().HideNoConnectionPins,
+				FExecuteAction::CreateSP(this, &SDataflowGraphEditor::SetPinVisibility, SGraphEditor::Pin_HideNoConnection),
+				FCanExecuteAction(),
+				FIsActionChecked::CreateSP(this, &SDataflowGraphEditor::GetPinVisibility, SGraphEditor::Pin_HideNoConnection)
+			);
+			GraphEditorCommands->MapAction(
+				FGraphEditorCommands::Get().HideNoConnectionNoDefaultPins,
+				FExecuteAction::CreateSP(this, &SDataflowGraphEditor::SetPinVisibility, SGraphEditor::Pin_HideNoConnectionNoDefault),
+				FCanExecuteAction(),
+				FIsActionChecked::CreateSP(this, &SDataflowGraphEditor::GetPinVisibility, SGraphEditor::Pin_HideNoConnectionNoDefault)
+			);
 		}
 	}
 
@@ -501,4 +520,13 @@ void SDataflowGraphEditor::ZoomToFitGraph()
 	ZoomToFit(bOnlySelection);
 }
 
+bool SDataflowGraphEditor::GetPinVisibility(SGraphEditor::EPinVisibility PinVisibility) const
+{
+	if (const SGraphPanel* GraphPanel = GetGraphPanel())
+	{
+		return GraphPanel->GetPinVisibility() == PinVisibility;
+	}
+	return false;
+}
+
 #undef LOCTEXT_NAMESPACE
diff --git a/Engine/Plugins/Experimental/Dataflow/Source/DataflowEditor/Private/Dataflow/DataflowSchema.cpp b/Engine/Plugins/Experimental/Dataflow/Source/DataflowEditor/Private/Dataflow/DataflowSchema.cpp
--- a/Engine/Plugins/Experimental/Dataflow/Source/DataflowEditor/Private/Dataflow/DataflowSchema.cpp
+++ b/Engine/Plugins/Experimental/Dataflow/Source/DataflowEditor/Private/Dataflow/DataflowSchema.cpp
@@ -43,7 +43,6 @@ void UDataflowSchema::GetContextMenuActions(class UToolMenu* Menu, class UGraphN
 				Section.AddMenuEntry(FDataflowEditorCommands::Get().EvaluateNode);
 			}
 		}
-
 		{
 			FToolMenuSection& Section = Menu->AddSection("TestGraphSchemaOrganization", LOCTEXT("GraphSchemaOrganization_MenuHeader", "Organization"));
 			{
@@ -69,8 +68,21 @@ void UDataflowSchema::GetContextMenuActions(class UToolMenu* Menu, class UGraphN
 				}));
 			}
 		}
-	}
+		{
+			FToolMenuSection& Section = Menu->AddSection("TestGraphSchemaDisplay", LOCTEXT("GraphSchemaDisplay_MenuHeader", "Display"));
+			{
+				Section.AddSubMenu("PinVisibility", LOCTEXT("PinVisibilityHeader", "Pin Visibility"), FText(),
+					FNewToolMenuDelegate::CreateLambda([](UToolMenu* PinVisibilityMenu)
+				{
+					FToolMenuSection& InSection = PinVisibilityMenu->AddSection("TestGraphSchemaPinVisibility");
+					InSection.AddMenuEntry(FGraphEditorCommands::Get().ShowAllPins);
+					InSection.AddMenuEntry(FGraphEditorCommands::Get().HideNoConnectionPins);
+					InSection.AddMenuEntry(FGraphEditorCommands::Get().HideNoConnectionNoDefaultPins);
+				}));
+			}
 
+		}
+	}
 	Super::GetContextMenuActions(Menu, Context);
 }
 
diff --git a/Engine/Plugins/Experimental/Dataflow/Source/DataflowEditor/Public/Dataflow/DataflowGraphEditor.h b/Engine/Plugins/Experimental/Dataflow/Source/DataflowEditor/Public/Dataflow/DataflowGraphEditor.h
--- a/Engine/Plugins/Experimental/Dataflow/Source/DataflowEditor/Public/Dataflow/DataflowGraphEditor.h
+++ b/Engine/Plugins/Experimental/Dataflow/Source/DataflowEditor/Public/Dataflow/DataflowGraphEditor.h
@@ -133,6 +133,8 @@ private:
 	/** Return whether all currently selected Dataflow nodes can execute the RemovePin function. */
 	bool CanRemoveOptionPin() const;
 
+	bool GetPinVisibility(SGraphEditor::EPinVisibility InVisibility) const;
+
 	FDataflowEditorCommands::FOnDragDropEventCallback OnDragDropEventCallback;
 	FDataflowEditorCommands::FGraphEvaluationCallback EvaluateGraphCallback;
 
