Dataflow:
- Fixed: Crash when attempting to Duplicate Node in Dataflow Graph

diff --git a/Engine/Plugins/Experimental/Dataflow/Source/DataflowEditor/Private/Dataflow/DataflowSNode.cpp b/Engine/Plugins/Experimental/Dataflow/Source/DataflowEditor/Private/Dataflow/DataflowSNode.cpp
--- a/Engine/Plugins/Experimental/Dataflow/Source/DataflowEditor/Private/Dataflow/DataflowSNode.cpp
+++ b/Engine/Plugins/Experimental/Dataflow/Source/DataflowEditor/Private/Dataflow/DataflowSNode.cpp
@@ -15,6 +15,7 @@
 #include "Styling/AppStyle.h"
 #include "Widgets/Layout/SBorder.h"
 #include "Widgets/Input/SCheckBox.h"
+#include "Editor/Transactor.h"
 
 #include UE_INLINE_GENERATED_CPP_BY_NAME(DataflowSNode)
 
@@ -185,6 +186,18 @@ TSharedPtr<FAssetSchemaAction_Dataflow_CreateNode_DataflowEdNode> FAssetSchemaAc
 	return TSharedPtr<FAssetSchemaAction_Dataflow_CreateNode_DataflowEdNode>(nullptr);
 }
 
+void SDataflowEdNode::CopyDataflowNodeSettings(TSharedPtr<FDataflowNode> SourceDataflowNode, TSharedPtr<FDataflowNode> TargetDataflowNode)
+{
+	using namespace UE::Transaction;
+	FSerializedObject SerializationObject;
+
+	FSerializedObjectDataWriter ArWriter(SerializationObject);
+	SourceDataflowNode->SerializeInternal(ArWriter);
+
+	FSerializedObjectDataReader ArReader(SerializationObject);
+	TargetDataflowNode->SerializeInternal(ArReader);
+}
+
 static UDataflowEdNode* CreateNode(UDataflow* Dataflow, UEdGraphPin* FromPin, const FVector2D Location, bool bSelectNewNode, const FName NodeUniqueName, const FName NodeTypeName, TSharedPtr<FDataflowNode> DataflowNodeToDuplicate, bool bCopySettings = false)
 {
 	if (Dataflow::FNodeFactory* Factory = Dataflow::FNodeFactory::GetInstance())
@@ -204,7 +217,7 @@ static UDataflowEdNode* CreateNode(UDataflow* Dataflow, UEdGraphPin* FromPin, co
 				// Copy properties from DataflowNodeToDuplicate to DataflowNode
 				if (bCopySettings)
 				{
-					DataflowNode->CopyNodeProperties(DataflowNodeToDuplicate);
+					SDataflowEdNode::CopyDataflowNodeSettings(DataflowNodeToDuplicate, DataflowNode);
 				}
 
 				EdNode->CreateNewGuid();
diff --git a/Engine/Plugins/Experimental/Dataflow/Source/DataflowEditor/Public/Dataflow/DataflowSNode.h b/Engine/Plugins/Experimental/Dataflow/Source/DataflowEditor/Public/Dataflow/DataflowSNode.h
--- a/Engine/Plugins/Experimental/Dataflow/Source/DataflowEditor/Public/Dataflow/DataflowSNode.h
+++ b/Engine/Plugins/Experimental/Dataflow/Source/DataflowEditor/Public/Dataflow/DataflowSNode.h
@@ -36,7 +36,7 @@ public:
 	virtual FReply OnMouseButtonDoubleClick(const FGeometry& InMyGeometry, const FPointerEvent& InMouseEvent) override;
 	TArray<FOverlayWidgetInfo> GetOverlayWidgets(bool bSelected, const FVector2D& WidgetSize) const;
 
-
+	static void CopyDataflowNodeSettings(TSharedPtr<FDataflowNode> SourceDataflowNode, TSharedPtr<FDataflowNode> TargetDataflowNode);
 
 	//~ Begin FGCObject interface
 	virtual void AddReferencedObjects(FReferenceCollector& Collector) override;
diff --git a/Engine/Source/Runtime/Experimental/Dataflow/Core/Public/Dataflow/DataflowNode.h b/Engine/Source/Runtime/Experimental/Dataflow/Core/Public/Dataflow/DataflowNode.h
--- a/Engine/Source/Runtime/Experimental/Dataflow/Core/Public/Dataflow/DataflowNode.h
+++ b/Engine/Source/Runtime/Experimental/Dataflow/Core/Public/Dataflow/DataflowNode.h
@@ -81,6 +81,7 @@ struct FDataflowNode
 	DATAFLOWCORE_API TArray<FString> GetPinMetaData(const FName& PropertyName);
 	virtual TArray<Dataflow::FRenderingParameter> GetRenderParameters() const { return GetRenderParametersImpl(); }
 	// Copy node property values from another node
+	UE_DEPRECATED(5.4, "FDataflowNode::CopyNodeProperties is deprecated.")
 	DATAFLOWCORE_API void CopyNodeProperties(const TSharedPtr<FDataflowNode> CopyFromDataflowNode);
 
 	//
