Dataflow : Restructured editor to use the FBaseCharacterFXEditorToolkit
-	Updated Geometry Collection Editor and Chaos Flesh Editor to reflect the initialization changes.
-	Updated the Asset Actions for Geometry Collection and Chaos Flesh to use the Asset Definition
-	Moved render state into the Dataflow Editor Viewport Client.
-	Added linker dependencies as needed.
-	Stopped short of adding the subsystems for viewport editing operations.

diff --git a/Engine/Plugins/Experimental/ChaosFlesh/Source/ChaosFleshEditor/ChaosFleshEditor.Build.cs b/Engine/Plugins/Experimental/ChaosFlesh/Source/ChaosFleshEditor/ChaosFleshEditor.Build.cs
--- a/Engine/Plugins/Experimental/ChaosFlesh/Source/ChaosFleshEditor/ChaosFleshEditor.Build.cs
+++ b/Engine/Plugins/Experimental/ChaosFlesh/Source/ChaosFleshEditor/ChaosFleshEditor.Build.cs
@@ -9,37 +9,38 @@ namespace UnrealBuildTool.Rules
 			PublicDependencyModuleNames.AddRange(
 				new string[]
 				{
+					"AdvancedPreviewScene",
+					"AssetDefinition",
+					"AssetRegistry",
+					"AssetTools",
+					"Chaos",
+					"ChaosFlesh",
+					"ChaosFleshEngine",
 					"Core",
 					"CoreUObject",
-				    "Slate",
-				    "SlateCore",
-				    "Engine",
+					"DataflowCore",
+					"DataflowEditor",
+					"DataflowEngine",
+					"DataflowEnginePlugin",
+					"Engine",
 					"EditorFramework",
-					"UnrealEd",
+					"EditorStyle",
+					"GraphEditor",
+					"InputCore",
+					"LevelEditor",
+					"MeshDescription",
+					"Slate",
+				    "SlateCore",
 					"Projects",
 					"PropertyEditor",
-				    "RenderCore",
-				    "RHI",
-                    "ChaosFleshEngine",
-					"ChaosFlesh",
+					"RHI",
 					"RawMesh",
-				    "AssetTools",
-				    "AssetRegistry",
-				    "SceneOutliner",
-					"EditorStyle",
-					"ToolMenus",
-					"Chaos",
-					"MeshDescription",
-					"StaticMeshDescription",
-					"LevelEditor",
-					"InputCore",
-					"AdvancedPreviewScene",
-					"GraphEditor",
-					"DataflowCore",
-					"DataflowEngine",
-					"DataflowEnginePlugin",
-					"DataflowEditor",
+					"RenderCore",
+					"SceneOutliner",
 					"SkeletonEditor",
+					"StaticMeshDescription",
+					"ToolMenus",
+					"UnrealEd",
 				}
 			);
 
diff --git a/Engine/Plugins/Experimental/ChaosFlesh/Source/ChaosFleshEditor/Private/ChaosFlesh/Asset/AssetDefinition_FleshAsset.cpp b/Engine/Plugins/Experimental/ChaosFlesh/Source/ChaosFleshEditor/Private/ChaosFlesh/Asset/AssetDefinition_FleshAsset.cpp
new file mode 100644
--- /dev/null
+++ b/Engine/Plugins/Experimental/ChaosFlesh/Source/ChaosFleshEditor/Private/ChaosFlesh/Asset/AssetDefinition_FleshAsset.cpp
@@ -0,0 +1,79 @@
+// Copyright Epic Games, Inc. All Rights Reserved.
+
+#include "ChaosFlesh/Asset/AssetDefinition_FleshAsset.h"
+#include "ChaosFlesh/FleshAsset.h"
+#include "Dataflow/DataflowEditor.h"
+#include "Dataflow/DataflowEditorModule.h"
+#include "Dataflow/DataflowEditorToolkit.h"
+#include "Math/Color.h"
+#include "ThumbnailRendering/SceneThumbnailInfo.h"
+#include "Toolkits/SimpleAssetEditor.h"
+
+
+#define LOCTEXT_NAMESPACE "AssetActions_FleshAsset"
+
+
+namespace UE::ChaosFlesh::FleshAsset
+{
+	struct FColorScheme
+	{
+		static inline const FLinearColor Asset = FColor(180, 120, 110);
+		static inline const FLinearColor NodeHeader = FColor(180, 120, 110);
+		static inline const FLinearColor NodeBody = FColor(18, 12, 11, 127);
+	};
+}
+
+FText UAssetDefinition_FleshAsset::GetAssetDisplayName() const
+{
+	return NSLOCTEXT("AssetTypeActions", "AssetTypeActions_FleshAsset", "FleshAsset");
+}
+
+TSoftClassPtr<UObject> UAssetDefinition_FleshAsset::GetAssetClass() const
+{
+	return UFleshAsset::StaticClass();
+}
+
+FLinearColor UAssetDefinition_FleshAsset::GetAssetColor() const
+{
+	return UE::ChaosFlesh::FleshAsset::FColorScheme::Asset;
+}
+
+TConstArrayView<FAssetCategoryPath> UAssetDefinition_FleshAsset::GetAssetCategories() const
+{
+	static const auto Categories = { EAssetCategoryPaths::Physics };
+	return Categories;
+}
+
+UThumbnailInfo* UAssetDefinition_FleshAsset::LoadThumbnailInfo(const FAssetData& InAsset) const
+{
+	return UE::Editor::FindOrCreateThumbnailInfo(InAsset.GetAsset(), USceneThumbnailInfo::StaticClass());
+}
+
+EAssetCommandResult UAssetDefinition_FleshAsset::OpenAssets(const FAssetOpenArgs& OpenArgs) const
+{
+	TArray<UFleshAsset*> FleshObjects = OpenArgs.LoadObjects<UFleshAsset>();
+
+	// For now the flesh editor only works on one asset at a time
+	ensure(FleshObjects.Num() == 0 || FleshObjects.Num() == 1);
+	if (FleshObjects.Num() == 1)
+	{
+		// Validate the asset
+		if (UFleshAsset* const FleshAsset = Cast<UFleshAsset>(FleshObjects[0]))
+		{
+			if (FDataflowEditorToolkit::CanOpenDataflowEditor(FleshAsset))
+			{
+				UAssetEditorSubsystem* const AssetEditorSubsystem = GEditor->GetEditorSubsystem<UAssetEditorSubsystem>();
+				UDataflowEditor* const AssetEditor = NewObject<UDataflowEditor>(AssetEditorSubsystem, NAME_None, RF_Transient);
+				AssetEditor->Initialize({ FleshAsset });
+				return EAssetCommandResult::Handled;
+			}
+
+			FSimpleAssetEditor::CreateEditor(EToolkitMode::Standalone, OpenArgs.ToolkitHost, FleshAsset );
+			return EAssetCommandResult::Handled;
+		}
+	}
+	return EAssetCommandResult::Unhandled;
+}
+
+
+#undef LOCTEXT_NAMESPACE
diff --git a/Engine/Plugins/Experimental/ChaosFlesh/Source/ChaosFleshEditor/Private/ChaosFlesh/Asset/AssetTypeActions_FleshAsset.cpp b/Engine/Plugins/Experimental/ChaosFlesh/Source/ChaosFleshEditor/Private/ChaosFlesh/Asset/AssetTypeActions_FleshAsset.cpp
deleted file mode 100644
diff --git a/Engine/Plugins/Experimental/ChaosFlesh/Source/ChaosFleshEditor/Private/ChaosFlesh/ChaosFleshEditorPlugin.cpp b/Engine/Plugins/Experimental/ChaosFlesh/Source/ChaosFleshEditor/Private/ChaosFlesh/ChaosFleshEditorPlugin.cpp
--- a/Engine/Plugins/Experimental/ChaosFlesh/Source/ChaosFleshEditor/Private/ChaosFlesh/ChaosFleshEditorPlugin.cpp
+++ b/Engine/Plugins/Experimental/ChaosFlesh/Source/ChaosFleshEditor/Private/ChaosFlesh/ChaosFleshEditorPlugin.cpp
@@ -3,7 +3,7 @@
 
 #include "ChaosFlesh/ChaosFleshEditorPlugin.h"
 
-#include "ChaosFlesh/Asset/AssetTypeActions_FleshAsset.h"
+#include "ChaosFlesh/Asset/AssetDefinition_FleshAsset.h"
 #include "ChaosFlesh/Asset/FleshDeformableInterfaceDetails.h"
 #include "ChaosFlesh/Asset/FleshAssetThumbnailRenderer.h"
 #include "ChaosFlesh/ChaosDeformableCollisionsActor.h"
@@ -23,11 +23,6 @@ void IChaosFleshEditorPlugin::StartupModule()
 {
 	FChaosFleshEditorStyle::Get();
 
-	IAssetTools& AssetTools = FAssetToolsModule::GetModule().Get();
-
-	FleshAssetActions = new FAssetTypeActions_FleshAsset();
-	AssetTools.RegisterAssetTypeActions(MakeShareable(FleshAssetActions));
-
 	if (GIsEditor && !IsRunningCommandlet())
 	{
 		EditorCommands.Add(IConsoleManager::Get().RegisterConsoleCommand(
@@ -50,8 +45,6 @@ void IChaosFleshEditorPlugin::StartupModule()
 		));
 	}
 
-	UThumbnailManager::Get().RegisterCustomRenderer(UFleshAsset::StaticClass(), UFleshAssetThumbnailRenderer::StaticClass());
-
 	// register details customization
 	FPropertyEditorModule& PropertyModule = FModuleManager::LoadModuleChecked<FPropertyEditorModule>("PropertyEditor");
 
@@ -79,13 +72,6 @@ void IChaosFleshEditorPlugin::StartupModule()
 
 void IChaosFleshEditorPlugin::ShutdownModule()
 {
-	if (UObjectInitialized())
-	{	
-		UThumbnailManager::Get().UnregisterCustomRenderer(UFleshAsset::StaticClass());
-
-		IAssetTools& AssetTools = FAssetToolsModule::GetModule().Get();
-		AssetTools.UnregisterAssetTypeActions(FleshAssetActions->AsShared());
-	}
 }
 
 IMPLEMENT_MODULE(IChaosFleshEditorPlugin, ChaosFleshEditor)
diff --git a/Engine/Plugins/Experimental/ChaosFlesh/Source/ChaosFleshEditor/Public/ChaosFlesh/Asset/AssetDefinition_FleshAsset.h b/Engine/Plugins/Experimental/ChaosFlesh/Source/ChaosFleshEditor/Public/ChaosFlesh/Asset/AssetDefinition_FleshAsset.h
new file mode 100644
--- /dev/null
+++ b/Engine/Plugins/Experimental/ChaosFlesh/Source/ChaosFleshEditor/Public/ChaosFlesh/Asset/AssetDefinition_FleshAsset.h
@@ -0,0 +1,24 @@
+// Copyright Epic Games, Inc. All Rights Reserved.
+
+#pragma once
+
+#include "AssetDefinitionDefault.h"
+
+#include "AssetDefinition_FleshAsset.generated.h"
+
+UCLASS()
+class CHAOSFLESHEDITOR_API UAssetDefinition_FleshAsset : public UAssetDefinitionDefault
+{
+	GENERATED_BODY()
+
+private:
+
+	virtual FText GetAssetDisplayName() const override;
+	virtual TSoftClassPtr<UObject> GetAssetClass() const override;
+	virtual FLinearColor GetAssetColor() const override;
+	virtual TConstArrayView<FAssetCategoryPath> GetAssetCategories() const override;
+	virtual UThumbnailInfo* LoadThumbnailInfo(const FAssetData& InAssetData) const override;
+
+	virtual EAssetCommandResult OpenAssets(const FAssetOpenArgs& OpenArgs) const override;
+};
+
diff --git a/Engine/Plugins/Experimental/ChaosFlesh/Source/ChaosFleshEditor/Public/ChaosFlesh/Asset/AssetTypeActions_FleshAsset.h b/Engine/Plugins/Experimental/ChaosFlesh/Source/ChaosFleshEditor/Public/ChaosFlesh/Asset/AssetTypeActions_FleshAsset.h
deleted file mode 100644
diff --git a/Engine/Plugins/Experimental/ChaosFlesh/Source/ChaosFleshEditor/Public/ChaosFlesh/ChaosFleshEditorPlugin.h b/Engine/Plugins/Experimental/ChaosFlesh/Source/ChaosFleshEditor/Public/ChaosFlesh/ChaosFleshEditorPlugin.h
--- a/Engine/Plugins/Experimental/ChaosFlesh/Source/ChaosFleshEditor/Public/ChaosFlesh/ChaosFleshEditorPlugin.h
+++ b/Engine/Plugins/Experimental/ChaosFlesh/Source/ChaosFleshEditor/Public/ChaosFlesh/ChaosFleshEditorPlugin.h
@@ -56,7 +56,6 @@ private:
 	TArray<IConsoleObject*> EditorCommands;
 
 	// Asset actions for new asset types
-	FAssetTypeActions_FleshAsset* FleshAssetActions;
 	FAssetTypeActions_ChaosDeformableSolver* ChaosDeformableSolverAssetActions;
 
 	// Styleset for flesh tool brushes/fonts etc.
diff --git a/Engine/Plugins/Experimental/Dataflow/Dataflow.uplugin b/Engine/Plugins/Experimental/Dataflow/Dataflow.uplugin
--- a/Engine/Plugins/Experimental/Dataflow/Dataflow.uplugin
+++ b/Engine/Plugins/Experimental/Dataflow/Dataflow.uplugin
@@ -33,6 +33,11 @@
 		}
 	],
 	"Plugins": [
+		{
+			"Name": "BaseCharacterFXEditor",
+			"Enabled": true,
+			"TargetAllowList": [ "Editor" ]
+		},
 		{
 			"Name": "ProceduralMeshComponent",
 			"Enabled": true
diff --git a/Engine/Plugins/Experimental/Dataflow/Source/DataflowEditor/DataflowEditor.Build.cs b/Engine/Plugins/Experimental/Dataflow/Source/DataflowEditor/DataflowEditor.Build.cs
--- a/Engine/Plugins/Experimental/Dataflow/Source/DataflowEditor/DataflowEditor.Build.cs
+++ b/Engine/Plugins/Experimental/Dataflow/Source/DataflowEditor/DataflowEditor.Build.cs
@@ -16,35 +16,37 @@ namespace UnrealBuildTool.Rules
 			PublicDependencyModuleNames.AddRange(
 				new string[]
 				{
-					"Core",
+					"AdvancedPreviewScene",
 					"ApplicationCore",
+					"AssetDefinition",
+					"AssetTools",
+					"AssetRegistry",
+					"BaseCharacterFXEditor",
+					"Chaos",
+					"Core",
 					"CoreUObject",
-				    "Slate",
-				    "SlateCore",
-				    "Engine",
+					"DataflowCore",
+					"DataflowEngine",
+					"DataflowEnginePlugin",
+					"DataflowNodes",
+					"Engine",
 					"EditorFramework",
+					"EditorStyle",
+					"GraphEditor",
+					"InputCore",
+					"InteractiveToolsFramework",
+					"LevelEditor",
+					"Slate",
+				    "SlateCore",
 					"UnrealEd",
 					"Projects",
 					"PropertyEditor",
 				    "RenderCore",
 				    "RHI",
-				    "AssetTools",
-				    "AssetRegistry",
 				    "SceneOutliner",
-					"EditorStyle",
-					"AssetTools",
 					"ToolMenus",
-					"LevelEditor",
-					"InputCore",
-					"AdvancedPreviewScene",
-					"GraphEditor",
-					"DataflowCore",
-					"DataflowEngine",
-					"DataflowEnginePlugin",
-					"DataflowNodes",
 					"Slate",
-					"Chaos",
-					"XmlParser",
+					"XmlParser"
 				}
 			);
 		}
diff --git a/Engine/Plugins/Experimental/Dataflow/Source/DataflowEditor/Private/Dataflow/AssetDefinition_DataflowAsset.cpp b/Engine/Plugins/Experimental/Dataflow/Source/DataflowEditor/Private/Dataflow/AssetDefinition_DataflowAsset.cpp
new file mode 100644
--- /dev/null
+++ b/Engine/Plugins/Experimental/Dataflow/Source/DataflowEditor/Private/Dataflow/AssetDefinition_DataflowAsset.cpp
@@ -0,0 +1,74 @@
+// Copyright Epic Games, Inc. All Rights Reserved.
+
+#include "Dataflow/AssetDefinition_DataflowAsset.h"
+#include "Dataflow/DataflowEditor.h"
+#include "Dataflow/DataflowEditorModule.h"
+#include "Dataflow/DataflowObject.h"
+#include "Math/Color.h"
+#include "ThumbnailRendering/SceneThumbnailInfo.h"
+
+
+#define LOCTEXT_NAMESPACE "AssetActions_DataflowAsset"
+
+
+namespace UE::Dataflow::DataflowAsset
+{
+	struct FColorScheme
+	{
+		static inline const FLinearColor Asset = FColor(180, 120, 110);
+		static inline const FLinearColor NodeHeader = FColor(180, 120, 110);
+		static inline const FLinearColor NodeBody = FColor(18, 12, 11, 127);
+	};
+}
+
+FText UAssetDefinition_DataflowAsset::GetAssetDisplayName() const
+{
+	return NSLOCTEXT("AssetTypeActions", "AssetTypeActions_DataflowAsset", "DataflowAsset");
+}
+
+TSoftClassPtr<UObject> UAssetDefinition_DataflowAsset::GetAssetClass() const
+{
+	return UDataflow::StaticClass();
+}
+
+FLinearColor UAssetDefinition_DataflowAsset::GetAssetColor() const
+{
+	return UE::Dataflow::DataflowAsset::FColorScheme::Asset;
+}
+
+TConstArrayView<FAssetCategoryPath> UAssetDefinition_DataflowAsset::GetAssetCategories() const
+{
+	static const auto Categories = { EAssetCategoryPaths::Physics };
+	return Categories;
+}
+
+UThumbnailInfo* UAssetDefinition_DataflowAsset::LoadThumbnailInfo(const FAssetData& InAsset) const
+{
+	return UE::Editor::FindOrCreateThumbnailInfo(InAsset.GetAsset(), USceneThumbnailInfo::StaticClass());
+}
+
+EAssetCommandResult UAssetDefinition_DataflowAsset::OpenAssets(const FAssetOpenArgs& OpenArgs) const
+{
+	TArray<UDataflow*> DataflowObjects = OpenArgs.LoadObjects<UDataflow>();
+
+	// For now the dataflow editor only works on one asset at a time
+	ensure(DataflowObjects.Num() == 0 || DataflowObjects.Num() == 1);
+
+	if (DataflowObjects.Num() == 1)
+	{
+		UAssetEditorSubsystem* const AssetEditorSubsystem = GEditor->GetEditorSubsystem<UAssetEditorSubsystem>();
+		UDataflowEditor* const AssetEditor = NewObject<UDataflowEditor>(AssetEditorSubsystem, NAME_None, RF_Transient);
+
+		// Validate the asset
+		if (UDataflow* const DataflowAsset = CastChecked<UDataflow>(DataflowObjects[0]))
+		{
+			AssetEditor->Initialize({ DataflowAsset });
+			return EAssetCommandResult::Handled;
+		}
+	}
+
+	return EAssetCommandResult::Unhandled;
+}
+
+
+#undef LOCTEXT_NAMESPACE
diff --git a/Engine/Plugins/Experimental/Dataflow/Source/DataflowEditor/Private/Dataflow/DataflowAssetActions.cpp b/Engine/Plugins/Experimental/Dataflow/Source/DataflowEditor/Private/Dataflow/DataflowAssetActions.cpp
deleted file mode 100644
diff --git a/Engine/Plugins/Experimental/Dataflow/Source/DataflowEditor/Private/Dataflow/DataflowEditor.cpp b/Engine/Plugins/Experimental/Dataflow/Source/DataflowEditor/Private/Dataflow/DataflowEditor.cpp
new file mode 100644
--- /dev/null
+++ b/Engine/Plugins/Experimental/Dataflow/Source/DataflowEditor/Private/Dataflow/DataflowEditor.cpp
@@ -0,0 +1,17 @@
+// Copyright Epic Games, Inc. All Rights Reserved.
+
+#include "Dataflow/DataflowEditor.h"
+#include "Dataflow/DataflowEditorToolkit.h"
+
+#include UE_INLINE_GENERATED_CPP_BY_NAME(DataflowEditor)
+
+DEFINE_LOG_CATEGORY(LogDataflowEditor);
+
+TSharedPtr<FBaseAssetToolkit> UDataflowEditor::CreateToolkit()
+{
+	TSharedPtr<FDataflowEditorToolkit> DataflowToolkit = MakeShared<FDataflowEditorToolkit>(this);
+	return DataflowToolkit;
+}
+
+
+
diff --git a/Engine/Plugins/Experimental/Dataflow/Source/DataflowEditor/Private/Dataflow/DataflowEditorCommands.cpp b/Engine/Plugins/Experimental/Dataflow/Source/DataflowEditor/Private/Dataflow/DataflowEditorCommands.cpp
--- a/Engine/Plugins/Experimental/Dataflow/Source/DataflowEditor/Private/Dataflow/DataflowEditorCommands.cpp
+++ b/Engine/Plugins/Experimental/Dataflow/Source/DataflowEditor/Private/Dataflow/DataflowEditorCommands.cpp
@@ -16,6 +16,15 @@
 
 #define LOCTEXT_NAMESPACE "DataflowEditorCommands"
 
+FDataflowEditorCommandsImpl::FDataflowEditorCommandsImpl()
+	: TBaseCharacterFXEditorCommands<FDataflowEditorCommandsImpl>("DataflowEditor", 
+		LOCTEXT("ContextDescription", "Dataflow Editor"), 
+		NAME_None,
+		FAppStyle::GetAppStyleSetName())
+{
+}
+
+
 void FDataflowEditorCommandsImpl::RegisterCommands()
 {
 	UI_COMMAND(EvaluateNode, "Evaluate", "Trigger an evaluation of the selected node.", EUserInterfaceActionType::Button, FInputChord());
@@ -46,9 +55,11 @@ void FDataflowEditorCommandsImpl::RegisterCommands()
 			);
 			CreateNodesMap.Add(Parameters.TypeName, AddNode);
 		}
-	}
+	}		
+}
 
-		
+void FDataflowEditorCommandsImpl::GetToolDefaultObjectList(TArray<UInteractiveTool*>& ToolCDOs)
+{
 }
 
 const FDataflowEditorCommandsImpl& FDataflowEditorCommands::Get()
diff --git a/Engine/Plugins/Experimental/Dataflow/Source/DataflowEditor/Private/Dataflow/DataflowEditorMode.cpp b/Engine/Plugins/Experimental/Dataflow/Source/DataflowEditor/Private/Dataflow/DataflowEditorMode.cpp
new file mode 100644
--- /dev/null
+++ b/Engine/Plugins/Experimental/Dataflow/Source/DataflowEditor/Private/Dataflow/DataflowEditorMode.cpp
@@ -0,0 +1,32 @@
+// Copyright Epic Games, Inc. All Rights Reserved.
+
+#include "Dataflow/DataflowEditorMode.h"
+#include "Dataflow/DataflowEditorToolkit.h"
+#include "Dataflow/DataflowEditorModeToolkit.h"
+
+
+#include UE_INLINE_GENERATED_CPP_BY_NAME(DataflowEditorMode)
+
+#define LOCTEXT_NAMESPACE "UDataflowEditorMode"
+
+const FEditorModeID UDataflowEditorMode::EM_DataflowEditorModeId = TEXT("EM_DataflowAssetEditorMode");
+
+
+
+
+UDataflowEditorMode::UDataflowEditorMode()
+{
+	Info = FEditorModeInfo(
+		EM_DataflowEditorModeId,
+		LOCTEXT("DataflowEditorModeName", "Dataflow"),
+		FSlateIcon(),
+		false);
+}
+
+void UDataflowEditorMode::CreateToolkit()
+{
+	Toolkit = MakeShared<FDataflowEditorModeToolkit>();
+}
+
+#undef LOCTEXT_NAMESPACE
+
diff --git a/Engine/Plugins/Experimental/Dataflow/Source/DataflowEditor/Private/Dataflow/DataflowEditorModeToolkit.cpp b/Engine/Plugins/Experimental/Dataflow/Source/DataflowEditor/Private/Dataflow/DataflowEditorModeToolkit.cpp
new file mode 100644
--- /dev/null
+++ b/Engine/Plugins/Experimental/Dataflow/Source/DataflowEditor/Private/Dataflow/DataflowEditorModeToolkit.cpp
@@ -0,0 +1,38 @@
+// Copyright Epic Games, Inc. All Rights Reserved.
+
+#include "Dataflow/DataflowEditorModeToolkit.h"
+
+#include "AssetEditorModeManager.h"
+#include "Dataflow/DataflowEditorCommands.h"
+#include "Dataflow/DataflowEditorStyle.h"
+#include "InteractiveToolManager.h"
+
+#define LOCTEXT_NAMESPACE "FDataflowEditorModeToolkit"
+
+void FDataflowEditorModeToolkit::OnToolStarted(UInteractiveToolManager* Manager, UInteractiveTool* Tool)
+{
+	ensure(Tool == Manager->GetActiveTool(EToolSide::Left));
+
+	ActiveToolName = Tool->GetToolInfo().ToolDisplayName;
+	FString ActiveToolIdentifier = Manager->GetActiveToolName(EToolSide::Mouse);
+	ActiveToolIdentifier.InsertAt(0, ".");
+	ActiveToolIcon = GetActiveToolIcon(ActiveToolIdentifier);
+}
+
+FName FDataflowEditorModeToolkit::GetToolkitFName() const
+{
+	return FName("DataflowEditorMode");
+}
+
+FText FDataflowEditorModeToolkit::GetBaseToolkitName() const
+{
+	return NSLOCTEXT("DataflowEditorModeToolkit", "DisplayName", "DataflowEditorMode");
+}
+
+const FSlateBrush* FDataflowEditorModeToolkit::GetActiveToolIcon(const FString& ActiveToolIdentifier) const
+{
+	FName ActiveToolIconName = ISlateStyle::Join(FDataflowEditorCommandsImpl::Get().GetContextName(), TCHAR_TO_ANSI(*ActiveToolIdentifier));
+	return FDataflowEditorStyle::Get().GetOptionalBrush(ActiveToolIconName);
+}
+
+#undef LOCTEXT_NAMESPACE
diff --git a/Engine/Plugins/Experimental/Dataflow/Source/DataflowEditor/Private/Dataflow/DataflowEditorModule.cpp b/Engine/Plugins/Experimental/Dataflow/Source/DataflowEditor/Private/Dataflow/DataflowEditorModule.cpp
--- a/Engine/Plugins/Experimental/Dataflow/Source/DataflowEditor/Private/Dataflow/DataflowEditorModule.cpp
+++ b/Engine/Plugins/Experimental/Dataflow/Source/DataflowEditor/Private/Dataflow/DataflowEditorModule.cpp
@@ -2,15 +2,17 @@
 
 #include "Dataflow/DataflowEditorModule.h"
 #include "Dataflow/DataflowEditorStyle.h"
+#include "Dataflow/DataflowEditorMode.h"
 
 #include "AssetToolsModule.h"
 #include "CoreMinimal.h"
 #include "EdGraphUtilities.h"
 #include "Dataflow/DataflowEditorToolkit.h"
 #include "Dataflow/DataflowNodeFactory.h"
-#include "Dataflow/DataflowAssetActions.h"
+#include "Dataflow/AssetDefinition_DataflowAsset.h"
 #include "Dataflow/DataflowSNodeFactories.h"
 #include "PropertyEditorModule.h"
+#include "EditorModeRegistry.h"
 
 #define LOCTEXT_NAMESPACE "DataflowEditor"
 
@@ -21,12 +23,6 @@ void FDataflowEditorModule::StartupModule()
 {
 	FDataflowEditorStyle::Get();
 
-	DataflowAssetActions = new FDataflowAssetActions();
-
-	FAssetToolsModule& AssetToolsModule = FAssetToolsModule::GetModule();
-	IAssetTools& AssetTools = AssetToolsModule.Get();
-	AssetTools.RegisterAssetTypeActions(MakeShareable(DataflowAssetActions));
-
 	FPropertyEditorModule& PropertyEditorModule = FModuleManager::LoadModuleChecked<FPropertyEditorModule>(TEXT("PropertyEditor"));
 
 	DataflowSNodeFactory = MakeShareable(new FDataflowSNodeFactory());
@@ -37,23 +33,12 @@ void FDataflowEditorModule::ShutdownModule()
 {
 	if (UObjectInitialized())
 	{
-		FAssetToolsModule& AssetToolsModule = FAssetToolsModule::GetModule();
-		IAssetTools& AssetTools = AssetToolsModule.Get();
-
-		AssetTools.UnregisterAssetTypeActions(DataflowAssetActions->AsShared());
-
 		FEdGraphUtilities::UnregisterVisualNodeFactory(DataflowSNodeFactory);
 	}
-}
+	FEditorModeRegistry::Get().UnregisterMode(UDataflowEditorMode::EM_DataflowEditorModeId);
 
-TSharedRef<FAssetEditorToolkit> FDataflowEditorModule::CreateDataflowAssetEditor(const EToolkitMode::Type Mode, const TSharedPtr<IToolkitHost>& InitToolkitHost, UObject* Asset)
-{
-	TSharedPtr<FDataflowEditorToolkit> NewDataflowAssetEditor = MakeShared<FDataflowEditorToolkit>();
-	NewDataflowAssetEditor->InitializeEditor(Mode, InitToolkitHost, Asset);
-	return StaticCastSharedPtr<FAssetEditorToolkit>(NewDataflowAssetEditor).ToSharedRef();
 }
 
-
 IMPLEMENT_MODULE(FDataflowEditorModule, DataflowEditor)
 
 
diff --git a/Engine/Plugins/Experimental/Dataflow/Source/DataflowEditor/Private/Dataflow/DataflowEditorToolkit.cpp b/Engine/Plugins/Experimental/Dataflow/Source/DataflowEditor/Private/Dataflow/DataflowEditorToolkit.cpp
--- a/Engine/Plugins/Experimental/Dataflow/Source/DataflowEditor/Private/Dataflow/DataflowEditorToolkit.cpp
+++ b/Engine/Plugins/Experimental/Dataflow/Source/DataflowEditor/Private/Dataflow/DataflowEditorToolkit.cpp
@@ -2,10 +2,14 @@
 
 #include "Dataflow/DataflowEditorToolkit.h"
 
+#include "AdvancedPreviewScene.h"
 #include "Animation/Skeleton.h"
 #include "Dataflow/DataflowCore.h"
+#include "Dataflow/DataflowEditor.h"
 #include "Dataflow/DataflowEditorCommands.h"
+#include "Dataflow/DataflowEditorMode.h"
 #include "Dataflow/DataflowEditorViewport.h"
+#include "Dataflow/DataflowEditorViewportClient.h"
 #include "Dataflow/DataflowEdNode.h"
 #include "Dataflow/DataflowGraphEditor.h"
 #include "Dataflow/DataflowNodeFactory.h"
@@ -13,6 +17,7 @@
 #include "Dataflow/DataflowObjectInterface.h"
 #include "Dataflow/DataflowSchema.h"
 #include "DynamicMeshBuilder.h"
+#include "EditorModeManager.h"
 #include "EditorStyleSet.h"
 #include "EditorViewportTabContent.h"
 #include "EditorViewportLayout.h"
@@ -31,7 +36,7 @@
 
 #define LOCTEXT_NAMESPACE "DataflowEditorToolkit"
 
-//DEFINE_LOG_CATEGORY_STATIC(FDataflowEditorToolkitLog, Log, All);
+//DEFINE_LOG_CATEGORY_STATIC(EditorToolkitLog, Log, All);
 
 const FName FDataflowEditorToolkit::ViewportTabId(TEXT("DataflowEditor_Viewport"));
 const FName FDataflowEditorToolkit::GraphCanvasTabId(TEXT("DataflowEditor_GraphCanvas"));
@@ -47,54 +52,90 @@ const FName FDataflowEditorToolkit::CollectionSpreadSheetTabId_2(TEXT("DataflowE
 const FName FDataflowEditorToolkit::CollectionSpreadSheetTabId_3(TEXT("DataflowEditor_CollectionSpreadSheet_3"));
 const FName FDataflowEditorToolkit::CollectionSpreadSheetTabId_4(TEXT("DataflowEditor_CollectionSpreadSheet_4"));
 
-UDataflow* GetDataflowFrom(UObject* InObject)
+namespace Private
 {
-	if (UClass* Class = InObject->GetClass())
+	UDataflow* GetDataflowFrom(UObject* InObject)
 	{
-		if (FProperty* Property = Class->FindPropertyByName(FName("DataflowAsset")))
+		if (UClass* Class = InObject->GetClass())
 		{
-			return *Property->ContainerPtrToValuePtr<UDataflow*>(InObject);
+			if (FProperty* Property = Class->FindPropertyByName(FName("DataflowAsset")))
+			{
+				return *Property->ContainerPtrToValuePtr<UDataflow*>(InObject);
+			}
 		}
-	}
-	return nullptr;
+		return nullptr;
 
-}
+	}
 
-USkeletalMesh* GetSkeletalMeshFrom(UObject* InObject)
-{
-	if (UClass* Class = InObject->GetClass())
+	USkeletalMesh* GetSkeletalMeshFrom(UObject* InObject)
 	{
-		if (FProperty* Property = Class->FindPropertyByName(FName("SkeletalMesh")))
+		if (UClass* Class = InObject->GetClass())
 		{
-			return *Property->ContainerPtrToValuePtr<USkeletalMesh*>(InObject);
+			if (FProperty* Property = Class->FindPropertyByName(FName("SkeletalMesh")))
+			{
+				return *Property->ContainerPtrToValuePtr<USkeletalMesh*>(InObject);
+			}
 		}
-	}
-	return nullptr;
+		return nullptr;
 
-}
+	}
 
-FString GetDataflowTerminalFrom(UObject* InObject)
-{
-	if (UClass* Class = InObject->GetClass())
+	FString GetDataflowTerminalFrom(UObject* InObject)
 	{
-		if (FProperty* Property = Class->FindPropertyByName(FName("DataflowTerminal")))
+		if (UClass* Class = InObject->GetClass())
 		{
-			return *Property->ContainerPtrToValuePtr<FString>(InObject);
+			if (FProperty* Property = Class->FindPropertyByName(FName("DataflowTerminal")))
+			{
+				return *Property->ContainerPtrToValuePtr<FString>(InObject);
+			}
 		}
+		return FString();
 	}
-	return FString();
-}
+};
 
-bool FDataflowEditorToolkit::CanOpenDataflowEditor(UObject* ObjectToEdit)
+
+FDataflowEditorToolkit::FDataflowEditorToolkit(UAssetEditor* InOwningAssetEditor)
+	: FBaseCharacterFXEditorToolkit(InOwningAssetEditor, FName("DataflowEditor")) 
 {
-	UDataflow* Dataflow = Cast<UDataflow>(ObjectToEdit);
-	if (Dataflow == nullptr)
-	{
-		Dataflow = GetDataflowFrom(ObjectToEdit);
-	}
-	return Dataflow != nullptr;
+	check(Cast<UDataflowEditor>(InOwningAssetEditor));
+
+	StandaloneDefaultLayout = FTabManager::NewLayout("Dataflow_Layout.V3")
+		->AddArea
+		(
+			FTabManager::NewPrimaryArea()->SetOrientation(Orient_Vertical)
+			->Split
+			(
+				FTabManager::NewSplitter()->SetOrientation(Orient_Vertical)
+				->Split
+				(
+					FTabManager::NewSplitter()->SetOrientation(Orient_Horizontal)
+					->SetSizeCoefficient(0.9f)
+					->Split
+					(
+						FTabManager::NewStack()
+						->SetSizeCoefficient(0.6f)
+						->AddTab(GraphCanvasTabId, ETabState::OpenedTab)
+					)
+					->Split
+					(
+						FTabManager::NewSplitter()->SetOrientation(Orient_Vertical)
+						->SetSizeCoefficient(0.2f)
+						->Split
+						(
+							FTabManager::NewStack()
+							->SetSizeCoefficient(0.7f)
+							->AddTab(AssetDetailsTabId, ETabState::OpenedTab)
+						)
+					)
+				)
+			)
+		);
+
+	// Construction view scene @todo(DataflowMode) : Add a FDataflowPreviewScene to the Toolkit here
+	// ObjectScene = MakeUnique<FAdvancedPreviewScene>(FAdvancedPreviewScene::ConstructionValues().SetCreateDefaultLighting(false));
 }
 
+
 FDataflowEditorToolkit::~FDataflowEditorToolkit()
 {
 	if (GraphEditor)
@@ -110,22 +151,157 @@ FDataflowEditorToolkit::~FDataflowEditorToolkit()
 
 	if (AssetDetailsEditor)
 	{
-		 AssetDetailsEditor->OnFinishedChangingProperties().Remove(OnFinishedChangingAssetPropertiesDelegateHandle);
+		AssetDetailsEditor->OnFinishedChangingProperties().Remove(OnFinishedChangingAssetPropertiesDelegateHandle);
 	}
+
+	// We need to force the dataflow editor mode deletion now because otherwise the preview and rest-space worlds
+	// will end up getting destroyed before the mode's Exit() function gets to run, and we'll get some
+	// warnings when we destroy any mode actors.
+	EditorModeManager->DestroyMode(UDataflowEditorMode::EM_DataflowEditorModeId);
 }
 
-void FDataflowEditorToolkit::InitializeEditor(const EToolkitMode::Type Mode, const TSharedPtr<IToolkitHost>& InitToolkitHost, UObject* ObjectToEdit)
+bool FDataflowEditorToolkit::CanOpenDataflowEditor(UObject* ObjectToEdit)
 {
-	Asset = nullptr;
-	Dataflow = Cast<UDataflow>(ObjectToEdit);
+	UDataflow* Dataflow = Cast<UDataflow>(ObjectToEdit);
 	if (Dataflow == nullptr)
 	{
-		Dataflow = GetDataflowFrom(ObjectToEdit);
-		if (Dataflow)
+		Dataflow = Private::GetDataflowFrom(ObjectToEdit);
+	}
+	return Dataflow != nullptr;
+}
+
+//~ Begin FBaseCharacterFXEditorToolkit overrides
+
+FEditorModeID FDataflowEditorToolkit::GetEditorModeId() const
+{
+	return UDataflowEditorMode::EM_DataflowEditorModeId;
+}
+
+bool FDataflowEditorToolkit::OnRequestClose(EAssetEditorCloseReason InCloseReason)
+{
+	// Note: This needs a bit of adjusting, because currently OnRequestClose seems to be 
+	// called multiple times when the editor itself is being closed. We can take the route 
+	// of NiagaraScriptToolkit and remember when changes are discarded, but this can cause
+	// issues if the editor close sequence is interrupted due to some other asset editor.
+
+	UDataflowEditorMode* DataflowEdMode = Cast<UDataflowEditorMode>(EditorModeManager->GetActiveScriptableMode(UDataflowEditorMode::EM_DataflowEditorModeId));
+	if (!DataflowEdMode) {
+		// If we don't have a valid mode, because the OnRequestClose is currently being called multiple times,
+		// simply return true because there's nothing left to do.
+		return true;
+	}
+
+	// Give any active modes a chance to shutdown while the toolkit host is still alive
+	// This is super important to do, otherwise currently opened tabs won't be marked as "closed".
+	// This results in tabs not being properly recycled upon reopening the editor and tab
+	// duplication for each opening event.
+	GetEditorModeManager().ActivateDefaultMode();
+
+	return FAssetEditorToolkit::OnRequestClose(InCloseReason);
+}
+
+void FDataflowEditorToolkit::PostInitAssetEditor()
+{
+	// @todo(DataflowMode) : Enable base class functionality when the mode is ready.
+	// FBaseCharacterFXEditorToolkit::PostInitAssetEditor();
+	
+	// @todo(DataflowMode) : Access to mode in Toolkit construction
+	//UDataflowEditorMode* const DataflowEditorMode = CastChecked<UDataflowEditorMode>(EditorModeManager->GetActiveScriptableMode(UDataflowEditorMode::EM_DataflowEditorModeId));
+	
+	auto SetCommonViewportClientOptions = [](FEditorViewportClient* Client)
+	{
+		// Normally the bIsRealtime flag is determined by whether the connection is remote, but our
+		// tools require always being ticked.
+		Client->SetRealtime(true);
+
+		// Disable motion blur effects that cause our renders to "fade in" as things are moved
+		Client->EngineShowFlags.SetTemporalAA(false);
+		Client->EngineShowFlags.SetAntiAliasing(true);
+		Client->EngineShowFlags.SetMotionBlur(false);
+
+		// Disable the dithering of occluded portions of gizmos.
+		Client->EngineShowFlags.SetOpaqueCompositeEditorPrimitives(true);
+
+		// Disable hardware occlusion queries, which make it harder to use vertex shaders to pull materials
+		// toward camera for z ordering because non-translucent materials start occluding themselves (once
+		// the component bounds are behind the displaced geometry).
+		Client->EngineShowFlags.SetDisableOcclusionQueries(true);
+
+		// Default FOV of 90 degrees causes a fair bit of lens distortion, especially noticeable with smaller viewports
+		Client->ViewFOV = 45.0;
+
+		// Ortho has too many problems with rendering things, unfortunately, so we should use perspective.
+		Client->SetViewportType(ELevelViewportType::LVT_Perspective);
+
+		// Lit gives us the most options in terms of the materials we can use.
+		Client->SetViewMode(EViewModeIndex::VMI_Lit);
+
+		// If exposure isn't set to fixed, it will flash as we stare into the void
+		Client->ExposureSettings.bFixed = true;
+	};
+	SetCommonViewportClientOptions(ViewportClient.Get());
+
+	// Set up 3D viewport
+	TSharedPtr<FDataflowEditorViewportClient> DataflowViewportClient = StaticCastSharedPtr<FDataflowEditorViewportClient>(ViewportClient);
+	DataflowViewportClient->SetDataflowEditorToolkit(StaticCastSharedRef<FDataflowEditorToolkit>(this->AsShared()));
+
+	// @todo(DataflowMode) : Do this in UDataflowEditorMode::RefocusViewportClient and use the DataflowComponents bounds
+	//		... see ChaosClothAssetEditorMode::RefocusRestSpaceViewportClient() for reference
+	// FBoxSphereBounds SphereBounds = FBoxSphereBounds(EForceInit::ForceInitToZero);
+	// DataflowViewportClient->OverrideFarClipPlane(0);
+	// DataflowViewportClient->FocusViewportOnBox(SphereBounds.GetBox());
+}
+
+void FDataflowEditorToolkit::InitializeEdMode(UBaseCharacterFXEditorMode* EdMode)
+{
+	UDataflowEditorMode* DataflowMode = Cast<UDataflowEditorMode>(EdMode);
+	check(DataflowMode);
+}
+
+void FDataflowEditorToolkit::CreateEditorModeUILayer()
+{
+	FBaseCharacterFXEditorToolkit::CreateEditorModeUILayer();
+}
+
+void FDataflowEditorToolkit::GetSaveableObjects(TArray<UObject*>& OutObjects) const
+{
+	FBaseCharacterFXEditorToolkit::GetSaveableObjects(OutObjects);
+
+	if (Dataflow)
+	{
+		check(Dataflow->IsAsset());
+		OutObjects.Add(Dataflow);
+	}
+}
+
+//~ End FBaseCharacterFXEditorToolkit overrides
+
+
+//~ Begin FBaseAssetToolkit overrides
+
+void FDataflowEditorToolkit::CreateWidgets()
+{
+	FBaseCharacterFXEditorToolkit::CreateWidgets();
+
+	Asset = nullptr;
+	UObject* const ObjectToEdit = GetAsset();
+
+	if (ObjectToEdit)
+	{
+		Dataflow = Cast<UDataflow>(ObjectToEdit);
+		if (Dataflow == nullptr)
 		{
-			Asset = ObjectToEdit;
-			TerminalPath = GetDataflowTerminalFrom(ObjectToEdit);
+			Dataflow = Private::GetDataflowFrom(ObjectToEdit);
+
+			// TODO: Figure out how to create the GraphEditor widgets when the Input Asset doesn't have a Dataflow property set
+			if (Dataflow)
+			{
+				DataflowTerminalPath = Private::GetDataflowTerminalFrom(Dataflow);
+
+				Dataflow->Schema = UDataflowSchema::StaticClass();
+			}
 		}
+		Asset = ObjectToEdit;
 	}
 
 	if (Dataflow != nullptr)
@@ -139,52 +315,37 @@ void FDataflowEditorToolkit::InitializeEditor(const EToolkitMode::Type Mode, con
 		AssetDetailsEditor = CreateAssetDetailsEditorWidget(ObjectToEdit);
 		GraphEditor = CreateGraphEditorWidget(Dataflow, NodeDetailsEditor);
 		SkeletalEditor = CreateSkeletalEditorWidget(ObjectToEdit);
+	}
+}
 
-		const TSharedRef<FTabManager::FLayout> StandaloneDefaultLayout = FTabManager::NewLayout("Dataflow_Layout.V1")
-			->AddArea
-			(
-				FTabManager::NewPrimaryArea()->SetOrientation(Orient_Vertical)
-				->Split
-				(
-					FTabManager::NewSplitter()->SetOrientation(Orient_Vertical)
-					->Split
-					(
-						FTabManager::NewSplitter()->SetOrientation(Orient_Horizontal)
-						->SetSizeCoefficient(0.9f)
-						->Split
-						(
-							FTabManager::NewStack()
-							->SetSizeCoefficient(0.6f)
-							->AddTab(GraphCanvasTabId, ETabState::OpenedTab)
-						)
-						->Split
-						(
-							FTabManager::NewSplitter()->SetOrientation(Orient_Vertical)
-							->SetSizeCoefficient(0.2f)
-							->Split
-							(
-								FTabManager::NewStack()
-								->SetSizeCoefficient(0.7f)
-								->AddTab(AssetDetailsTabId, ETabState::OpenedTab)
-							)
-						)
-					)
-				)
-			);
-
-		const bool bCreateDefaultStandaloneMenu = true;
-		const bool bCreateDefaultToolbar = true;
-		FAssetEditorToolkit::InitAssetEditor(Mode, InitToolkitHost, FName(TEXT("DataflowEditorApp")), StandaloneDefaultLayout, bCreateDefaultToolbar, bCreateDefaultStandaloneMenu, ObjectToEdit);
+// Called from FBaseAssetToolkit::CreateWidgets. The delegate call path goes through FAssetEditorToolkit::InitAssetEditor
+// and FBaseAssetToolkit::SpawnTab_Viewport.
+AssetEditorViewportFactoryFunction FDataflowEditorToolkit::GetViewportDelegate()
+{
+	AssetEditorViewportFactoryFunction TempViewportDelegate = [this](FAssetEditorViewportConstructionArgs InArgs)
+	{
+		return SAssignNew(DataflowEditorViewport, SDataflowEditorViewport, InArgs)
+			.ViewportClient(StaticCastSharedPtr<FDataflowEditorViewportClient>(ViewportClient));
+	};
 
+	return TempViewportDelegate;
+}
 
-		AddEditingObject(Dataflow);
-		if (Asset) 
-		{
-			AddEditingObject(Asset);
-		}
-	}
+// Called from FBaseAssetToolkit::CreateWidgets to populate ViewportClient, but otherwise only used 
+// in our own viewport delegate.
+TSharedPtr<FEditorViewportClient> FDataflowEditorToolkit::CreateEditorViewportClient() const
+{
+	// Note that we can't reliably adjust the viewport client here because we will be passing it
+	// into the viewport created by the viewport delegate we get from GetViewportDelegate(), and
+	// that delegate may (will) affect the settings based on FAssetEditorViewportConstructionArgs,
+	// namely ViewportType.
+	// Instead, we do viewport client adjustment in PostInitAssetEditor().
+	check(EditorModeManager.IsValid());
+	return MakeShared<FDataflowEditorViewportClient>(EditorModeManager.Get(), ObjectScene.Get());
 }
 
+//~ End FBaseAssetToolkit overrides
+
 void FDataflowEditorToolkit::OnPropertyValueChanged(const FPropertyChangedEvent& PropertyChangedEvent)
 {
 	FDataflowEditorCommands::OnPropertyValueChanged(this->GetDataflow(), Context, LastNodeTimestamp, PropertyChangedEvent, PrevNodeSelection);
@@ -282,8 +443,8 @@ void FDataflowEditorToolkit::Tick(float DeltaTime)
 			Context = TSharedPtr< Dataflow::FEngineContext>(new Dataflow::FAssetContext(Asset, Dataflow, Dataflow::FTimestamp::Invalid));
 			LastNodeTimestamp = Dataflow::FTimestamp::Invalid;
 		}
-		TerminalPath = GetDataflowTerminalFrom(Asset);
-		FDataflowEditorCommands::EvaluateTerminalNode(*Context.Get(), LastNodeTimestamp, Dataflow, nullptr, nullptr, Asset, TerminalPath);
+		DataflowTerminalPath = Private::GetDataflowTerminalFrom(Asset);
+		FDataflowEditorCommands::EvaluateTerminalNode(*Context.Get(), LastNodeTimestamp, Dataflow, nullptr, nullptr, Asset, DataflowTerminalPath);
 	}
 }
 
@@ -292,6 +453,8 @@ TStatId FDataflowEditorToolkit::GetStatId() const
 	RETURN_QUICK_DECLARE_CYCLE_STAT(FDataflowEditorToolkit, STATGROUP_Tickables);
 }
 
+
+
 TSharedRef<SDataflowGraphEditor> FDataflowEditorToolkit::CreateGraphEditorWidget(UDataflow* DataflowToEdit, TSharedPtr<IStructureDetailsView> InNodeDetailsEditor)
 {
 	ensure(DataflowToEdit);
@@ -306,7 +469,7 @@ TSharedRef<SDataflowGraphEditor> FDataflowEditorToolkit::CreateGraphEditorWidget
 		Node->Invalidate();
 		LastNodeTimestamp = Dataflow::FTimestamp::Invalid;
 
-		FDataflowEditorCommands::EvaluateTerminalNode(*Context.Get(), LastNodeTimestamp, Dataflow, Node, Out, Asset, TerminalPath);
+		FDataflowEditorCommands::EvaluateTerminalNode(*Context.Get(), LastNodeTimestamp, Dataflow, Node, Out, Asset, DataflowTerminalPath);
 	};
 
 	SGraphEditor::FGraphEditorEvents InEvents;
@@ -350,11 +513,11 @@ TSharedPtr<IStructureDetailsView> FDataflowEditorToolkit::CreateNodeDetailsEdito
 		StructureViewArgs.bShowClasses = true;
 		StructureViewArgs.bShowInterfaces = true;
 	}
-	TSharedPtr<IStructureDetailsView> DetailsView = PropertyEditorModule.CreateStructureDetailView(DetailsViewArgs, StructureViewArgs, nullptr);
-	DetailsView->GetDetailsView()->SetObject(ObjectToEdit);
-	OnFinishedChangingPropertiesDelegateHandle = DetailsView->GetOnFinishedChangingPropertiesDelegate().AddSP(this, &FDataflowEditorToolkit::OnPropertyValueChanged);
+	TSharedPtr<IStructureDetailsView> LocalDetailsView = PropertyEditorModule.CreateStructureDetailView(DetailsViewArgs, StructureViewArgs, nullptr);
+	LocalDetailsView->GetDetailsView()->SetObject(ObjectToEdit);
+	OnFinishedChangingPropertiesDelegateHandle = LocalDetailsView->GetOnFinishedChangingPropertiesDelegate().AddSP(this, &FDataflowEditorToolkit::OnPropertyValueChanged);
 
-	return DetailsView;
+	return LocalDetailsView;
 }
 
 TSharedPtr<IDetailsView> FDataflowEditorToolkit::CreateAssetDetailsEditorWidget(UObject* ObjectToEdit)
@@ -371,12 +534,12 @@ TSharedPtr<IDetailsView> FDataflowEditorToolkit::CreateAssetDetailsEditorWidget(
 		DetailsViewArgs.NotifyHook = this;
 	}
 
-	TSharedPtr<IDetailsView> DetailsView = PropertyEditorModule.CreateDetailView(DetailsViewArgs);
-	DetailsView->SetObject(ObjectToEdit);
+	TSharedPtr<IDetailsView> LocalDetailsView = PropertyEditorModule.CreateDetailView(DetailsViewArgs);
+	LocalDetailsView->SetObject(ObjectToEdit);
 
-	OnFinishedChangingAssetPropertiesDelegateHandle = DetailsView->OnFinishedChangingProperties().AddSP(this, &FDataflowEditorToolkit::OnAssetPropertyValueChanged);
+	OnFinishedChangingAssetPropertiesDelegateHandle = LocalDetailsView->OnFinishedChangingProperties().AddSP(this, &FDataflowEditorToolkit::OnAssetPropertyValueChanged);
 
-	return DetailsView;
+	return LocalDetailsView;
 
 }
 
@@ -401,7 +564,7 @@ TSharedPtr<ISkeletonTree> FDataflowEditorToolkit::CreateSkeletalEditorWidget(UOb
 		USkeleton* Skeleton = StubSkeleton;
 		if (Asset)
 		{
-			if (USkeletalMesh* SkeletalMesh = GetSkeletalMeshFrom(Asset))
+			if (USkeletalMesh* SkeletalMesh = Private::GetSkeletalMeshFrom(Asset))
 			{
 				Skeleton = SkeletalMesh->GetSkeleton();
 			}
@@ -413,22 +576,18 @@ TSharedPtr<ISkeletonTree> FDataflowEditorToolkit::CreateSkeletalEditorWidget(UOb
 	return TSharedPtr<ISkeletonTree>(nullptr);
 }
 
-TSharedRef<SDockTab> FDataflowEditorToolkit::SpawnTab_Viewport(const FSpawnTabArgs& Args)
+UObject* FDataflowEditorToolkit::GetAsset() const
 {
-	check(Args.GetTabId() == ViewportTabId);
-
-	TSharedRef< SDockTab > DockableTab = SNew(SDockTab);
-	ViewportEditor = MakeShareable(new FEditorViewportTabContent());
-	TWeakPtr<FDataflowEditorToolkit> WeakSharedThis = SharedThis(this);
+	TArray<TObjectPtr<UObject>> ObjectsToEdit;
+	OwningAssetEditor->GetObjectsToEdit(MutableView(ObjectsToEdit));
 
-	const FString LayoutId = FString("DataflowEditorViewport");
-	ViewportEditor->Initialize([WeakSharedThis](const FAssetEditorViewportConstructionArgs& InConstructionArgs)
-		{
-			return SNew(SDataflowEditorViewport)
-				.DataflowEditorToolkit(WeakSharedThis);
-		}, DockableTab, LayoutId);
+	UObject* ObjectToEdit = nullptr;
+	if (ensure(ObjectsToEdit.Num() == 1))
+	{
+		ObjectToEdit = ObjectsToEdit[0];
+	}
 
-	return DockableTab;
+	return Cast<UObject>(ObjectToEdit);
 }
 
 TSharedRef<SDockTab> FDataflowEditorToolkit::SpawnTab_GraphCanvas(const FSpawnTabArgs& Args)
@@ -473,7 +632,7 @@ TSharedRef<SDockTab> FDataflowEditorToolkit::SpawnTab_Skeletal(const FSpawnTabAr
 	USkeletalMesh* SkeletalMesh = StubSkeletalMesh;
 	if (Asset)
 	{
-		SkeletalMesh = GetSkeletalMeshFrom(Asset);
+		SkeletalMesh = Private::GetSkeletalMeshFrom(Asset);
 	}
 
 	SkeletalEditor->SetSkeletalMesh(SkeletalMesh);
@@ -674,33 +833,37 @@ TSharedRef<SDockTab> FDataflowEditorToolkit::SpawnTab_CollectionSpreadSheet(cons
 
 void FDataflowEditorToolkit::RegisterTabSpawners(const TSharedRef<FTabManager>& InTabManager)
 {
-	TSharedRef<FWorkspaceItem> WorkspaceMenuCategoryRef = InTabManager->AddLocalWorkspaceMenuCategory(LOCTEXT("WorkspaceMenu_DataflowEditor", "Dataflow Editor"));
-	TSharedRef<FWorkspaceItem> SelectionViewWorkspaceMenuCategoryRef = WorkspaceMenuCategoryRef->AddGroup(LOCTEXT("WorkspaceMenu_SelectionView", "Selection View"), FSlateIcon(FAppStyle::GetAppStyleSetName(), "LevelEditor.Tabs.Outliner"));
-	TSharedRef<FWorkspaceItem> CollectionSpreadSheetWorkspaceMenuCategoryRef = WorkspaceMenuCategoryRef->AddGroup(LOCTEXT("WorkspaceMenu_CollectionSpreadSheet", "Collection SpreadSheet"), FSlateIcon(FAppStyle::GetAppStyleSetName(), "LevelEditor.Tabs.Outliner"));
+	// We bypass FBaseAssetToolkit::RegisterTabSpawners because it doesn't seem to provide us with
+	// anything except tabs that we don't want.
+	FAssetEditorToolkit::RegisterTabSpawners(InTabManager);
+
+	EditorMenuCategory = InTabManager->AddLocalWorkspaceMenuCategory(LOCTEXT("WorkspaceMenu_DataflowEditor", "Dataflow Editor"));
+	TSharedRef<FWorkspaceItem> SelectionViewWorkspaceMenuCategoryRef = EditorMenuCategory->AddGroup(LOCTEXT("WorkspaceMenu_SelectionView", "Selection View"), FSlateIcon(FAppStyle::GetAppStyleSetName(), "LevelEditor.Tabs.Outliner"));
+	TSharedRef<FWorkspaceItem> CollectionSpreadSheetWorkspaceMenuCategoryRef = EditorMenuCategory->AddGroup(LOCTEXT("WorkspaceMenu_CollectionSpreadSheet", "Collection SpreadSheet"), FSlateIcon(FAppStyle::GetAppStyleSetName(), "LevelEditor.Tabs.Outliner"));
 
 	InTabManager->RegisterTabSpawner(ViewportTabId, FOnSpawnTab::CreateSP(this, &FDataflowEditorToolkit::SpawnTab_Viewport))
 		.SetDisplayName(LOCTEXT("DataflowViewportTab", "Dataflow Viewport"))
-		.SetGroup(WorkspaceMenuCategoryRef)
+		.SetGroup(EditorMenuCategory.ToSharedRef())
 		.SetIcon(FSlateIcon(FAppStyle::GetAppStyleSetName(), "GraphEditor.EventGraph_16x"));
 
 	InTabManager->RegisterTabSpawner(GraphCanvasTabId, FOnSpawnTab::CreateSP(this, &FDataflowEditorToolkit::SpawnTab_GraphCanvas))
 		.SetDisplayName(LOCTEXT("DataflowTab", "Graph"))
-		.SetGroup(WorkspaceMenuCategoryRef)
+		.SetGroup(EditorMenuCategory.ToSharedRef())
 		.SetIcon(FSlateIcon(FAppStyle::GetAppStyleSetName(), "GraphEditor.EventGraph_16x"));
 
 	InTabManager->RegisterTabSpawner(AssetDetailsTabId, FOnSpawnTab::CreateSP(this, &FDataflowEditorToolkit::SpawnTab_AssetDetails))
 		.SetDisplayName(LOCTEXT("AssetDetailsTab", "Details"))
-		.SetGroup(WorkspaceMenuCategoryRef)
+		.SetGroup(EditorMenuCategory.ToSharedRef())
 		.SetIcon(FSlateIcon(FAppStyle::GetAppStyleSetName(), "LevelEditor.Tabs.Details"));
 
 	InTabManager->RegisterTabSpawner(NodeDetailsTabId, FOnSpawnTab::CreateSP(this, &FDataflowEditorToolkit::SpawnTab_NodeDetails))
 		.SetDisplayName(LOCTEXT("NodeDetailsTab", "Node Details"))
-		.SetGroup(WorkspaceMenuCategoryRef)
+		.SetGroup(EditorMenuCategory.ToSharedRef())
 		.SetIcon(FSlateIcon(FAppStyle::GetAppStyleSetName(), "LevelEditor.Tabs.Details"));
 
 	InTabManager->RegisterTabSpawner(SkeletalTabId, FOnSpawnTab::CreateSP(this, &FDataflowEditorToolkit::SpawnTab_Skeletal))
 		.SetDisplayName(LOCTEXT("DataflowSkeletalTab", "Skeletal Hierarchy"))
-		.SetGroup(WorkspaceMenuCategoryRef)
+		.SetGroup(EditorMenuCategory.ToSharedRef())
 		.SetIcon(FSlateIcon(FAppStyle::GetAppStyleSetName(), "LevelEditor.Tabs.SkeletalHierarchy"));
 
 	InTabManager->RegisterTabSpawner(SelectionViewTabId_1, FOnSpawnTab::CreateSP(this, &FDataflowEditorToolkit::SpawnTab_SelectionView))
@@ -743,7 +906,6 @@ void FDataflowEditorToolkit::RegisterTabSpawners(const TSharedRef<FTabManager>&
 		.SetGroup(CollectionSpreadSheetWorkspaceMenuCategoryRef)
 		.SetIcon(FSlateIcon(FAppStyle::GetAppStyleSetName(), "LevelEditor.Tabs.Outliner"));
 
-	FAssetEditorToolkit::RegisterTabSpawners(InTabManager);
 }
 
 void FDataflowEditorToolkit::UnregisterTabSpawners(const TSharedRef<class FTabManager>& InTabManager)
@@ -839,20 +1001,4 @@ FLinearColor FDataflowEditorToolkit::GetWorldCentricTabColorScale() const
 	return FLinearColor(0.3f, 0.2f, 0.5f, 0.5f);
 }
 
-FString FDataflowEditorToolkit::GetReferencerName() const
-{
-	return TEXT("DataflowEditorToolkit");
-}
-
-void FDataflowEditorToolkit::AddReferencedObjects(FReferenceCollector& Collector)
-{
-	if (Dataflow)
-	{
-		Collector.AddReferencedObject(Dataflow);
-	}
-	if (Asset)
-	{
-		Collector.AddReferencedObject(Asset);
-	}
-}
 #undef LOCTEXT_NAMESPACE
diff --git a/Engine/Plugins/Experimental/Dataflow/Source/DataflowEditor/Private/Dataflow/DataflowEditorViewport.cpp b/Engine/Plugins/Experimental/Dataflow/Source/DataflowEditor/Private/Dataflow/DataflowEditorViewport.cpp
--- a/Engine/Plugins/Experimental/Dataflow/Source/DataflowEditor/Private/Dataflow/DataflowEditorViewport.cpp
+++ b/Engine/Plugins/Experimental/Dataflow/Source/DataflowEditor/Private/Dataflow/DataflowEditorViewport.cpp
@@ -1,97 +1,67 @@
 // Copyright Epic Games, Inc. All Rights Reserved.
+
 #include "Dataflow/DataflowEditorViewport.h"
 
 #include "AdvancedPreviewScene.h"
 #include "Dataflow/DataflowActor.h"
 #include "Dataflow/DataflowEditorCommands.h"
-#include "Dataflow/DataflowEditorViewportToolbar.h"
+#include "Dataflow/DataflowEditorMode.h"
 #include "Dataflow/DataflowEditorViewportClient.h"
-#include "EditorViewportClient.h"
-#include "Framework/Application/SlateApplication.h"
-#include "PreviewScene.h"
+#include "EditorModeManager.h"
+#include "SViewportToolBar.h"
+
+#define LOCTEXT_NAMESPACE "SDataflowEditorViewport"
+
 
 SDataflowEditorViewport::SDataflowEditorViewport()
 {
-	PreviewScene = MakeShareable(new FAdvancedPreviewScene(FPreviewScene::ConstructionValues()));
-	PreviewScene->SetFloorVisibility(false);
 }
 
-void SDataflowEditorViewport::Construct(const FArguments& InArgs)
+void SDataflowEditorViewport::Construct(const FArguments& InArgs, const FAssetEditorViewportConstructionArgs& InViewportConstructionArgs)
 {
-	DataflowEditorToolkitPtr = InArgs._DataflowEditorToolkit;
-	TSharedPtr<FDataflowEditorToolkit> DataflowEditorToolkit = DataflowEditorToolkitPtr.Pin();
-	check(DataflowEditorToolkitPtr.IsValid());
-
-	SEditorViewport::Construct(SEditorViewport::FArguments());
-
-	FBoxSphereBounds SphereBounds = FBoxSphereBounds(EForceInit::ForceInitToZero);
-	CustomDataflowActor = CastChecked<ADataflowActor>(PreviewScene->GetWorld()->SpawnActor(ADataflowActor::StaticClass()));
-
-	ViewportClient->SetDataflowActor(CustomDataflowActor);
-	ViewportClient->FocusViewportOnBox( SphereBounds.GetBox());
+	SAssetEditorViewport::FArguments ParentArgs;
+	ParentArgs._EditorViewportClient = InArgs._ViewportClient;
+	SAssetEditorViewport::Construct(ParentArgs, InViewportConstructionArgs);
+	Client->VisibilityDelegate.BindSP(this, &SDataflowEditorViewport::IsVisible);
 }
 
-TSharedRef<SEditorViewport> SDataflowEditorViewport::GetViewportWidget()
+UDataflowEditorMode* SDataflowEditorViewport::GetEdMode() const
 {
-	return SharedThis(this);
+	if (const FEditorModeTools* const EditorModeTools = Client->GetModeTools())
+	{
+		if (UDataflowEditorMode* const DataflowEdMode = Cast<UDataflowEditorMode>(EditorModeTools->GetActiveScriptableMode(UDataflowEditorMode::EM_DataflowEditorModeId)))
+		{
+			return DataflowEdMode;
+		}
+	}
+	return nullptr;
 }
 
-
-TSharedPtr<FExtender> SDataflowEditorViewport::GetExtenders() const
+void SDataflowEditorViewport::BindCommands()
 {
-	TSharedPtr<FExtender> Result(MakeShareable(new FExtender));
-	return Result;
+	SAssetEditorViewport::BindCommands();
 }
 
-void SDataflowEditorViewport::OnFloatingButtonClicked()
-{
-}
 
-void SDataflowEditorViewport::AddReferencedObjects(FReferenceCollector& Collector)
+bool SDataflowEditorViewport::IsVisible() const
 {
-	Collector.AddReferencedObject(CustomDataflowActor);
+	// Intentionally not calling SEditorViewport::IsVisible because it will return false if our simulation is more than 250ms.
+	return ViewportWidget.IsValid();
 }
 
-TSharedRef<FEditorViewportClient> SDataflowEditorViewport::MakeEditorViewportClient()
+TSharedRef<class SEditorViewport> SDataflowEditorViewport::GetViewportWidget()
 {
-	ViewportClient = MakeShareable(new FDataflowEditorViewportClient(PreviewScene.Get(), SharedThis(this), DataflowEditorToolkitPtr));
-	return ViewportClient.ToSharedRef();
+	return SharedThis(this);
 }
 
-TSharedPtr<SWidget> SDataflowEditorViewport::MakeViewportToolbar()
+TSharedPtr<FExtender> SDataflowEditorViewport::GetExtenders() const
 {
-	return
-		SNew(SDataflowViewportSelectionToolBar)
-		.EditorViewport(SharedThis(this))
-		.IsEnabled(FSlateApplication::Get().GetNormalExecutionAttribute());
+	TSharedPtr<FExtender> Result(MakeShareable(new FExtender));
+	return Result;
 }
 
-void SDataflowEditorViewport::BindCommands()
+void SDataflowEditorViewport::OnFloatingButtonClicked()
 {
-	SEditorViewport::BindCommands();
-	{
-		const FDataflowEditorCommandsImpl& Commands = FDataflowEditorCommands::Get();
-		TSharedRef<FDataflowEditorViewportClient> ClientRef = ViewportClient.ToSharedRef();
-
-		CommandList->MapAction(
-			Commands.ToggleObjectSelection,
-			FExecuteAction::CreateSP(ClientRef, &FDataflowEditorViewportClient::SetSelectionMode, FDataflowSelectionState::EMode::DSS_Dataflow_Object)
-			,FCanExecuteAction::CreateSP(ClientRef, &FDataflowEditorViewportClient::CanSetSelectionMode, FDataflowSelectionState::EMode::DSS_Dataflow_Object)
-			,FIsActionChecked::CreateSP(ClientRef, &FDataflowEditorViewportClient::IsSelectionModeActive, FDataflowSelectionState::EMode::DSS_Dataflow_Object)
-		);
-
-		CommandList->MapAction(
-			Commands.ToggleFaceSelection,
-			FExecuteAction::CreateSP(ClientRef, &FDataflowEditorViewportClient::SetSelectionMode, FDataflowSelectionState::EMode::DSS_Dataflow_Face)
-			, FCanExecuteAction::CreateSP(ClientRef, &FDataflowEditorViewportClient::CanSetSelectionMode, FDataflowSelectionState::EMode::DSS_Dataflow_Face)
-			, FIsActionChecked::CreateSP(ClientRef, &FDataflowEditorViewportClient::IsSelectionModeActive, FDataflowSelectionState::EMode::DSS_Dataflow_Face)
-		);
-
-		CommandList->MapAction(
-			Commands.ToggleVertexSelection,
-			FExecuteAction::CreateSP(ClientRef, &FDataflowEditorViewportClient::SetSelectionMode, FDataflowSelectionState::EMode::DSS_Dataflow_Vertex)
-			, FCanExecuteAction::CreateSP(ClientRef, &FDataflowEditorViewportClient::CanSetSelectionMode, FDataflowSelectionState::EMode::DSS_Dataflow_Vertex)
-			, FIsActionChecked::CreateSP(ClientRef, &FDataflowEditorViewportClient::IsSelectionModeActive, FDataflowSelectionState::EMode::DSS_Dataflow_Vertex)
-		);
-	}
 }
+
+#undef LOCTEXT_NAMESPACE
diff --git a/Engine/Plugins/Experimental/Dataflow/Source/DataflowEditor/Private/Dataflow/DataflowEditorViewportClient.cpp b/Engine/Plugins/Experimental/Dataflow/Source/DataflowEditor/Private/Dataflow/DataflowEditorViewportClient.cpp
--- a/Engine/Plugins/Experimental/Dataflow/Source/DataflowEditor/Private/Dataflow/DataflowEditorViewportClient.cpp
+++ b/Engine/Plugins/Experimental/Dataflow/Source/DataflowEditor/Private/Dataflow/DataflowEditorViewportClient.cpp
@@ -1,6 +1,8 @@
 // Copyright Epic Games, Inc. All Rights Reserved.
 #include "Dataflow/DataflowEditorViewportClient.h"
 
+#include "AdvancedPreviewScene.h"
+#include "Components/DirectionalLightComponent.h"
 #include "Dataflow/DataflowActor.h"
 #include "Dataflow/DataflowEditorToolkit.h"
 #include "Dataflow/DataflowEngineSceneHitProxies.h"
@@ -8,23 +10,26 @@
 #include "HAL/PlatformApplicationMisc.h"
 #include "PreviewScene.h"
 
-FDataflowEditorViewportClient::FDataflowEditorViewportClient(FPreviewScene* InPreviewScene, 
-	const TWeakPtr<SEditorViewport> InEditorViewportWidget,
-	TWeakPtr<FDataflowEditorToolkit> InDataflowEditorToolkitPtr)
-	: 
-	FEditorViewportClient(nullptr, InPreviewScene, InEditorViewportWidget)
-	, DataflowEditorToolkitPtr(InDataflowEditorToolkitPtr)
+FDataflowEditorViewportClient::FDataflowEditorViewportClient(FEditorModeTools* InModeTools,
+															 FPreviewScene* InPreviewScene,
+															 const TWeakPtr<SEditorViewport> InEditorViewportWidget)
+	: FEditorViewportClient(InModeTools, InPreviewScene, InEditorViewportWidget)
 {
-	SetRealtime(true);
-	SetViewModes(VMI_Lit, VMI_Lit);
-	bSetListenerPosition = false;
-	EngineShowFlags.Grid = false;
+	// @todo(DataflowMode) : Move this to the Mode as a Component only (see UChaosClothAssetEditorMode)
+	DataflowActor = CastChecked<ADataflowActor>(PreviewScene->GetWorld()->SpawnActor(ADataflowActor::StaticClass()));
+	((FAdvancedPreviewScene*)InPreviewScene)->SetFloorVisibility(false, true);
 }
 
+void FDataflowEditorViewportClient::SetDataflowEditorToolkit(TWeakPtr<FDataflowEditorToolkit> InDataflowEditorToolkitPtr)
+{
+	DataflowEditorToolkitPtr = InDataflowEditorToolkitPtr;
+}
+
+
 void FDataflowEditorViewportClient::SetSelectionMode(FDataflowSelectionState::EMode InState)
 {
 	FDataflowSelectionState State = DataflowActor->DataflowComponent->GetSelectionState();
-		
+
 	if (SelectionMode == InState)
 	{
 		SelectionMode = FDataflowSelectionState::EMode::DSS_Dataflow_None;
@@ -70,18 +75,18 @@ bool FDataflowEditorViewportClient::CanSetSelectionMode(FDataflowSelectionState:
 
 	return false;
 }
-bool FDataflowEditorViewportClient::IsSelectionModeActive(FDataflowSelectionState::EMode InState) 
-{ 
+bool FDataflowEditorViewportClient::IsSelectionModeActive(FDataflowSelectionState::EMode InState)
+{
 	return SelectionMode == InState;
 }
 
-Dataflow::FTimestamp FDataflowEditorViewportClient::LatestTimestamp(const UDataflow* Dataflow, const Dataflow::FContext* Context)
+::Dataflow::FTimestamp FDataflowEditorViewportClient::LatestTimestamp(const UDataflow* Dataflow, const ::Dataflow::FContext* Context)
 {
 	if (Dataflow && Context)
 	{
 		return FMath::Max(Dataflow->GetRenderingTimestamp().Value, Context->GetTimestamp().Value);
 	}
-	return Dataflow::FTimestamp::Invalid;
+	return ::Dataflow::FTimestamp::Invalid;
 }
 
 bool FDataflowEditorViewportClient::InputKey(const FInputKeyEventArgs& EventArgs)
@@ -242,6 +247,6 @@ void FDataflowEditorViewportClient::Tick(float DeltaSeconds)
 void FDataflowEditorViewportClient::AddReferencedObjects(FReferenceCollector& Collector)
 {
 	Super::AddReferencedObjects(Collector);
+	Collector.AddReferencedObject(DataflowActor);
 }
 
-
diff --git a/Engine/Plugins/Experimental/Dataflow/Source/DataflowEditor/Public/Dataflow/AssetDefinition_DataflowAsset.h b/Engine/Plugins/Experimental/Dataflow/Source/DataflowEditor/Public/Dataflow/AssetDefinition_DataflowAsset.h
new file mode 100644
--- /dev/null
+++ b/Engine/Plugins/Experimental/Dataflow/Source/DataflowEditor/Public/Dataflow/AssetDefinition_DataflowAsset.h
@@ -0,0 +1,24 @@
+// Copyright Epic Games, Inc. All Rights Reserved.
+
+#pragma once
+
+#include "AssetDefinitionDefault.h"
+
+#include "AssetDefinition_DataflowAsset.generated.h"
+
+UCLASS()
+class UAssetDefinition_DataflowAsset : public UAssetDefinitionDefault
+{
+	GENERATED_BODY()
+
+private:
+
+	virtual FText GetAssetDisplayName() const override;
+	virtual TSoftClassPtr<UObject> GetAssetClass() const override;
+	virtual FLinearColor GetAssetColor() const override;
+	virtual TConstArrayView<FAssetCategoryPath> GetAssetCategories() const override;
+	virtual UThumbnailInfo* LoadThumbnailInfo(const FAssetData& InAssetData) const override;
+
+	virtual EAssetCommandResult OpenAssets(const FAssetOpenArgs& OpenArgs) const override;
+};
+
diff --git a/Engine/Plugins/Experimental/Dataflow/Source/DataflowEditor/Public/Dataflow/DataflowAssetActions.h b/Engine/Plugins/Experimental/Dataflow/Source/DataflowEditor/Public/Dataflow/DataflowAssetActions.h
deleted file mode 100644
diff --git a/Engine/Plugins/Experimental/Dataflow/Source/DataflowEditor/Public/Dataflow/DataflowEditor.h b/Engine/Plugins/Experimental/Dataflow/Source/DataflowEditor/Public/Dataflow/DataflowEditor.h
--- a/Engine/Plugins/Experimental/Dataflow/Source/DataflowEditor/Public/Dataflow/DataflowEditor.h
+++ b/Engine/Plugins/Experimental/Dataflow/Source/DataflowEditor/Public/Dataflow/DataflowEditor.h
@@ -2,11 +2,42 @@
 
 #pragma once
 
+#include "BaseCharacterFXEditor.h"
+
 #include "Dataflow/DataflowEngine.h"
-#include "Dataflow/DataflowAssetActions.h"
+#include "Dataflow/AssetDefinition_DataflowAsset.h"
 #include "Dataflow/DataflowAssetFactory.h"
 #include "Dataflow/DataflowEditorCommands.h"
 #include "Dataflow/DataflowEditorModule.h"
 #include "Dataflow/DataflowSchema.h"
 #include "Dataflow/DataflowSNode.h"
 #include "Dataflow/DataflowSNodeFactories.h"
+#include "Templates/SharedPointer.h"
+
+#include "DataflowEditor.generated.h"
+
+class FDataflowEditorToolkit;
+/** 
+ * The actual asset editor class doesn't have that much in it, intentionally. 
+ * 
+ * Our current asset editor guidelines ask us to place as little business logic as possible
+ * into the class, instead putting as much of the non-UI code into the subsystem as possible,
+ * and the UI code into the toolkit (which this class owns).
+ *
+ * However, since we're using a mode and the Interactive Tools Framework, a lot of our business logic
+ * ends up inside the mode and the tools, not the subsystem. The front-facing code is mostly in
+ * the asset editor toolkit, though the mode toolkit has most of the things that deal with the toolbar
+ * on the left.
+ */
+
+UCLASS()
+class DATAFLOWEDITOR_API UDataflowEditor : public UBaseCharacterFXEditor
+{
+	GENERATED_BODY()
+
+public:
+
+	virtual TSharedPtr<FBaseAssetToolkit> CreateToolkit() override;
+};
+
+DECLARE_LOG_CATEGORY_EXTERN(LogDataflowEditor, Log, All);
\ No newline at end of file
diff --git a/Engine/Plugins/Experimental/Dataflow/Source/DataflowEditor/Public/Dataflow/DataflowEditorCommands.h b/Engine/Plugins/Experimental/Dataflow/Source/DataflowEditor/Public/Dataflow/DataflowEditorCommands.h
--- a/Engine/Plugins/Experimental/Dataflow/Source/DataflowEditor/Public/Dataflow/DataflowEditorCommands.h
+++ b/Engine/Plugins/Experimental/Dataflow/Source/DataflowEditor/Public/Dataflow/DataflowEditorCommands.h
@@ -4,7 +4,7 @@
 
 #include "CoreMinimal.h"
 #include "Dataflow/DataflowObjectInterface.h"
-#include "Framework/Commands/Commands.h"
+#include "BaseCharacterFXEditorCommands.h"
 #include "Styling/AppStyle.h"
 
 class FDragDropEvent;
@@ -23,20 +23,17 @@ typedef TSet<class UObject*> FGraphPanelSelectionSet;
 * FDataflowEditorCommandsImpl
 * 
 */
-class FDataflowEditorCommandsImpl : public TCommands<FDataflowEditorCommandsImpl>
+class DATAFLOWEDITOR_API FDataflowEditorCommandsImpl : public TBaseCharacterFXEditorCommands<FDataflowEditorCommandsImpl>
 {
 public:
 
-	FDataflowEditorCommandsImpl()
-		: TCommands<FDataflowEditorCommandsImpl>( TEXT("DataflowEditor"), NSLOCTEXT("Contexts", "DataflowEditor", "Scene Graph Editor"), NAME_None, FAppStyle::GetAppStyleSetName() )
-	{
-	}	
+	FDataflowEditorCommandsImpl();
 
-	virtual ~FDataflowEditorCommandsImpl()
-	{
-	}
+	// TBaseCharacterFXEditorCommands<> interface
+	 virtual void RegisterCommands() override;
 
-	DATAFLOWEDITOR_API virtual void RegisterCommands() override;
+	// TInteractiveToolCommands<>
+	 virtual void GetToolDefaultObjectList(TArray<UInteractiveTool*>& ToolCDOs) override;
 
 	TSharedPtr< FUICommandInfo > EvaluateNode;
 	TSharedPtr< FUICommandInfo > CreateComment;
diff --git a/Engine/Plugins/Experimental/Dataflow/Source/DataflowEditor/Public/Dataflow/DataflowEditorMode.h b/Engine/Plugins/Experimental/Dataflow/Source/DataflowEditor/Public/Dataflow/DataflowEditorMode.h
new file mode 100644
--- /dev/null
+++ b/Engine/Plugins/Experimental/Dataflow/Source/DataflowEditor/Public/Dataflow/DataflowEditorMode.h
@@ -0,0 +1,36 @@
+// Copyright Epic Games, Inc. All Rights Reserved.
+
+#pragma once
+
+#include "BaseCharacterFXEditorMode.h"
+
+#include "DataflowEditorMode.generated.h"
+
+
+/**
+ * The dataflow editor mode is the mode used in the cloth asset editor. It holds most of the inter-tool state.
+ * We put things in a mode instead of directly into the asset editor in case we want to someday use the mode
+ * in multiple asset editors.
+ */
+UCLASS(Transient)
+class DATAFLOWEDITOR_API UDataflowEditorMode final : public UBaseCharacterFXEditorMode
+{
+	GENERATED_BODY()
+
+public:
+
+	const static FEditorModeID EM_DataflowEditorModeId;
+
+	UDataflowEditorMode();
+
+private:
+	// UEdMode
+	virtual void CreateToolkit() override;
+
+	// UBaseCharacterFXEditorMode
+	virtual void AddToolTargetFactories() override {}
+	virtual void RegisterTools() override {}
+	virtual void CreateToolTargets(const TArray<TObjectPtr<UObject>>& AssetsIn) override {}
+
+};
+
diff --git a/Engine/Plugins/Experimental/Dataflow/Source/DataflowEditor/Public/Dataflow/DataflowEditorModeToolkit.h b/Engine/Plugins/Experimental/Dataflow/Source/DataflowEditor/Public/Dataflow/DataflowEditorModeToolkit.h
new file mode 100644
--- /dev/null
+++ b/Engine/Plugins/Experimental/Dataflow/Source/DataflowEditor/Public/Dataflow/DataflowEditorModeToolkit.h
@@ -0,0 +1,27 @@
+// Copyright Epic Games, Inc. All Rights Reserved.
+
+#pragma once
+
+#include "BaseCharacterFXEditorModeToolkit.h"
+
+/**
+ * The dataflow editor mode toolkit is responsible for the panel on the side in the dataflow editor
+ * that shows mode and tool properties. Tool buttons would go in Init().
+ * Note: When there are separate viewports/worlds/modemanagers/toolscontexts, this ModeToolkit will track which
+ * one is currently active.
+ */
+
+class DATAFLOWEDITOR_API FDataflowEditorModeToolkit : public FBaseCharacterFXEditorModeToolkit
+{
+public:
+
+	virtual void OnToolStarted(UInteractiveToolManager* Manager, UInteractiveTool* Tool) override;
+
+
+	// IToolkit
+	virtual FName GetToolkitFName() const override;
+	virtual FText GetBaseToolkitName() const override;
+	virtual const FSlateBrush* GetActiveToolIcon(const FString& Identifier) const override;
+
+
+};
diff --git a/Engine/Plugins/Experimental/Dataflow/Source/DataflowEditor/Public/Dataflow/DataflowEditorModule.h b/Engine/Plugins/Experimental/Dataflow/Source/DataflowEditor/Public/Dataflow/DataflowEditorModule.h
--- a/Engine/Plugins/Experimental/Dataflow/Source/DataflowEditor/Public/Dataflow/DataflowEditorModule.h
+++ b/Engine/Plugins/Experimental/Dataflow/Source/DataflowEditor/Public/Dataflow/DataflowEditorModule.h
@@ -3,61 +3,24 @@
 #pragma once
 
 #include "CoreMinimal.h"
-#include "Modules/ModuleInterface.h"
-#include "Modules/ModuleManager.h"
-//#include "Toolkits/IToolkit.h"
-#include "Toolkits/AssetEditorToolkit.h"
+#include "BaseCharacterFXEditorModule.h"
+
 
-//class ISlateStyle;
-//class FSlateStyleSet;
 class FDataflowSNodeFactory;
-class FDataflowAssetActions;
 
 /**
  * The public interface to this module
  */
-class DATAFLOWEDITOR_API FDataflowEditorModule : public IModuleInterface
+class DATAFLOWEDITOR_API FDataflowEditorModule : public FBaseCharacterFXEditorModule
 {
 public:
+
+	/** IModuleInterface implementation */
 	virtual void StartupModule();
 	virtual void ShutdownModule();
 
-	TSharedRef<FAssetEditorToolkit> CreateDataflowAssetEditor(const EToolkitMode::Type Mode, const TSharedPtr< IToolkitHost >& InitToolkitHost, UObject* Dataflow);
-
-
-	/**
-	 * Singleton-like access to this module's interface.  This is just for convenience!
-	 * Beware of calling this during the shutdown phase, though.  Your module might have been unloaded already.
-	 *
-	 * @return Returns singleton instance, loading the module on demand if needed
-	 */
-	static inline FDataflowEditorModule& Get()
-	{
-		return FModuleManager::LoadModuleChecked< FDataflowEditorModule >( "DataflowEditor" );
-	}
-
-	/**
-	 * Checks to see if this module is loaded and ready.  It is only valid to call Get() if IsAvailable() returns true.
-	 *
-	 * @return True if the module is loaded and ready to use
-	 */
-	static inline bool IsAvailable()
-	{
-		return FModuleManager::Get().IsModuleLoaded( "DataflowEditor" );
-	}
-
-	//TSharedPtr<FSlateStyleSet> GetStyleSet() { return StyleSet; }
-	//static FName GetEditorStyleName();
-	//static const ISlateStyle* GetEditorStyle();
-
 private:
-	void RegisterMenus();
-
-private:
-
-	FDataflowAssetActions* DataflowAssetActions;
 
-	//TSharedPtr<FSlateStyleSet> StyleSet;
 	TSharedPtr<FDataflowSNodeFactory> DataflowSNodeFactory;
 
 };
diff --git a/Engine/Plugins/Experimental/Dataflow/Source/DataflowEditor/Public/Dataflow/DataflowEditorToolkit.h b/Engine/Plugins/Experimental/Dataflow/Source/DataflowEditor/Public/Dataflow/DataflowEditorToolkit.h
--- a/Engine/Plugins/Experimental/Dataflow/Source/DataflowEditor/Public/Dataflow/DataflowEditorToolkit.h
+++ b/Engine/Plugins/Experimental/Dataflow/Source/DataflowEditor/Public/Dataflow/DataflowEditorToolkit.h
@@ -2,16 +2,17 @@
 
 #pragma once
 
+#include "BaseCharacterFXEditorToolkit.h"
 #include "CoreMinimal.h"
 #include "Dataflow/DataflowObjectInterface.h"
-#include "Toolkits/SimpleAssetEditor.h"
 #include "Misc/NotifyHook.h"
 #include "GraphEditor.h"
 #include "TickableEditorObject.h"
+#include "Dataflow/CollectionSpreadSheetWidget.h"
 #include "Dataflow/DataflowSelectionView.h"
-#include "Dataflow/SelectionViewWidget.h"
 #include "Dataflow/DataflowCollectionSpreadSheet.h"
-#include "Dataflow/CollectionSpreadSheetWidget.h"
+#include "Dataflow/DataflowEditorViewport.h"
+#include "Dataflow/SelectionViewWidget.h"
 
 class FEditorViewportTabContent;
 class IDetailsView;
@@ -35,15 +36,18 @@ namespace Dataflow
 	};
 }
 
-class DATAFLOWEDITOR_API FDataflowEditorToolkit : public FAssetEditorToolkit, public FTickableEditorObject, public FNotifyHook, public FGCObject
+class DATAFLOWEDITOR_API FDataflowEditorToolkit final : public FBaseCharacterFXEditorToolkit, public FTickableEditorObject, public FNotifyHook
 {
+
+	using FBaseCharacterFXEditorToolkit::ObjectScene;
+
 public:
+
+	explicit FDataflowEditorToolkit(UAssetEditor* InOwningAssetEditor);
 	~FDataflowEditorToolkit();
 
 	static bool CanOpenDataflowEditor(UObject* ObjectToEdit);
 
-	virtual void InitializeEditor(const EToolkitMode::Type Mode, const TSharedPtr<IToolkitHost>& InitToolkitHost, UObject* ObjectToEdit);
-
 	// IToolkit Interface
 	virtual FName GetToolkitFName() const override;
 	virtual FText GetToolkitName() const override;
@@ -52,20 +56,10 @@ public:
 	virtual FString GetWorldCentricTabPrefix() const override;
 	virtual FLinearColor GetWorldCentricTabColorScale() const override;
 
-	// FTickableEditorObject interface
-	virtual void Tick(float DeltaTime) override;
-	virtual bool IsTickable() const override { return true; }
-	virtual TStatId GetStatId() const override;
-
-	// FGCObject Interface
-	virtual void AddReferencedObjects(FReferenceCollector& Collector) override;
-	virtual FString GetReferencerName() const override; 
-
 	// Tab spawners 
 	virtual void RegisterTabSpawners(const TSharedRef<FTabManager>& TabManager) override;
 	virtual void UnregisterTabSpawners(const TSharedRef<class FTabManager>& TabManager) override;
 
-	TSharedRef<SDockTab> SpawnTab_Viewport(const FSpawnTabArgs& Args);
 	TSharedRef<SDockTab> SpawnTab_GraphCanvas(const FSpawnTabArgs& Args);
 	TSharedRef<SDockTab> SpawnTab_AssetDetails(const FSpawnTabArgs& Args);
 	TSharedRef<SDockTab> SpawnTab_NodeDetails(const FSpawnTabArgs& Args);
@@ -77,8 +71,7 @@ public:
 	void OnTabClosed(TSharedRef<SDockTab> Tab);
 
 	// Member Access
-	UObject* GetAsset() { return Asset; }
-	const UObject* GetAsset() const { return Asset; }
+	UObject* GetAsset() const;
 
 	UDataflow* GetDataflow() { return Dataflow; }
 	const UDataflow* GetDataflow() const { return Dataflow; }
@@ -108,12 +101,34 @@ protected:
 
 private:
 
+	// FTickableEditorObject interface
+	virtual void Tick(float DeltaTime) override;
+	virtual bool IsTickable() const override { return true; }
+	virtual TStatId GetStatId() const override;
+
+	// FBaseCharacterFXEditorToolkit
+	virtual FEditorModeID GetEditorModeId() const override;
+	virtual void InitializeEdMode(UBaseCharacterFXEditorMode* EdMode) override;
+	virtual void CreateEditorModeUILayer() override;
+
+	// FAssetEditorToolkit
+	virtual bool OnRequestClose(EAssetEditorCloseReason InCloseReason) override;
+	virtual void PostInitAssetEditor() override;
+	virtual void GetSaveableObjects(TArray<UObject*>& OutObjects) const override;
+
+
+	// FBaseAssetToolkit
+	virtual void CreateWidgets() override;
+	virtual AssetEditorViewportFactoryFunction GetViewportDelegate() override;
+	virtual TSharedPtr<FEditorViewportClient> CreateEditorViewportClient() const override;
+
 	TObjectPtr<UObject> Asset = nullptr;
 	TObjectPtr<UDataflow> Dataflow = nullptr;
-	FString TerminalPath = "";
+	FString DataflowTerminalPath = "";
 
 	static const FName ViewportTabId;
-	TSharedPtr<FEditorViewportTabContent> ViewportEditor;
+	TSharedPtr<SDataflowEditorViewport> DataflowEditorViewport;
+
 
 	static const FName GraphCanvasTabId;
 	TSharedPtr<SDataflowGraphEditor> GraphEditor;
diff --git a/Engine/Plugins/Experimental/Dataflow/Source/DataflowEditor/Public/Dataflow/DataflowEditorViewport.h b/Engine/Plugins/Experimental/Dataflow/Source/DataflowEditor/Public/Dataflow/DataflowEditorViewport.h
--- a/Engine/Plugins/Experimental/Dataflow/Source/DataflowEditor/Public/Dataflow/DataflowEditorViewport.h
+++ b/Engine/Plugins/Experimental/Dataflow/Source/DataflowEditor/Public/Dataflow/DataflowEditorViewport.h
@@ -1,55 +1,38 @@
 // Copyright Epic Games, Inc. All Rights Reserved.
 #pragma once
 
-#include "CanvasTypes.h"
 #include "CoreMinimal.h"
-#include "SAssetEditorViewport.h"
+#include "SBaseCharacterFXEditorViewport.h"
 #include "SCommonEditorViewportToolbarBase.h"
 
-class FDataflowEditorToolkit;
+class ADataflowActor;
 class FAdvancedPreviewScene;
-class FEditorViewportClient;
+class UDataflowEditorMode;
 class FDataflowEditorViewportClient;
-class SEditorViewport;
-class ADataflowActor;
 
 // ----------------------------------------------------------------------------------
 
-class SDataflowEditorViewport : public SAssetEditorViewport, public ICommonEditorViewportToolbarInfoProvider, public FGCObject
+class SDataflowEditorViewport : public SBaseCharacterFXEditorViewport, public ICommonEditorViewportToolbarInfoProvider
 {
 public:
 	SLATE_BEGIN_ARGS(SDataflowEditorViewport) {}
-		SLATE_ARGUMENT(TWeakPtr<FDataflowEditorToolkit>, DataflowEditorToolkit)
-	SLATE_END_ARGS()
+	SLATE_ARGUMENT(TSharedPtr<FDataflowEditorViewportClient>, ViewportClient)
+		SLATE_END_ARGS()
 
 	SDataflowEditorViewport();
 
-	void Construct(const FArguments& InArgs);
+	void Construct(const FArguments& InArgs, const FAssetEditorViewportConstructionArgs& InViewportConstructionArgs);
+
+	// SEditorViewport
+	virtual void BindCommands() override;
+	virtual bool IsVisible() const override;
 
-	//~ ICommonEditorViewportToolbarInfoProvider interface
-	virtual TSharedRef<SEditorViewport> GetViewportWidget() override;
+	// ICommonEditorViewportToolbarInfoProvider
+	virtual TSharedRef<class SEditorViewport> GetViewportWidget() override;
 	virtual TSharedPtr<FExtender> GetExtenders() const override;
 	virtual void OnFloatingButtonClicked() override;
 
-	//~ FGCObject Interface
-	virtual void AddReferencedObjects(FReferenceCollector& Collector) override;
-	virtual FString GetReferencerName() const override{return TEXT("SDataflowEditorViewport");}
-
-
-protected:
-	virtual TSharedRef<FEditorViewportClient> MakeEditorViewportClient() override;
-	virtual TSharedPtr<SWidget> MakeViewportToolbar() override;
 private:
 
-	virtual void BindCommands() override;
-
-
-	/// The scene for this viewport. 
-	TSharedPtr<FAdvancedPreviewScene> PreviewScene;
-
-	/// Editor viewport client 
-	TSharedPtr<FDataflowEditorViewportClient> ViewportClient;
-
-	TWeakPtr<FDataflowEditorToolkit> DataflowEditorToolkitPtr;
-	TObjectPtr<ADataflowActor> CustomDataflowActor = nullptr;
+	UDataflowEditorMode* GetEdMode() const;
 };
diff --git a/Engine/Plugins/Experimental/Dataflow/Source/DataflowEditor/Public/Dataflow/DataflowEditorViewportClient.h b/Engine/Plugins/Experimental/Dataflow/Source/DataflowEditor/Public/Dataflow/DataflowEditorViewportClient.h
--- a/Engine/Plugins/Experimental/Dataflow/Source/DataflowEditor/Public/Dataflow/DataflowEditorViewportClient.h
+++ b/Engine/Plugins/Experimental/Dataflow/Source/DataflowEditor/Public/Dataflow/DataflowEditorViewportClient.h
@@ -9,19 +9,20 @@
 class FDataflowEditorToolkit;
 class ADataflowActor;
 
-// ----------------------------------------------------------------------------------
 
-class FDataflowEditorViewportClient : public FEditorViewportClient
+
+class DATAFLOWEDITOR_API FDataflowEditorViewportClient : public FEditorViewportClient
 {
 public:
 	using Super = FEditorViewportClient;
-	
-	FDataflowEditorViewportClient(FPreviewScene* InPreviewScene,
-		const TWeakPtr<SEditorViewport> InEditorViewportWidget = nullptr,
-		TWeakPtr<FDataflowEditorToolkit> InDataflowEditorToolkitPtr = nullptr);
 
-	Dataflow::FTimestamp LatestTimestamp(const UDataflow* Dataflow, const Dataflow::FContext* Context);
-	void SetDataflowActor(ADataflowActor* InActor) { DataflowActor = InActor; }
+	FDataflowEditorViewportClient(FEditorModeTools* InModeTools, FPreviewScene* InPreviewScene,
+								  const TWeakPtr<SEditorViewport> InEditorViewportWidget = nullptr);
+
+	void SetDataflowEditorToolkit(TWeakPtr<FDataflowEditorToolkit> DataflowToolkit);
+
+
+	::Dataflow::FTimestamp LatestTimestamp(const UDataflow* Dataflow, const ::Dataflow::FContext* Context);
 
 	void SetSelectionMode(FDataflowSelectionState::EMode InState);
 	bool CanSetSelectionMode(FDataflowSelectionState::EMode InState);
@@ -42,8 +43,8 @@ public:
 private:
 
 	TWeakPtr<FDataflowEditorToolkit> DataflowEditorToolkitPtr = nullptr;
-	Dataflow::FTimestamp LastModifiedTimestamp = Dataflow::FTimestamp::Invalid;
-	ADataflowActor* DataflowActor = nullptr;
+	::Dataflow::FTimestamp LastModifiedTimestamp = ::Dataflow::FTimestamp::Invalid;
+	TObjectPtr<ADataflowActor> DataflowActor = nullptr;
 
 	//
 	// Selection
diff --git a/Engine/Plugins/Experimental/GeometryCollectionPlugin/Source/GeometryCollectionEditor/GeometryCollectionEditor.Build.cs b/Engine/Plugins/Experimental/GeometryCollectionPlugin/Source/GeometryCollectionEditor/GeometryCollectionEditor.Build.cs
--- a/Engine/Plugins/Experimental/GeometryCollectionPlugin/Source/GeometryCollectionEditor/GeometryCollectionEditor.Build.cs
+++ b/Engine/Plugins/Experimental/GeometryCollectionPlugin/Source/GeometryCollectionEditor/GeometryCollectionEditor.Build.cs
@@ -9,6 +9,7 @@ namespace UnrealBuildTool.Rules
 			PublicDependencyModuleNames.AddRange(
 				new string[]
 				{
+					"AssetDefinition",
 					"Core",
 					"CoreUObject",
 				    "Slate",
diff --git a/Engine/Plugins/Experimental/GeometryCollectionPlugin/Source/GeometryCollectionEditor/Private/GeometryCollection/AssetDefinition_GeometryCollection.cpp b/Engine/Plugins/Experimental/GeometryCollectionPlugin/Source/GeometryCollectionEditor/Private/GeometryCollection/AssetDefinition_GeometryCollection.cpp
new file mode 100644
--- /dev/null
+++ b/Engine/Plugins/Experimental/GeometryCollectionPlugin/Source/GeometryCollectionEditor/Private/GeometryCollection/AssetDefinition_GeometryCollection.cpp
@@ -0,0 +1,80 @@
+// Copyright Epic Games, Inc. All Rights Reserved.
+
+#include "GeometryCollection/AssetDefinition_GeometryCollection.h"
+#include "GeometryCollection/GeometryCollectionObject.h"
+#include "Dataflow/DataflowEditor.h"
+#include "Dataflow/DataflowEditorModule.h"
+#include "Dataflow/DataflowEditorToolkit.h"
+#include "Math/Color.h"
+#include "ThumbnailRendering/SceneThumbnailInfo.h"
+#include "Toolkits/SimpleAssetEditor.h"
+
+
+#define LOCTEXT_NAMESPACE "AssetActions_GeometryCollection"
+
+
+namespace UE::GeometryCollection
+{
+	struct FColorScheme
+	{
+		static inline const FLinearColor Asset = FColor(180, 120, 110);
+		static inline const FLinearColor NodeHeader = FColor(180, 120, 110);
+		static inline const FLinearColor NodeBody = FColor(18, 12, 11, 127);
+	};
+}
+
+FText UAssetDefinition_GeometryCollection::GetAssetDisplayName() const
+{
+	return NSLOCTEXT("AssetTypeActions", "AssetTypeActions_GeometryCollection", "Geometry Collection");
+}
+
+TSoftClassPtr<UObject> UAssetDefinition_GeometryCollection::GetAssetClass() const
+{
+	return UGeometryCollection::StaticClass();
+}
+
+FLinearColor UAssetDefinition_GeometryCollection::GetAssetColor() const
+{
+	return UE::GeometryCollection::FColorScheme::Asset;
+}
+
+TConstArrayView<FAssetCategoryPath> UAssetDefinition_GeometryCollection::GetAssetCategories() const
+{
+	static const auto Categories = { EAssetCategoryPaths::Physics };
+	return Categories;
+}
+
+UThumbnailInfo* UAssetDefinition_GeometryCollection::LoadThumbnailInfo(const FAssetData& InAsset) const
+{
+	return UE::Editor::FindOrCreateThumbnailInfo(InAsset.GetAsset(), USceneThumbnailInfo::StaticClass());
+}
+
+EAssetCommandResult UAssetDefinition_GeometryCollection::OpenAssets(const FAssetOpenArgs& OpenArgs) const
+{
+	TArray<UGeometryCollection*> GeometryCollectionObjects = OpenArgs.LoadObjects<UGeometryCollection>();
+
+	// For now the geometry collection editor only works on one asset at a time
+	ensure(GeometryCollectionObjects.Num() == 0 || GeometryCollectionObjects.Num() == 1);
+	if (GeometryCollectionObjects.Num() == 1)
+	{
+		// Validate the asset
+		if (UGeometryCollection* const GeometryCollection = Cast<UGeometryCollection>(GeometryCollectionObjects[0]))
+		{
+			if (FDataflowEditorToolkit::CanOpenDataflowEditor(GeometryCollection))
+			{
+				UAssetEditorSubsystem* const AssetEditorSubsystem = GEditor->GetEditorSubsystem<UAssetEditorSubsystem>();
+				UDataflowEditor* const AssetEditor = NewObject<UDataflowEditor>(AssetEditorSubsystem, NAME_None, RF_Transient);
+				AssetEditor->Initialize({ GeometryCollection });
+				return EAssetCommandResult::Handled;
+			}
+
+			FSimpleAssetEditor::CreateEditor(EToolkitMode::Standalone, OpenArgs.ToolkitHost, GeometryCollection);
+			return EAssetCommandResult::Handled;
+
+		}
+	}
+	return EAssetCommandResult::Unhandled;
+}
+
+
+#undef LOCTEXT_NAMESPACE
diff --git a/Engine/Plugins/Experimental/GeometryCollectionPlugin/Source/GeometryCollectionEditor/Private/GeometryCollection/AssetTypeActions_GeometryCollection.cpp b/Engine/Plugins/Experimental/GeometryCollectionPlugin/Source/GeometryCollectionEditor/Private/GeometryCollection/AssetTypeActions_GeometryCollection.cpp
deleted file mode 100644
diff --git a/Engine/Plugins/Experimental/GeometryCollectionPlugin/Source/GeometryCollectionEditor/Private/GeometryCollection/GeometryCollectionEditorPlugin.cpp b/Engine/Plugins/Experimental/GeometryCollectionPlugin/Source/GeometryCollectionEditor/Private/GeometryCollection/GeometryCollectionEditorPlugin.cpp
--- a/Engine/Plugins/Experimental/GeometryCollectionPlugin/Source/GeometryCollectionEditor/Private/GeometryCollection/GeometryCollectionEditorPlugin.cpp
+++ b/Engine/Plugins/Experimental/GeometryCollectionPlugin/Source/GeometryCollectionEditor/Private/GeometryCollection/GeometryCollectionEditorPlugin.cpp
@@ -11,7 +11,7 @@
 #include "GeometryCollection/GeometryCollectionComponent.h"
 #include "GeometryCollection/GeometryCollectionConversion.h"
 #include "GeometryCollection/GeometryCollectionCommands.h"
-#include "GeometryCollection/AssetTypeActions_GeometryCollection.h"
+#include "GeometryCollection/AssetDefinition_GeometryCollection.h"
 #include "GeometryCollection/AssetTypeActions_GeometryCollectionCache.h"
 #include "GeometryCollection/GeometryCollectionAssetBroker.h"
 #include "GeometryCollection/GeometryCollectionEditorStyle.h"
@@ -45,9 +45,7 @@ void IGeometryCollectionEditorPlugin::StartupModule()
 
 	FAssetToolsModule& AssetToolsModule = FAssetToolsModule::GetModule();
 	IAssetTools& AssetTools = AssetToolsModule.Get();
-	GeometryCollectionAssetActions = new FAssetTypeActions_GeometryCollection();
 	GeometryCollectionCacheAssetActions = new FAssetTypeActions_GeometryCollectionCache();
-	AssetTools.RegisterAssetTypeActions(MakeShareable(GeometryCollectionAssetActions));
 	AssetTools.RegisterAssetTypeActions(MakeShareable(GeometryCollectionCacheAssetActions));
 
 	AssetBroker = new FGeometryCollectionAssetBroker();
@@ -272,7 +270,6 @@ void IGeometryCollectionEditorPlugin::ShutdownModule()
 
 		FAssetToolsModule& AssetToolsModule = FAssetToolsModule::GetModule();
 		IAssetTools& AssetTools = AssetToolsModule.Get();
-		AssetTools.UnregisterAssetTypeActions(GeometryCollectionAssetActions->AsShared());
 		AssetTools.UnregisterAssetTypeActions(GeometryCollectionCacheAssetActions->AsShared());
 
 		FComponentAssetBrokerage::UnregisterBroker(MakeShareable(AssetBroker));
diff --git a/Engine/Plugins/Experimental/GeometryCollectionPlugin/Source/GeometryCollectionEditor/Public/GeometryCollection/AssetDefinition_GeometryCollection.h b/Engine/Plugins/Experimental/GeometryCollectionPlugin/Source/GeometryCollectionEditor/Public/GeometryCollection/AssetDefinition_GeometryCollection.h
new file mode 100644
--- /dev/null
+++ b/Engine/Plugins/Experimental/GeometryCollectionPlugin/Source/GeometryCollectionEditor/Public/GeometryCollection/AssetDefinition_GeometryCollection.h
@@ -0,0 +1,24 @@
+// Copyright Epic Games, Inc. All Rights Reserved.
+
+#pragma once
+
+#include "AssetDefinitionDefault.h"
+
+#include "AssetDefinition_GeometryCollection.generated.h"
+
+UCLASS()
+class GEOMETRYCOLLECTIONEDITOR_API UAssetDefinition_GeometryCollection : public UAssetDefinitionDefault
+{
+	GENERATED_BODY()
+
+private:
+
+	virtual FText GetAssetDisplayName() const override;
+	virtual TSoftClassPtr<UObject> GetAssetClass() const override;
+	virtual FLinearColor GetAssetColor() const override;
+	virtual TConstArrayView<FAssetCategoryPath> GetAssetCategories() const override;
+	virtual UThumbnailInfo* LoadThumbnailInfo(const FAssetData& InAssetData) const override;
+
+	virtual EAssetCommandResult OpenAssets(const FAssetOpenArgs& OpenArgs) const override;
+};
+
diff --git a/Engine/Plugins/Experimental/GeometryCollectionPlugin/Source/GeometryCollectionEditor/Public/GeometryCollection/AssetTypeActions_GeometryCollection.h b/Engine/Plugins/Experimental/GeometryCollectionPlugin/Source/GeometryCollectionEditor/Public/GeometryCollection/AssetTypeActions_GeometryCollection.h
deleted file mode 100644
diff --git a/Engine/Plugins/Experimental/GeometryCollectionPlugin/Source/GeometryCollectionEditor/Public/GeometryCollection/GeometryCollectionEditorPlugin.h b/Engine/Plugins/Experimental/GeometryCollectionPlugin/Source/GeometryCollectionEditor/Public/GeometryCollection/GeometryCollectionEditorPlugin.h
--- a/Engine/Plugins/Experimental/GeometryCollectionPlugin/Source/GeometryCollectionEditor/Public/GeometryCollection/GeometryCollectionEditorPlugin.h
+++ b/Engine/Plugins/Experimental/GeometryCollectionPlugin/Source/GeometryCollectionEditor/Public/GeometryCollection/GeometryCollectionEditorPlugin.h
@@ -60,7 +60,6 @@ private:
 	TArray<IConsoleObject*> EditorCommands;
 
 	// Asset actions for new asset types
-	FAssetTypeActions_GeometryCollection* GeometryCollectionAssetActions;
 	FAssetTypeActions_GeometryCollectionCache* GeometryCollectionCacheAssetActions;
 	FGeometryCollectionAssetBroker* AssetBroker;
 
