Dataflow - Changed the order of execution of Terminal nodes' SetAssetValue and Evaluate in the evaluation command.
This is so that evaluating the outputs can reuse any information created during a previous SetAssetValue call.
#tests Ran Dataflow graphs with this enabled locally for a couple of weeks.

diff --git a/Engine/Plugins/Experimental/Dataflow/Source/DataflowEditor/Private/Dataflow/DataflowEditorCommands.cpp b/Engine/Plugins/Experimental/Dataflow/Source/DataflowEditor/Private/Dataflow/DataflowEditorCommands.cpp
--- a/Engine/Plugins/Experimental/Dataflow/Source/DataflowEditor/Private/Dataflow/DataflowEditorCommands.cpp
+++ b/Engine/Plugins/Experimental/Dataflow/Source/DataflowEditor/Private/Dataflow/DataflowEditorCommands.cpp
@@ -181,22 +181,30 @@ void FDataflowEditorCommands::EvaluateTerminalNode(Dataflow::FContext& Context,
 				{
 					if (Node->GetTimestamp() >= OutLastNodeTimestamp)
 					{
-						Context.Evaluate(Node, nullptr);
-						OutLastNodeTimestamp = Context.GetTimestamp();
-
 						if (const FDataflowTerminalNode* TerminalNode = Node->AsType<const FDataflowTerminalNode>())
 						{
 							if (InAsset)
 							{
-								TerminalNode->SetAssetValue(InAsset, Context);
+								TerminalNode->SetAssetValue(InAsset, Context);  // Must set asset value before call to Evaluate
 							}
 						}
+
+						Context.Evaluate(Node, nullptr);
+						OutLastNodeTimestamp = Context.GetTimestamp();
 					}
 				}
 				else // Output != nullptr
 				{
 					if (!Context.HasData(Output->CacheKey(), Context.GetTimestamp()))
 					{
+						if (const FDataflowTerminalNode* TerminalNode = Node->AsType<const FDataflowTerminalNode>())
+						{
+							if (InAsset)
+							{
+								TerminalNode->SetAssetValue(InAsset, Context);  // Must set asset value before call to Evaluate
+							}
+						}
+
 						Context.Evaluate(Node, Output);
 					}
 				}
