Dataflow : Prompt user for dataflow asset creation
- prevent geometry collection from oppening a dataflow editor without a DataflowAsset.

diff --git a/Engine/Plugins/Experimental/ChaosFlesh/Source/ChaosFleshEditor/Private/ChaosFlesh/Asset/AssetDefinition_FleshAsset.cpp b/Engine/Plugins/Experimental/ChaosFlesh/Source/ChaosFleshEditor/Private/ChaosFlesh/Asset/AssetDefinition_FleshAsset.cpp
--- a/Engine/Plugins/Experimental/ChaosFlesh/Source/ChaosFleshEditor/Private/ChaosFlesh/Asset/AssetDefinition_FleshAsset.cpp
+++ b/Engine/Plugins/Experimental/ChaosFlesh/Source/ChaosFleshEditor/Private/ChaosFlesh/Asset/AssetDefinition_FleshAsset.cpp
@@ -2,6 +2,7 @@
 
 #include "ChaosFlesh/Asset/AssetDefinition_FleshAsset.h"
 #include "ChaosFlesh/FleshAsset.h"
+#include "Dataflow/AssetDefinition_DataflowAsset.h"
 #include "Dataflow/DataflowEditor.h"
 #include "Dataflow/DataflowEditorModule.h"
 #include "Dataflow/DataflowEditorToolkit.h"
@@ -60,7 +61,15 @@ EAssetCommandResult UAssetDefinition_FleshAsset::OpenAssets(const FAssetOpenArgs
 		// Validate the asset
 		if (UFleshAsset* const FleshAsset = Cast<UFleshAsset>(FleshObjects[0]))
 		{
-			if (FDataflowEditorToolkit::CanOpenDataflowEditor(FleshAsset))
+			if (!FDataflowEditorToolkit::HasDataflowAsset(FleshAsset))
+			{
+				if (UDataflow* const NewDataflowAsset = Cast<UDataflow>(DataflowAssetDefinitionHelpers::NewOrOpenDataflowAsset(FleshAsset)))
+				{
+					FleshAsset->DataflowAsset = NewDataflowAsset;
+				}
+			}
+
+			if (FDataflowEditorToolkit::HasDataflowAsset(FleshAsset))
 			{
 				UAssetEditorSubsystem* const AssetEditorSubsystem = GEditor->GetEditorSubsystem<UAssetEditorSubsystem>();
 				UDataflowEditor* const AssetEditor = NewObject<UDataflowEditor>(AssetEditorSubsystem, NAME_None, RF_Transient);
diff --git a/Engine/Plugins/Experimental/Dataflow/Source/DataflowEditor/DataflowEditor.Build.cs b/Engine/Plugins/Experimental/Dataflow/Source/DataflowEditor/DataflowEditor.Build.cs
--- a/Engine/Plugins/Experimental/Dataflow/Source/DataflowEditor/DataflowEditor.Build.cs
+++ b/Engine/Plugins/Experimental/Dataflow/Source/DataflowEditor/DataflowEditor.Build.cs
@@ -45,6 +45,7 @@ namespace UnrealBuildTool.Rules
 				    "RHI",
 				    "SceneOutliner",
 					"ToolMenus",
+					"ToolWidgets",
 					"Slate",
 					"XmlParser"
 				}
diff --git a/Engine/Plugins/Experimental/Dataflow/Source/DataflowEditor/Private/Dataflow/AssetDefinition_DataflowAsset.cpp b/Engine/Plugins/Experimental/Dataflow/Source/DataflowEditor/Private/Dataflow/AssetDefinition_DataflowAsset.cpp
--- a/Engine/Plugins/Experimental/Dataflow/Source/DataflowEditor/Private/Dataflow/AssetDefinition_DataflowAsset.cpp
+++ b/Engine/Plugins/Experimental/Dataflow/Source/DataflowEditor/Private/Dataflow/AssetDefinition_DataflowAsset.cpp
@@ -1,15 +1,137 @@
 // Copyright Epic Games, Inc. All Rights Reserved.
 
 #include "Dataflow/AssetDefinition_DataflowAsset.h"
+
+#include "AssetRegistry/AssetRegistryModule.h"
+#include "ContentBrowserModule.h"
 #include "Dataflow/DataflowEditor.h"
 #include "Dataflow/DataflowEditorModule.h"
 #include "Dataflow/DataflowObject.h"
+#include "Dialog/SMessageDialog.h"
+#include "IContentBrowserSingleton.h"
 #include "Math/Color.h"
+#include "Misc/FileHelper.h"
 #include "ThumbnailRendering/SceneThumbnailInfo.h"
 
 
 #define LOCTEXT_NAMESPACE "AssetActions_DataflowAsset"
 
+namespace DataflowAssetDefinitionHelpers
+{
+	// Return true if we should proceed, false if we should re-open the dialog
+	bool CreateNewDataflowAsset(const UObject* Asset, UObject*& OutDataflowAsset)
+	{
+		const UClass* const DataflowClass = UDataflow::StaticClass();
+
+		FSaveAssetDialogConfig NewDataflowAssetDialogConfig;
+		{
+			const FString PackageName = Asset->GetOutermost()->GetName();
+			NewDataflowAssetDialogConfig.DefaultPath = FPackageName::GetLongPackagePath(PackageName);
+			const FString AssetName = Asset->GetName();
+			NewDataflowAssetDialogConfig.DefaultAssetName = AssetName + "_Dataflow";
+			NewDataflowAssetDialogConfig.AssetClassNames.Add(DataflowClass->GetClassPathName());
+			NewDataflowAssetDialogConfig.ExistingAssetPolicy = ESaveAssetDialogExistingAssetPolicy::Disallow;
+			NewDataflowAssetDialogConfig.DialogTitleOverride = LOCTEXT("NewDataflowAssetDialogTitle", "Save Dataflow Asset As");
+		}
+
+		FContentBrowserModule& ContentBrowserModule = FModuleManager::LoadModuleChecked<FContentBrowserModule>("ContentBrowser");
+
+		FString NewPackageName;
+		FText OutError;
+		for (bool bFilenameValid = false; !bFilenameValid; bFilenameValid = FFileHelper::IsFilenameValidForSaving(NewPackageName, OutError))
+		{
+			const FString AssetSavePath = ContentBrowserModule.Get().CreateModalSaveAssetDialog(NewDataflowAssetDialogConfig);
+			if (AssetSavePath.IsEmpty())
+			{
+				OutDataflowAsset = nullptr;
+				return false;
+			}
+			NewPackageName = FPackageName::ObjectPathToPackageName(AssetSavePath);
+		}
+
+		const FName NewAssetName(FPackageName::GetLongPackageAssetName(NewPackageName));
+		UPackage* const NewPackage = CreatePackage(*NewPackageName);
+		UObject* const NewAsset = NewObject<UObject>(NewPackage, DataflowClass, NewAssetName, RF_Public | RF_Standalone | RF_Transactional);
+
+		NewAsset->MarkPackageDirty();
+
+		// Notify the asset registry
+		FAssetRegistryModule::AssetCreated(NewAsset);
+
+		OutDataflowAsset = NewAsset;
+		return true;
+	}
+
+
+	// Return true if we should proceed, false if we should re-open the dialog
+	bool OpenDataflowAsset(const UObject* Asset, UObject*& OutDataflowAsset)
+	{
+		const UClass* const DataflowClass = UDataflow::StaticClass();
+
+		FOpenAssetDialogConfig NewDataflowAssetDialogConfig;
+		{
+			const FString PackageName = Asset->GetOutermost()->GetName();
+			NewDataflowAssetDialogConfig.DefaultPath = FPackageName::GetLongPackagePath(PackageName);
+			NewDataflowAssetDialogConfig.AssetClassNames.Add(DataflowClass->GetClassPathName());
+			NewDataflowAssetDialogConfig.bAllowMultipleSelection = false;
+			NewDataflowAssetDialogConfig.DialogTitleOverride = LOCTEXT("OpenDataflowAssetDialogTitle", "Open Dataflow Asset");
+		}
+
+		FContentBrowserModule& ContentBrowserModule = FModuleManager::LoadModuleChecked<FContentBrowserModule>("ContentBrowser");
+		TArray<FAssetData> AssetData = ContentBrowserModule.Get().CreateModalOpenAssetDialog(NewDataflowAssetDialogConfig);
+
+		if (AssetData.Num() == 1)
+		{
+			OutDataflowAsset = AssetData[0].GetAsset();
+			return true;
+		}
+
+		return false;
+	}
+
+	// Return true if we should proceed, false if we should re-open the dialog
+	bool NewOrOpenDialog(const UObject* Asset, UObject*& OutDataflowAsset)
+	{
+		TSharedRef<SMessageDialog> ConfirmDialog = SNew(SMessageDialog)
+			.Title(FText(LOCTEXT("Dataflow_WindowTitle", "Create or Open Dataflow graph?")))
+			.Message(LOCTEXT("Dataflow_WindowText", "This Asset currently has no Dataflow graph"))
+			.Buttons({
+				SMessageDialog::FButton(LOCTEXT("Dataflow_NewText", "Create new Dataflow")),
+				SMessageDialog::FButton(LOCTEXT("Dataflow_OpenText", "Open existing Dataflow")),
+				SMessageDialog::FButton(LOCTEXT("Dataflow_ContinueText", "Continue without Dataflow")),
+					 });
+
+		const int32 ResultButtonIdx = ConfirmDialog->ShowModal();
+		switch (ResultButtonIdx)
+		{
+		case 0:
+			return CreateNewDataflowAsset(Asset, OutDataflowAsset);
+			break;
+		case 1:
+			return OpenDataflowAsset(Asset, OutDataflowAsset);
+			break;
+		default:
+			break;
+		}
+
+		return true;
+	}
+
+	// Create a new UDataflow if one doesn't already exist for the Asset
+	UObject* NewOrOpenDataflowAsset(const UObject* Asset)
+	{
+		UObject* DataflowAsset = nullptr;
+		bool bDialogDone = false;
+		while (!bDialogDone)
+		{
+			bDialogDone = NewOrOpenDialog(Asset, DataflowAsset);
+		}
+
+		return DataflowAsset;
+	}
+}
+
+
 
 namespace UE::Dataflow::DataflowAsset
 {
diff --git a/Engine/Plugins/Experimental/Dataflow/Source/DataflowEditor/Private/Dataflow/DataflowEditorToolkit.cpp b/Engine/Plugins/Experimental/Dataflow/Source/DataflowEditor/Private/Dataflow/DataflowEditorToolkit.cpp
--- a/Engine/Plugins/Experimental/Dataflow/Source/DataflowEditor/Private/Dataflow/DataflowEditorToolkit.cpp
+++ b/Engine/Plugins/Experimental/Dataflow/Source/DataflowEditor/Private/Dataflow/DataflowEditorToolkit.cpp
@@ -170,6 +170,18 @@ bool FDataflowEditorToolkit::CanOpenDataflowEditor(UObject* ObjectToEdit)
 	return Dataflow != nullptr;
 }
 
+bool FDataflowEditorToolkit::HasDataflowAsset(UObject* ObjectToEdit)
+{
+	if (const UClass* Class = ObjectToEdit->GetClass())
+	{
+		if (FProperty* Property = Class->FindPropertyByName(FName("DataflowAsset")))
+		{
+			return *Property->ContainerPtrToValuePtr<UDataflow*>(ObjectToEdit) != nullptr;
+		}
+	}
+	return false;
+}
+
 //~ Begin FBaseCharacterFXEditorToolkit overrides
 
 FEditorModeID FDataflowEditorToolkit::GetEditorModeId() const
diff --git a/Engine/Plugins/Experimental/Dataflow/Source/DataflowEditor/Public/Dataflow/AssetDefinition_DataflowAsset.h b/Engine/Plugins/Experimental/Dataflow/Source/DataflowEditor/Public/Dataflow/AssetDefinition_DataflowAsset.h
--- a/Engine/Plugins/Experimental/Dataflow/Source/DataflowEditor/Public/Dataflow/AssetDefinition_DataflowAsset.h
+++ b/Engine/Plugins/Experimental/Dataflow/Source/DataflowEditor/Public/Dataflow/AssetDefinition_DataflowAsset.h
@@ -6,6 +6,24 @@
 
 #include "AssetDefinition_DataflowAsset.generated.h"
 
+
+namespace DataflowAssetDefinitionHelpers
+{
+	// Return true if we should proceed, false if we should re-open the dialog
+	DATAFLOWEDITOR_API bool CreateNewDataflowAsset(const UObject* Asset, UObject*& OutDataflowAsset);
+
+	// Return true if we should proceed, false if we should re-open the dialog
+	DATAFLOWEDITOR_API bool OpenDataflowAsset(const UObject* Asset, UObject*& OutDataflowAsset);
+
+	// Return true if we should proceed, false if we should re-open the dialog
+	DATAFLOWEDITOR_API bool NewOrOpenDialog(const UObject* Asset, UObject*& OutDataflowAsset);
+
+	// Create a new UDataflow if one doesn't already exist for the Cloth Asset
+	DATAFLOWEDITOR_API UObject* NewOrOpenDataflowAsset(const UObject* Asset);
+}
+
+
+
 UCLASS()
 class UAssetDefinition_DataflowAsset : public UAssetDefinitionDefault
 {
diff --git a/Engine/Plugins/Experimental/Dataflow/Source/DataflowEditor/Public/Dataflow/DataflowEditorToolkit.h b/Engine/Plugins/Experimental/Dataflow/Source/DataflowEditor/Public/Dataflow/DataflowEditorToolkit.h
--- a/Engine/Plugins/Experimental/Dataflow/Source/DataflowEditor/Public/Dataflow/DataflowEditorToolkit.h
+++ b/Engine/Plugins/Experimental/Dataflow/Source/DataflowEditor/Public/Dataflow/DataflowEditorToolkit.h
@@ -47,6 +47,7 @@ public:
 	~FDataflowEditorToolkit();
 
 	static bool CanOpenDataflowEditor(UObject* ObjectToEdit);
+	static bool HasDataflowAsset(UObject* ObjectToEdit);
 
 	// IToolkit Interface
 	virtual FName GetToolkitFName() const override;
diff --git a/Engine/Plugins/Experimental/GeometryCollectionPlugin/Source/GeometryCollectionEditor/Private/GeometryCollection/AssetDefinition_GeometryCollection.cpp b/Engine/Plugins/Experimental/GeometryCollectionPlugin/Source/GeometryCollectionEditor/Private/GeometryCollection/AssetDefinition_GeometryCollection.cpp
--- a/Engine/Plugins/Experimental/GeometryCollectionPlugin/Source/GeometryCollectionEditor/Private/GeometryCollection/AssetDefinition_GeometryCollection.cpp
+++ b/Engine/Plugins/Experimental/GeometryCollectionPlugin/Source/GeometryCollectionEditor/Private/GeometryCollection/AssetDefinition_GeometryCollection.cpp
@@ -60,7 +60,7 @@ EAssetCommandResult UAssetDefinition_GeometryCollection::OpenAssets(const FAsset
 		// Validate the asset
 		if (UGeometryCollection* const GeometryCollection = Cast<UGeometryCollection>(GeometryCollectionObjects[0]))
 		{
-			if (FDataflowEditorToolkit::CanOpenDataflowEditor(GeometryCollection))
+			if (FDataflowEditorToolkit::HasDataflowAsset(GeometryCollection))
 			{
 				UAssetEditorSubsystem* const AssetEditorSubsystem = GEditor->GetEditorSubsystem<UAssetEditorSubsystem>();
 				UDataflowEditor* const AssetEditor = NewObject<UDataflowEditor>(AssetEditorSubsystem, NAME_None, RF_Transient);
