Dataflow : Rendering
- Update rendering to reflect the correct transform type.

diff --git a/Engine/Plugins/Experimental/Dataflow/Source/DataflowEnginePlugin/Private/Dataflow/DataflowEngineRendering.cpp b/Engine/Plugins/Experimental/Dataflow/Source/DataflowEnginePlugin/Private/Dataflow/DataflowEngineRendering.cpp
--- a/Engine/Plugins/Experimental/Dataflow/Source/DataflowEnginePlugin/Private/Dataflow/DataflowEngineRendering.cpp
+++ b/Engine/Plugins/Experimental/Dataflow/Source/DataflowEnginePlugin/Private/Dataflow/DataflowEngineRendering.cpp
@@ -35,10 +35,10 @@ namespace Dataflow
 
 					FManagedArrayCollection Default;
 					const FManagedArrayCollection& Collection = State.GetValue<FManagedArrayCollection>(PrimaryOutput, Default);
-
+					
 					if (Collection.FindAttributeTyped<FIntVector>("Indices", FGeometryCollection::FacesGroup)
 						&& Collection.FindAttributeTyped<FVector3f>("Vertex", FGeometryCollection::VerticesGroup)
-						&& Collection.FindAttributeTyped<FTransform>(FTransformCollection::TransformAttribute, FTransformCollection::TransformGroup)
+						&& Collection.FindAttributeTyped<FTransform3f>(FTransformCollection::TransformAttribute, FTransformCollection::TransformGroup)
 						&& Collection.FindAttributeTyped<int32>("BoneMap", FGeometryCollection::VerticesGroup)
 						&& Collection.FindAttributeTyped<int32>(FTransformCollection::ParentAttribute, FTransformCollection::TransformGroup))
 					{
