Dataflow : Mesh Editor Tools
- Prototype for updating the dataflow editor to support the mesh editor tools.

diff --git a/Engine/Plugins/Experimental/Dataflow/Dataflow.uplugin b/Engine/Plugins/Experimental/Dataflow/Dataflow.uplugin
--- a/Engine/Plugins/Experimental/Dataflow/Dataflow.uplugin
+++ b/Engine/Plugins/Experimental/Dataflow/Dataflow.uplugin
@@ -45,6 +45,10 @@
 		{
 			"Name": "GeometryProcessing",
 			"Enabled": true
+		},
+		{
+			"Name": "MeshModelingToolsetExp",
+			"Enabled": true
 		}
 	]
 }
diff --git a/Engine/Plugins/Experimental/Dataflow/Source/DataflowEditor/DataflowEditor.Build.cs b/Engine/Plugins/Experimental/Dataflow/Source/DataflowEditor/DataflowEditor.Build.cs
--- a/Engine/Plugins/Experimental/Dataflow/Source/DataflowEditor/DataflowEditor.Build.cs
+++ b/Engine/Plugins/Experimental/Dataflow/Source/DataflowEditor/DataflowEditor.Build.cs
@@ -31,10 +31,19 @@ namespace UnrealBuildTool.Rules
 					"DataflowNodes",
 					"Engine",
 					"EditorFramework",
+					"EditorInteractiveToolsFramework",
 					"EditorStyle",
 					"GraphEditor",
 					"InputCore",
 					"InteractiveToolsFramework",
+					"MeshModelingToolsEditorOnlyExp",
+					"ModelingComponentsEditorOnly",
+					"ModelingComponents",
+					"GeometryCore",
+					"GeometryFramework",
+					"MeshDescription",
+					"StaticMeshDescription",
+					"MeshConversion",
 					"LevelEditor",
 					"Slate",
 				    "SlateCore",
@@ -47,7 +56,9 @@ namespace UnrealBuildTool.Rules
 					"ToolMenus",
 					"ToolWidgets",
 					"Slate",
-					"XmlParser"
+					"XmlParser",
+					"TypedElementRuntime",
+					"MeshModelingToolsExp"
 				}
 			);
 		}
diff --git a/Engine/Plugins/Experimental/Dataflow/Source/DataflowEditor/Private/Dataflow/DataflowComponentToolTarget.cpp b/Engine/Plugins/Experimental/Dataflow/Source/DataflowEditor/Private/Dataflow/DataflowComponentToolTarget.cpp
new file mode 100644
--- /dev/null
+++ b/Engine/Plugins/Experimental/Dataflow/Source/DataflowEditor/Private/Dataflow/DataflowComponentToolTarget.cpp
@@ -0,0 +1,276 @@
+// Copyright Epic Games, Inc. All Rights Reserved.
+
+#include "Dataflow/DataflowComponentToolTarget.h"
+
+#include "ConversionUtils/DynamicMeshViaMeshDescriptionUtil.h"
+#include "Dataflow/DataflowComponent.h"
+#include "DynamicMeshToMeshDescription.h"
+#include "GeometryCollection/Facades/CollectionRenderingFacade.h"
+#include "MeshDescriptionToDynamicMesh.h"
+#include "MaterialDomain.h"
+#include "Materials/Material.h"
+#include "Materials/MaterialInterface.h"
+#include "StaticMeshAttributes.h"
+#include "UObject/Package.h"
+
+#include UE_INLINE_GENERATED_CPP_BY_NAME(DataflowComponentToolTarget)
+
+using namespace UE::Geometry;
+
+namespace UE
+{
+	namespace Conversion
+	{
+		// Convert a dataflow component to a dynamic mesh
+		void DataflowComponentToDynamicMesh(const UDataflowComponent* DataflowComponent, FDynamicMesh3& DynamicMesh)
+		{
+			const FManagedArrayCollection& RenderingCollection = DataflowComponent->GetRenderingCollection();
+			const GeometryCollection::Facades::FRenderingFacade Facade(RenderingCollection);
+
+			if(Facade.CanRenderSurface())
+			{
+				const int32 NumTriangles = Facade.NumTriangles();
+				const int32 NumVertices = Facade.NumVertices();
+				
+				const TManagedArray<FIntVector>& Indices = Facade.GetIndices();
+				const TManagedArray<FVector3f>& Positions = Facade.GetVertices();
+				const TManagedArray<FVector3f>& Normals = Facade.GetNormals();
+				const TManagedArray<FLinearColor>& Colors = Facade.GetVertexColor();
+
+				DynamicMesh.Clear();
+				for(int32 VertexIndex = 0; VertexIndex < NumVertices; ++VertexIndex)
+				{
+					DynamicMesh.AppendVertex(FVertexInfo(FVector3d(Positions[VertexIndex]), Normals[VertexIndex],
+						FVector3f(Colors[VertexIndex].R, Colors[VertexIndex].G, Colors[VertexIndex].B)));
+				}
+				for(int32 TriangleIndex = 0; TriangleIndex < NumTriangles; ++TriangleIndex)
+				{
+					DynamicMesh.AppendTriangle(FIndex3i(Indices[TriangleIndex].X, Indices[TriangleIndex].Y, Indices[TriangleIndex].Z));
+				}
+			}
+		}
+		
+		// Convert a dynamic mesh to a dataflow component
+		void DynamicMeshToDataflowComponent(const FDynamicMesh3& DynamicMesh, UDataflowComponent* DataflowComponent)
+		{
+			FManagedArrayCollection& RenderingCollection = DataflowComponent->ModifyRenderingCollection();
+			GeometryCollection::Facades::FRenderingFacade Facade(RenderingCollection);
+
+			if(Facade.CanRenderSurface())
+			{
+				const int32 NumTriangles = Facade.NumTriangles();
+				const int32 NumVertices = Facade.NumVertices();
+
+				// We can only override vertices attributes (position, normals, colors)
+				if((NumTriangles == DynamicMesh.TriangleCount()) && (NumVertices == DynamicMesh.VertexCount()))
+				{
+					TManagedArray<FVector3f>& Positions = Facade.ModifyVertices();
+                    TManagedArray<FVector3f>& Normals = Facade.ModifyNormals();
+                    TManagedArray<FLinearColor>& Colors = Facade.ModifyVertexColor();
+
+					for(int32 VertexIndex = 0; VertexIndex < NumVertices; ++VertexIndex)
+					{
+						Positions[VertexIndex] = FVector3f(DynamicMesh.GetVertex(VertexIndex));
+						Normals[VertexIndex] = DynamicMesh.GetVertexNormal(VertexIndex);
+						Colors[VertexIndex] = DynamicMesh.GetVertexColor(VertexIndex);
+					}
+				}
+			}
+		}
+	}
+}
+
+//
+// UDataflowComponentReadOnlyToolTarget
+//
+
+bool UDataflowComponentReadOnlyToolTarget::IsValid() const
+{
+	if (!UPrimitiveComponentToolTarget::IsValid())
+	{
+		return false;
+	}
+	UDataflowComponent* DataflowComponent = Cast<UDataflowComponent>(Component);
+	if (DataflowComponent == nullptr)
+	{
+		return false;
+	}
+	const UDataflow* Dataflow = DataflowComponent->GetDataflow();
+
+	return Dataflow != nullptr;
+}
+
+int32 UDataflowComponentReadOnlyToolTarget::GetNumMaterials() const
+{
+	return ensure(IsValid()) ? Component->GetNumMaterials() : 0;
+}
+
+UMaterialInterface* UDataflowComponentReadOnlyToolTarget::GetMaterial(int32 MaterialIndex) const
+{
+	return ensure(IsValid()) ? Component->GetMaterial(MaterialIndex) : nullptr;
+}
+
+void UDataflowComponentReadOnlyToolTarget::GetMaterialSet(FComponentMaterialSet& MaterialSetOut, bool bPreferAssetMaterials) const
+{
+	if (!ensure(IsValid())) return;
+
+	const int32 NumMaterials = Component->GetNumMaterials();
+	MaterialSetOut.Materials.SetNum(NumMaterials);
+	for (int32 k = 0; k < NumMaterials; ++k)
+	{
+		MaterialSetOut.Materials[k] = Component->GetMaterial(k);
+	}
+}
+
+bool UDataflowComponentReadOnlyToolTarget::CommitMaterialSetUpdate(const FComponentMaterialSet& MaterialSet, bool bApplyToAsset)
+{
+	// filter out any Engine materials that we don't want to be permanently assigning
+	TArray<UMaterialInterface*> FilteredMaterials = MaterialSet.Materials;
+	for (int32 k = 0; k < FilteredMaterials.Num(); ++k)
+	{
+		FString AssetPath = FilteredMaterials[k]->GetPathName();
+		if (AssetPath.StartsWith(TEXT("/MeshModelingToolsetExp/")))
+		{
+			FilteredMaterials[k] = UMaterial::GetDefaultMaterial(MD_Surface);
+		}
+	}
+
+	UDataflowComponent* DynamicMeshComponent = Cast<UDataflowComponent>(Component);
+	const int32 NumMaterialsGiven = FilteredMaterials.Num();
+
+	DynamicMeshComponent->Modify();
+	for (int32 k = 0; k < NumMaterialsGiven; ++k)
+	{
+		DynamicMeshComponent->SetMaterial(k, FilteredMaterials[k]);
+	}
+
+	return true;
+}
+
+const FMeshDescription* UDataflowComponentReadOnlyToolTarget::GetMeshDescription(const FGetMeshParameters& GetMeshParams)
+{
+	if (ensure(IsValid()) && !DataflowMeshDescription.IsValid())
+	{
+		DataflowMeshDescription = MakeUnique<FMeshDescription>();
+		FStaticMeshAttributes Attributes(*DataflowMeshDescription);
+		Attributes.Register();
+
+		const FDynamicMesh3 DynamicMesh = GetDynamicMesh();
+		
+		FDynamicMeshToMeshDescription Converter;
+		Converter.Convert(&DynamicMesh, *DataflowMeshDescription, true);
+	}
+	return DataflowMeshDescription.Get();
+}
+
+FMeshDescription UDataflowComponentReadOnlyToolTarget::GetEmptyMeshDescription()
+{
+	FMeshDescription EmptyMeshDescription;
+	FStaticMeshAttributes MeshAttributes(EmptyMeshDescription);
+	MeshAttributes.Register();
+
+	return EmptyMeshDescription;
+}
+
+FDynamicMesh3 UDataflowComponentReadOnlyToolTarget::GetDynamicMesh()
+{
+	FDynamicMesh3 DynamicMesh;
+	if( const UDataflowComponent* DataflowComponent = Cast<UDataflowComponent>(Component))
+	{
+		UE::Conversion::DataflowComponentToDynamicMesh(DataflowComponent, DynamicMesh);
+	}
+	return DynamicMesh;
+}
+
+FDynamicMesh3 UDataflowComponentReadOnlyToolTarget::GetDynamicMesh(bool bRequestTangents)
+{
+	return GetDynamicMesh();
+}
+
+//
+// UDataflowComponentToolTarget
+//
+
+
+void UDataflowComponentToolTarget::CommitMeshDescription(const FCommitter& Committer, const FCommitMeshParameters& CommitMeshParams)
+{
+	if (ensure(IsValid()) == false) return;
+
+	// Let the user fill our mesh description with the Committer
+	if (!DataflowMeshDescription.IsValid())
+	{
+		DataflowMeshDescription = MakeUnique<FMeshDescription>();
+		FStaticMeshAttributes Attributes(*DataflowMeshDescription);
+		Attributes.Register();
+	}
+	FCommitterParams CommitParams;
+	CommitParams.MeshDescriptionOut = DataflowMeshDescription.Get();
+	Committer(CommitParams);
+
+	// The conversion we have right now is from dynamic mesh to volume, so we convert
+	// to dynamic mesh first.
+	FDynamicMesh3 DynamicMesh;
+	FMeshDescriptionToDynamicMesh Converter;
+	Converter.Convert(DataflowMeshDescription.Get(), DynamicMesh);
+
+	CommitDynamicMesh(DynamicMesh);
+}
+
+void UDataflowComponentToolTarget::CommitDynamicMesh(const FDynamicMesh3& DynamicMesh, const FDynamicMeshCommitInfo& CommitInfo)
+{
+	if(UDataflowComponent* DataflowComponent = Cast<UDataflowComponent>(Component))
+	{
+		UE::Conversion::DynamicMeshToDataflowComponent(DynamicMesh, DataflowComponent);
+	}
+}
+
+// Factory
+
+bool UDataflowComponentReadOnlyToolTargetFactory::CanBuildTarget(UObject* SourceObject, const FToolTargetTypeRequirements& Requirements) const
+{
+	// We are using an exact cast here to prevent subclasses, which might not meet all
+	// requirements for functionality such as the deprecated DestructibleMeshComponent, from 
+	// being caught up as valid targets.
+	// If you want to make the tool target work with some subclass of UDataflowComponent,
+	// just add another factory that allows that class specifically(but make sure that
+	// GetMeshDescription and such work properly)
+
+	return Cast<UDataflowComponent>(SourceObject) && Cast<UDataflowComponent>(SourceObject)->GetDataflow() &&
+		ExactCast<UDataflow>(Cast<UDataflowComponent>(SourceObject)->GetDataflow()) &&
+		!ExactCast<UDataflow>(Cast<UDataflowComponent>(SourceObject)->GetDataflow())->GetOutermost()->bIsCookedForEditor &&
+		Requirements.AreSatisfiedBy(UDataflowComponentReadOnlyToolTarget::StaticClass());
+}
+
+UToolTarget* UDataflowComponentReadOnlyToolTargetFactory::BuildTarget(UObject* SourceObject, const FToolTargetTypeRequirements& Requirements)
+{
+	UDataflowComponentReadOnlyToolTarget* Target = NewObject<UDataflowComponentReadOnlyToolTarget>();
+	Target->InitializeComponent(Cast<UDataflowComponent>(SourceObject));
+	checkSlow(Target->Component.IsValid() && Requirements.AreSatisfiedBy(Target));
+
+	return Target;
+}
+
+bool UDataflowComponentToolTargetFactory::CanBuildTarget(UObject* SourceObject, const FToolTargetTypeRequirements& Requirements) const
+{
+	// We are using an exact cast here to prevent subclasses, which might not meet all
+	// requirements for functionality such as the deprecated DestructibleMeshComponent, from 
+	// being caught up as valid targets.
+	// If you want to make the tool target work with some subclass of UDataflowComponent,
+	// just add another factory that allows that class specifically(but make sure that
+	// GetMeshDescription and such work properly)
+
+	return Cast<UDataflowComponent>(SourceObject) && Cast<UDataflowComponent>(SourceObject)->GetDataflow() &&
+		ExactCast<UDataflow>(Cast<UDataflowComponent>(SourceObject)->GetDataflow()) &&
+		!ExactCast<UDataflow>(Cast<UDataflowComponent>(SourceObject)->GetDataflow())->GetOutermost()->bIsCookedForEditor &&
+		Requirements.AreSatisfiedBy(UDataflowComponentToolTarget::StaticClass());
+}
+
+UToolTarget* UDataflowComponentToolTargetFactory::BuildTarget(UObject* SourceObject, const FToolTargetTypeRequirements& Requirements)
+{
+	UDataflowComponentToolTarget* Target = NewObject<UDataflowComponentToolTarget>();
+	Target->InitializeComponent(Cast<UDataflowComponent>(SourceObject));
+	checkSlow(Target->Component.IsValid() && Requirements.AreSatisfiedBy(Target));
+
+	return Target;
+}
+
diff --git a/Engine/Plugins/Experimental/Dataflow/Source/DataflowEditor/Private/Dataflow/DataflowEditor.cpp b/Engine/Plugins/Experimental/Dataflow/Source/DataflowEditor/Private/Dataflow/DataflowEditor.cpp
--- a/Engine/Plugins/Experimental/Dataflow/Source/DataflowEditor/Private/Dataflow/DataflowEditor.cpp
+++ b/Engine/Plugins/Experimental/Dataflow/Source/DataflowEditor/Private/Dataflow/DataflowEditor.cpp
@@ -1,17 +1,117 @@
 // Copyright Epic Games, Inc. All Rights Reserved.
 
 #include "Dataflow/DataflowEditor.h"
+
+#include "Animation/Skeleton.h"
 #include "Dataflow/DataflowEditorToolkit.h"
+#include "Engine/SkeletalMesh.h"
 
 #include UE_INLINE_GENERATED_CPP_BY_NAME(DataflowEditor)
 
 DEFINE_LOG_CATEGORY(LogDataflowEditor);
 
+namespace Private
+{
+	UDataflow* GetDataflowAssetFrom(UObject* InObject)
+	{
+		if (UClass* Class = InObject->GetClass())
+		{
+			if (FProperty* Property = Class->FindPropertyByName(FName("DataflowAsset")))
+			{
+				return *Property->ContainerPtrToValuePtr<UDataflow*>(InObject);
+			}
+		}
+		return nullptr;
+	}
+
+	USkeletalMesh* GetSkeletalMeshFrom(UObject* InObject)
+	{
+		if (UClass* Class = InObject->GetClass())
+		{
+			if (FProperty* Property = Class->FindPropertyByName(FName("SkeletalMesh")))
+			{
+				return *Property->ContainerPtrToValuePtr<USkeletalMesh*>(InObject);
+			}
+		}
+		return nullptr;
+	}
+	
+	UAnimationAsset* GetAnimationAssetFrom(UObject* InObject)
+	{
+		if (UClass* Class = InObject->GetClass())
+		{
+			if (FProperty* Property = Class->FindPropertyByName(FName("AnimationAsset")))
+			{
+				return *Property->ContainerPtrToValuePtr<UAnimationAsset*>(InObject);
+			}
+		}
+		return nullptr;
+	}
+
+	FString GetDataflowTerminalFrom(UObject* InObject)
+	{
+		if (UClass* Class = InObject->GetClass())
+		{
+			if (FProperty* Property = Class->FindPropertyByName(FName("DataflowTerminal")))
+			{
+				return *Property->ContainerPtrToValuePtr<FString>(InObject);
+			}
+		}
+		return FString();
+	}
+};
+
 TSharedPtr<FBaseAssetToolkit> UDataflowEditor::CreateToolkit()
 {
 	TSharedPtr<FDataflowEditorToolkit> DataflowToolkit = MakeShared<FDataflowEditorToolkit>(this);
 	return DataflowToolkit;
 }
 
+void UDataflowEditor::Initialize(const TArray<TObjectPtr<UObject>>& InObjects)
+{
+	// We extract the dataflow and the associated datas from the list of objects
+	TArray<TObjectPtr<UObject>> ObjectsToEdit;
+	DataflowDatas.bHasValidSkeletalMesh = false;
+	if (ensure(InObjects.Num() == 1))
+	{
+		TObjectPtr<UObject> RootObject = InObjects[0];
+		DataflowDatas.DataflowAsset = Cast<UDataflow>(RootObject);
+		if(!DataflowDatas.DataflowAsset)
+		{
+			DataflowDatas.DataflowAsset = Private::GetDataflowAssetFrom(RootObject);
+			DataflowDatas.DataflowTerminal = Private::GetDataflowTerminalFrom(RootObject);
+			DataflowDatas.SkeletalMesh = Private::GetSkeletalMeshFrom(RootObject);
+			DataflowDatas.AnimationAsset = Private::GetAnimationAssetFrom(RootObject);
+		}
+		if(DataflowDatas.DataflowAsset)
+		{
+			DataflowDatas.DataflowOwner = RootObject;
+			DataflowDatas.DataflowAsset->Schema = UDataflowSchema::StaticClass();
+			DataflowDatas.DataflowContext = MakeShared<Dataflow::FEngineContext>(RootObject, DataflowDatas.DataflowAsset, FPlatformTime::Cycles64());
+			DataflowDatas.LastNodeTimestamp = DataflowDatas.DataflowContext->GetTimestamp();
+
+			ObjectsToEdit.Add(DataflowDatas.DataflowOwner);
+
+			if(!DataflowDatas.SkeletalMesh)
+			{
+				const FName SkeletonName = MakeUniqueObjectName(DataflowDatas.DataflowAsset, UDataflow::StaticClass(), FName("USkeleton"));
+				const FName SkeletalMeshName = MakeUniqueObjectName(DataflowDatas.DataflowAsset, UDataflow::StaticClass(), FName("USkeletalMesh"));
+
+				DataflowDatas.Skeleton = NewObject<USkeleton>(DataflowDatas.DataflowAsset, SkeletonName);
+				DataflowDatas.SkeletalMesh = NewObject<USkeletalMesh>(DataflowDatas.DataflowAsset, SkeletalMeshName);
+				DataflowDatas.SkeletalMesh->SetSkeleton(DataflowDatas.Skeleton);
+			}
+			else
+			{
+				DataflowDatas.Skeleton = DataflowDatas.SkeletalMesh->GetSkeleton();
+				DataflowDatas.bHasValidSkeletalMesh = true;
+			}
+		}
+	}
+	// Potentially we could add additional objects to edit here (fields, meshes....)
+	// If these objects have a matching factory we would be able to use geometry tools
+	UBaseCharacterFXEditor::Initialize(ObjectsToEdit);
+}
+
 
 
diff --git a/Engine/Plugins/Experimental/Dataflow/Source/DataflowEditor/Private/Dataflow/DataflowEditorCommands.cpp b/Engine/Plugins/Experimental/Dataflow/Source/DataflowEditor/Private/Dataflow/DataflowEditorCommands.cpp
--- a/Engine/Plugins/Experimental/Dataflow/Source/DataflowEditor/Private/Dataflow/DataflowEditorCommands.cpp
+++ b/Engine/Plugins/Experimental/Dataflow/Source/DataflowEditor/Private/Dataflow/DataflowEditorCommands.cpp
@@ -6,6 +6,7 @@
 #include "Dataflow/DataflowNodeFactory.h"
 #include "Dataflow/DataflowObject.h"
 #include "Dataflow/DataflowOverrideNode.h"
+#include "Dataflow/DataflowEditorStyle.h"
 #include "EdGraph/EdGraphNode.h"
 #include "IStructureDetailsView.h"
 #include "EdGraphNode_Comment.h"
@@ -16,17 +17,21 @@
 
 #define LOCTEXT_NAMESPACE "DataflowEditorCommands"
 
+const FString FDataflowEditorCommandsImpl::BeginAttributeEditorToolIdentifier = TEXT("BeginAttributeEditorTool");
+const FString FDataflowEditorCommandsImpl::BeginMeshSelectionToolIdentifier = TEXT("BeginMeshSelectionTool");
+
 FDataflowEditorCommandsImpl::FDataflowEditorCommandsImpl()
 	: TBaseCharacterFXEditorCommands<FDataflowEditorCommandsImpl>("DataflowEditor", 
 		LOCTEXT("ContextDescription", "Dataflow Editor"), 
 		NAME_None,
-		FAppStyle::GetAppStyleSetName())
+		FDataflowEditorStyle::Get().GetStyleSetName())
 {
 }
 
-
 void FDataflowEditorCommandsImpl::RegisterCommands()
 {
+	TBaseCharacterFXEditorCommands::RegisterCommands();
+	
 	UI_COMMAND(EvaluateNode, "Evaluate", "Trigger an evaluation of the selected node.", EUserInterfaceActionType::Button, FInputChord());
 	UI_COMMAND(CreateComment, "CreateComment", "Create a Comment node.", EUserInterfaceActionType::None, FInputChord());
 	UI_COMMAND(ToggleEnabledState, "ToggleEnabledState", "Toggle node between Enabled/Disabled state.", EUserInterfaceActionType::Button, FInputChord());
@@ -36,7 +41,9 @@ void FDataflowEditorCommandsImpl::RegisterCommands()
 	UI_COMMAND(AddOptionPin, "AddOptionPin", "Add an option pin to the selected nodes.", EUserInterfaceActionType::Button, FInputChord());
 	UI_COMMAND(RemoveOptionPin, "RemoveOptionPin", "Remove the last option pin from the selected nodes.", EUserInterfaceActionType::Button, FInputChord());
 	UI_COMMAND(ZoomToFitGraph, "ZoomToFitGraph", "Fit the graph in the graph editor viewport.", EUserInterfaceActionType::None, FInputChord(EKeys::F));
-	
+
+	UI_COMMAND(BeginAttributeEditorTool, "AttrEd", "Edit/configure mesh attributes", EUserInterfaceActionType::Button, FInputChord());
+	UI_COMMAND(BeginMeshSelectionTool, "SelVtx", "Select mesh vertices", EUserInterfaceActionType::Button, FInputChord());
 
 	if (Dataflow::FNodeFactory* Factory = Dataflow::FNodeFactory::GetInstance())
 	{
@@ -77,8 +84,6 @@ void FDataflowEditorCommands::Unregister()
 	return FDataflowEditorCommandsImpl::Unregister();
 }
 
-
-
 void FDataflowEditorCommands::EvaluateSelectedNodes(const FGraphPanelSelectionSet& SelectedNodes, FDataflowEditorCommands::FGraphEvaluationCallback Evaluate)
 {
 	const bool bIsInPIEOrSimulate = GEditor->PlayWorld != NULL || GEditor->bIsSimulatingInEditor;
diff --git a/Engine/Plugins/Experimental/Dataflow/Source/DataflowEditor/Private/Dataflow/DataflowEditorMode.cpp b/Engine/Plugins/Experimental/Dataflow/Source/DataflowEditor/Private/Dataflow/DataflowEditorMode.cpp
--- a/Engine/Plugins/Experimental/Dataflow/Source/DataflowEditor/Private/Dataflow/DataflowEditorMode.cpp
+++ b/Engine/Plugins/Experimental/Dataflow/Source/DataflowEditor/Private/Dataflow/DataflowEditorMode.cpp
@@ -1,9 +1,33 @@
 // Copyright Epic Games, Inc. All Rights Reserved.
 
 #include "Dataflow/DataflowEditorMode.h"
-#include "Dataflow/DataflowEditorToolkit.h"
 #include "Dataflow/DataflowEditorModeToolkit.h"
-
+#include "Dataflow/DataflowComponentToolTarget.h"
+#include "Dataflow/DataflowToolTarget.h"
+#include "Dataflow/DataflowGraphEditor.h"
+#include "Dataflow/DataflowSNode.h"
+#include "Dataflow/DataflowEditorCommands.h"
+#include "Dataflow/DataflowPreviewScene.h"
+#include "Dataflow/DataflowEditor.h"
+#include "TargetInterfaces/MaterialProvider.h"
+#include "TargetInterfaces/DynamicMeshCommitter.h"
+#include "TargetInterfaces/DynamicMeshProvider.h"
+#include "EdModeInteractiveToolsContext.h"
+#include "ToolTargetManager.h"
+#include "ToolTargets/DynamicMeshComponentToolTarget.h"
+#include "ToolTargets/StaticMeshComponentToolTarget.h"
+#include "ToolTargets/StaticMeshToolTarget.h"
+#include "ToolTargets/SkeletalMeshComponentToolTarget.h"
+#include "ToolTargets/SkeletalMeshToolTarget.h"
+#include "Tools/UEdMode.h"
+#include "AttributeEditorTool.h"
+#include "MeshSelectionTool.h"
+#include "MeshVertexPaintTool.h"
+#include "MeshAttributePaintTool.h"
+#include "ModelingToolTargetUtil.h"
+#include "Components/DynamicMeshComponent.h"
+#include "EditorModeManager.h"
+#include "Selection.h"
 
 #include UE_INLINE_GENERATED_CPP_BY_NAME(DataflowEditorMode)
 
@@ -11,22 +35,189 @@
 
 const FEditorModeID UDataflowEditorMode::EM_DataflowEditorModeId = TEXT("EM_DataflowAssetEditorMode");
 
-
-
-
 UDataflowEditorMode::UDataflowEditorMode()
 {
 	Info = FEditorModeInfo(
 		EM_DataflowEditorModeId,
-		LOCTEXT("DataflowEditorModeName", "Dataflow"),
+		LOCTEXT("DataflowEditorModeName", "Dataflow Tools"),
 		FSlateIcon(),
 		false);
 }
 
+void UDataflowEditorMode::Exit()
+{
+	UActorComponent::MarkRenderStateDirtyEvent.RemoveAll(this);
+	PreviewScene->ResetDynamicMeshComponents();
+	
+	Super::Exit();
+}
+
+void UDataflowEditorMode::Enter()
+{
+	Super::Enter();
+}
+
+void UDataflowEditorMode::ModeTick(float DeltaTime)
+{
+	Super::ModeTick(DeltaTime);
+	if (PreviewScene)
+	{
+		PreviewScene->GetWorld()->Tick(ELevelTick::LEVELTICK_All, DeltaTime);
+	}
+}
+
+const FToolTargetTypeRequirements& UDataflowEditorMode::GetToolTargetRequirements()
+{
+	static const FToolTargetTypeRequirements ToolTargetRequirements =
+		FToolTargetTypeRequirements({
+			UMaterialProvider::StaticClass(),
+			UDynamicMeshCommitter::StaticClass(),
+			UDynamicMeshProvider::StaticClass()
+			});
+
+	return ToolTargetRequirements;
+}
+
+void UDataflowEditorMode::AddToolTargetFactories()
+{
+	GetInteractiveToolsContext()->TargetManager->AddTargetFactory(NewObject<UDynamicMeshComponentToolTargetFactory>(GetToolManager()));
+	GetInteractiveToolsContext()->TargetManager->AddTargetFactory(NewObject<UStaticMeshComponentToolTargetFactory>(GetToolManager()));
+	GetInteractiveToolsContext()->TargetManager->AddTargetFactory(NewObject<UStaticMeshToolTargetFactory>(GetToolManager()));
+	GetInteractiveToolsContext()->TargetManager->AddTargetFactory(NewObject<USkeletalMeshComponentToolTargetFactory>(GetToolManager()));
+	GetInteractiveToolsContext()->TargetManager->AddTargetFactory(NewObject<USkeletalMeshToolTargetFactory>(GetToolManager()));
+	GetInteractiveToolsContext()->TargetManager->AddTargetFactory(NewObject<UDataflowComponentToolTargetFactory>(GetToolManager()));
+	GetInteractiveToolsContext()->TargetManager->AddTargetFactory(NewObject<UDataflowToolTargetFactory>(GetToolManager()));
+}
+
+void UDataflowEditorMode::RegisterTools()
+{
+	const FDataflowEditorCommandsImpl& CommandInfos = FDataflowEditorCommands::Get();
+	RegisterTool(CommandInfos.BeginAttributeEditorTool, FDataflowEditorCommandsImpl::BeginAttributeEditorToolIdentifier, NewObject<UAttributeEditorToolBuilder>());
+
+	// For testing...
+	//RegisterTool(CommandInfos.BeginMeshSelectionTool, FDataflowEditorCommandsImpl::BeginMeshSelectionToolIdentifier, NewObject<UMeshSelectionToolBuilder>());
+	//RegisterTool(CommandInfos.BeginMeshSelectionTool, FDataflowEditorCommandsImpl::BeginMeshSelectionToolIdentifier, NewObject<UMeshVertexPaintToolBuilder>());
+	RegisterTool(CommandInfos.BeginMeshSelectionTool, FDataflowEditorCommandsImpl::BeginMeshSelectionToolIdentifier, NewObject<UMeshAttributePaintToolBuilder>());
+}
+
 void UDataflowEditorMode::CreateToolkit()
 {
 	Toolkit = MakeShared<FDataflowEditorModeToolkit>();
 }
 
+void UDataflowEditorMode::BindCommands()
+{
+	const FDataflowEditorCommandsImpl& CommandInfos = FDataflowEditorCommands::Get();
+	const TSharedRef<FUICommandList>& CommandList = Toolkit->GetToolkitCommands();
+
+	// Hook up to Enter/Esc key presses
+	CommandList->MapAction(
+		CommandInfos.AcceptOrCompleteActiveTool,
+		FExecuteAction::CreateLambda([this]()
+		{
+			AcceptActiveToolActionOrTool();
+			
+		}),
+		FCanExecuteAction::CreateLambda([this]() {
+		return GetInteractiveToolsContext()->CanAcceptActiveTool() || GetInteractiveToolsContext()->CanCompleteActiveTool();
+	}),
+		FGetActionCheckState(),
+		FIsActionButtonVisible(),
+		EUIActionRepeatMode::RepeatDisabled);
+
+	CommandList->MapAction(
+		CommandInfos.CancelOrCompleteActiveTool,
+		FExecuteAction::CreateLambda([this]() { CancelActiveToolActionOrTool(); }),
+		FCanExecuteAction::CreateLambda([this]() {
+		return GetInteractiveToolsContext()->CanCompleteActiveTool() || GetInteractiveToolsContext()->CanCancelActiveTool();
+	}),
+		FGetActionCheckState(),
+		FIsActionButtonVisible(),
+		EUIActionRepeatMode::RepeatDisabled);
+}
+
+void UDataflowEditorMode::CreateToolTargets(const TArray<TObjectPtr<UObject>>& AssetsIn)
+{
+	ToolTargets.Reset();
+	for (UObject* Object : AssetsIn)
+	{
+		if (UToolTarget* Target = GetInteractiveToolsContext()->TargetManager->BuildTarget(Object, GetToolTargetRequirements()))
+		{
+			ToolTargets.Add(Target);
+		}
+	}
+}
+
+void UDataflowEditorMode::InitializeTargets(const TArray<TObjectPtr<UObject>>& AssetsIn)
+{
+	UBaseCharacterFXEditorMode::InitializeTargets(AssetsIn);
+	
+	USelection* SelectedComponents = GetModeManager()->GetSelectedComponents();
+	SelectedComponents->Modify();
+	SelectedComponents->BeginBatchSelectOperation();
+	SelectedComponents->DeselectAll();
+
+	PreviewScene->ResetDynamicMeshComponents();
+	for (UToolTarget* Target : ToolTargets)
+	{
+		UE::Geometry::FDynamicMesh3 DynamicMesh = UE::ToolTarget::GetDynamicMeshCopy(Target);
+
+		// @todo(DynamicMeshRendering) : Enable Dynamic Mesh Rendering for dataflow terminals. 
+		//  TObjectPtr<UDynamicMeshComponent>& DynamicMeshComponent =
+		//	PreviewScene->AddDynamicMeshComponent(MoveTemp(DynamicMesh), UE::ToolTarget::GetMaterialSet(Target).Materials);
+		//
+		//SelectedComponents->Select(DynamicMeshComponent);
+	}
+	SelectedComponents->EndBatchSelectOperation();
+}
+
+// UEdGraphNode* UDataflowEditorMode::CreateNewNode(const FName& NewNodeTypeName)
+// {
+// 	const TSharedPtr<const SDataflowGraphEditor> PinnedDataflowGraphEditor = GraphEditor.Pin();
+// 	if (!PinnedDataflowGraphEditor)
+// 	{
+// 		return nullptr;
+// 	}
+//
+// 	const TSharedPtr<FAssetSchemaAction_Dataflow_CreateNode_DataflowEdNode> NodeAction =
+// 		FAssetSchemaAction_Dataflow_CreateNode_DataflowEdNode::CreateAction(PreviewScene->GetDataflowDatas().DataflowAsset, NewNodeTypeName);
+// 	constexpr UEdGraphPin* FromPin = nullptr;
+// 	constexpr bool bSelectNewNode = true;
+// 	UEdGraphNode* const NewEdNode = NodeAction->PerformAction(PreviewScene->GetDataflowDatas().DataflowAsset, FromPin,
+// 			PinnedDataflowGraphEditor->GetPasteLocation(), bSelectNewNode);
+//
+// 	PreviewScene->GetDataflowDatas().DataflowAsset->NotifyGraphChanged();
+//
+// 	return NewEdNode;
+// }
+//
+// void UDataflowEditorMode::RegisterAddNodeCommand(TSharedPtr<FUICommandInfo> AddNodeCommand, const FName& NewNodeType, TSharedPtr<FUICommandInfo> StartToolCommand)
+// {
+// 	auto AddNode = [this](const FName& NewNodeType)
+// 	{
+// 		const FName ConnectionType = FManagedArrayCollection::StaticType();
+// 		const FName ConnectionName("Collection");
+// 		
+// 		const UEdGraphNode* const NewNode = CreateNewNode(NewNodeType);
+// 		verifyf(NewNode, TEXT("Failed to create a new node: %s"), *NewNodeType.ToString());
+//
+// 		StartToolForSelectedNode(NewNode);
+// 	};
+//
+// 	const TSharedRef<FUICommandList>& CommandList = Toolkit->GetToolkitCommands();
+//
+// 	CommandList->MapAction(AddNodeCommand,
+// 		FExecuteAction::CreateWeakLambda(this, AddNode, NewNodeType),
+// 		FCanExecuteAction::CreateLambda([]() {return true;})
+// 	);
+//
+// 	NodeToolMap.Add(NewNodeType, StartToolCommand);
+// }
+
+FBox UDataflowEditorMode::SceneBoundingBox() const
+{
+	return PreviewScene->GetBoundingBox();
+}
+
 #undef LOCTEXT_NAMESPACE
 
diff --git a/Engine/Plugins/Experimental/Dataflow/Source/DataflowEditor/Private/Dataflow/DataflowEditorModeToolkit.cpp b/Engine/Plugins/Experimental/Dataflow/Source/DataflowEditor/Private/Dataflow/DataflowEditorModeToolkit.cpp
--- a/Engine/Plugins/Experimental/Dataflow/Source/DataflowEditor/Private/Dataflow/DataflowEditorModeToolkit.cpp
+++ b/Engine/Plugins/Experimental/Dataflow/Source/DataflowEditor/Private/Dataflow/DataflowEditorModeToolkit.cpp
@@ -6,17 +6,24 @@
 #include "Dataflow/DataflowEditorCommands.h"
 #include "Dataflow/DataflowEditorStyle.h"
 #include "InteractiveToolManager.h"
+#include "MeshAttributePaintTool.h"
 
 #define LOCTEXT_NAMESPACE "FDataflowEditorModeToolkit"
 
 void FDataflowEditorModeToolkit::OnToolStarted(UInteractiveToolManager* Manager, UInteractiveTool* Tool)
 {
-	ensure(Tool == Manager->GetActiveTool(EToolSide::Left));
+	FBaseCharacterFXEditorModeToolkit::OnToolStarted(Manager, Tool);
+}
 
-	ActiveToolName = Tool->GetToolInfo().ToolDisplayName;
-	FString ActiveToolIdentifier = Manager->GetActiveToolName(EToolSide::Mouse);
-	ActiveToolIdentifier.InsertAt(0, ".");
-	ActiveToolIcon = GetActiveToolIcon(ActiveToolIdentifier);
+void FDataflowEditorModeToolkit::BuildToolPalette(FName PaletteIndex, class FToolBarBuilder& ToolbarBuilder)
+{
+	const FDataflowEditorCommandsImpl& Commands = FDataflowEditorCommands::Get();
+	if (PaletteIndex == ToolsTabName)
+	{
+		// @todo(DynamicMeshRendering) : Enable Dynamic Mesh Rendering for dataflow terminals. 
+		//ToolbarBuilder.AddToolBarButton(Commands.BeginAttributeEditorTool);
+		//ToolbarBuilder.AddToolBarButton(Commands.BeginMeshSelectionTool);
+	}
 }
 
 FName FDataflowEditorModeToolkit::GetToolkitFName() const
diff --git a/Engine/Plugins/Experimental/Dataflow/Source/DataflowEditor/Private/Dataflow/DataflowEditorModule.cpp b/Engine/Plugins/Experimental/Dataflow/Source/DataflowEditor/Private/Dataflow/DataflowEditorModule.cpp
--- a/Engine/Plugins/Experimental/Dataflow/Source/DataflowEditor/Private/Dataflow/DataflowEditorModule.cpp
+++ b/Engine/Plugins/Experimental/Dataflow/Source/DataflowEditor/Private/Dataflow/DataflowEditorModule.cpp
@@ -3,40 +3,34 @@
 #include "Dataflow/DataflowEditorModule.h"
 #include "Dataflow/DataflowEditorStyle.h"
 #include "Dataflow/DataflowEditorMode.h"
+#include "Dataflow/DataflowEditorToolkit.h"
+#include "Dataflow/DataflowEditorCommands.h"
+#include "Dataflow/DataflowSNodeFactories.h"
 
-#include "AssetToolsModule.h"
 #include "CoreMinimal.h"
 #include "EdGraphUtilities.h"
-#include "Dataflow/DataflowEditorToolkit.h"
-#include "Dataflow/DataflowNodeFactory.h"
-#include "Dataflow/AssetDefinition_DataflowAsset.h"
-#include "Dataflow/DataflowSNodeFactories.h"
 #include "PropertyEditorModule.h"
 #include "EditorModeRegistry.h"
 
 #define LOCTEXT_NAMESPACE "DataflowEditor"
 
-//#define BOX_BRUSH(StyleSet, RelativePath, ...) FSlateBoxBrush(StyleSet->RootToContentDir(RelativePath, TEXT(".png")), __VA_ARGS__)
-//#define DEFAULT_FONT(...) FCoreStyle::GetDefaultFontStyle(__VA_ARGS__)
-
 void FDataflowEditorModule::StartupModule()
 {
 	FDataflowEditorStyle::Get();
-
+	
 	FPropertyEditorModule& PropertyEditorModule = FModuleManager::LoadModuleChecked<FPropertyEditorModule>(TEXT("PropertyEditor"));
-
+	
 	DataflowSNodeFactory = MakeShareable(new FDataflowSNodeFactory());
 	FEdGraphUtilities::RegisterVisualNodeFactory(DataflowSNodeFactory);
 }
 
 void FDataflowEditorModule::ShutdownModule()
-{
+{	
 	if (UObjectInitialized())
 	{
 		FEdGraphUtilities::UnregisterVisualNodeFactory(DataflowSNodeFactory);
 	}
 	FEditorModeRegistry::Get().UnregisterMode(UDataflowEditorMode::EM_DataflowEditorModeId);
-
 }
 
 IMPLEMENT_MODULE(FDataflowEditorModule, DataflowEditor)
diff --git a/Engine/Plugins/Experimental/Dataflow/Source/DataflowEditor/Private/Dataflow/DataflowEditorStyle.h b/Engine/Plugins/Experimental/Dataflow/Source/DataflowEditor/Private/Dataflow/DataflowEditorStyle.h
--- a/Engine/Plugins/Experimental/Dataflow/Source/DataflowEditor/Private/Dataflow/DataflowEditorStyle.h
+++ b/Engine/Plugins/Experimental/Dataflow/Source/DataflowEditor/Private/Dataflow/DataflowEditorStyle.h
@@ -6,6 +6,7 @@
 #include "Styling/SlateStyleRegistry.h"
 #include "Interfaces/IPluginManager.h"
 #include "Misc/Paths.h"
+#include "Dataflow/DataflowEditorCommands.h"
 
 class FDataflowEditorStyle final : public FSlateStyleSet
 {
@@ -33,6 +34,11 @@ public:
 		Set("Dataflow.SelectFace", new FSlateImageBrush(RootToContentDir(TEXT("Slate/Dataflow_SelectFace_40x.png")), Icon40x40));
 		Set("Dataflow.SelectVertex", new FSlateImageBrush(RootToContentDir(TEXT("Slate/Dataflow_SelectVertex_40x.png")), Icon40x40));
 
+		const FString AttributeEditorPropertyName = "DataflowEditor." + FDataflowEditorCommandsImpl::BeginAttributeEditorToolIdentifier;
+		Set(*AttributeEditorPropertyName, new FSlateImageBrush(RootToContentDir(TEXT("Slate/Dataflow_SelectObject20x.png")), Icon20x20));
+
+		const FString MeshSelectionPropertyName = "DataflowEditor." + FDataflowEditorCommandsImpl::BeginMeshSelectionToolIdentifier;
+		Set(*MeshSelectionPropertyName, new FSlateImageBrush(RootToContentDir(TEXT("Slate/Dataflow_SelectVertex20x.png")), Icon20x20));
 
 		FSlateStyleRegistry::RegisterSlateStyle(*this);
 	}
diff --git a/Engine/Plugins/Experimental/Dataflow/Source/DataflowEditor/Private/Dataflow/DataflowEditorToolkit.cpp b/Engine/Plugins/Experimental/Dataflow/Source/DataflowEditor/Private/Dataflow/DataflowEditorToolkit.cpp
--- a/Engine/Plugins/Experimental/Dataflow/Source/DataflowEditor/Private/Dataflow/DataflowEditorToolkit.cpp
+++ b/Engine/Plugins/Experimental/Dataflow/Source/DataflowEditor/Private/Dataflow/DataflowEditorToolkit.cpp
@@ -15,6 +15,7 @@
 #include "Dataflow/DataflowNodeFactory.h"
 #include "Dataflow/DataflowObject.h"
 #include "Dataflow/DataflowObjectInterface.h"
+#include "Dataflow/DataflowPreviewScene.h"
 #include "Dataflow/DataflowSchema.h"
 #include "DynamicMeshBuilder.h"
 #include "EditorModeManager.h"
@@ -30,6 +31,7 @@
 #include "Framework/Commands/GenericCommands.h"
 #include "Widgets/Docking/SDockTab.h"
 #include "ISkeletonTree.h"
+#include "AssetEditorModeManager.h"
 #include "ISkeletonEditorModule.h"
 #include "IStructureDetailsView.h"
 #include "Misc/MessageDialog.h"
@@ -38,9 +40,7 @@
 
 //DEFINE_LOG_CATEGORY_STATIC(EditorToolkitLog, Log, All);
 
-const FName FDataflowEditorToolkit::ViewportTabId(TEXT("DataflowEditor_Viewport"));
 const FName FDataflowEditorToolkit::GraphCanvasTabId(TEXT("DataflowEditor_GraphCanvas"));
-const FName FDataflowEditorToolkit::AssetDetailsTabId(TEXT("DataflowEditor_AssetDetails"));
 const FName FDataflowEditorToolkit::NodeDetailsTabId(TEXT("DataflowEditor_NodeDetails"));
 const FName FDataflowEditorToolkit::SkeletalTabId(TEXT("DataflowEditor_Skeletal"));
 const FName FDataflowEditorToolkit::SelectionViewTabId_1(TEXT("DataflowEditor_SelectionView_1"));
@@ -52,54 +52,12 @@ const FName FDataflowEditorToolkit::CollectionSpreadSheetTabId_2(TEXT("DataflowE
 const FName FDataflowEditorToolkit::CollectionSpreadSheetTabId_3(TEXT("DataflowEditor_CollectionSpreadSheet_3"));
 const FName FDataflowEditorToolkit::CollectionSpreadSheetTabId_4(TEXT("DataflowEditor_CollectionSpreadSheet_4"));
 
-namespace Private
-{
-	UDataflow* GetDataflowFrom(UObject* InObject)
-	{
-		if (UClass* Class = InObject->GetClass())
-		{
-			if (FProperty* Property = Class->FindPropertyByName(FName("DataflowAsset")))
-			{
-				return *Property->ContainerPtrToValuePtr<UDataflow*>(InObject);
-			}
-		}
-		return nullptr;
-
-	}
-
-	USkeletalMesh* GetSkeletalMeshFrom(UObject* InObject)
-	{
-		if (UClass* Class = InObject->GetClass())
-		{
-			if (FProperty* Property = Class->FindPropertyByName(FName("SkeletalMesh")))
-			{
-				return *Property->ContainerPtrToValuePtr<USkeletalMesh*>(InObject);
-			}
-		}
-		return nullptr;
-
-	}
-
-	FString GetDataflowTerminalFrom(UObject* InObject)
-	{
-		if (UClass* Class = InObject->GetClass())
-		{
-			if (FProperty* Property = Class->FindPropertyByName(FName("DataflowTerminal")))
-			{
-				return *Property->ContainerPtrToValuePtr<FString>(InObject);
-			}
-		}
-		return FString();
-	}
-};
-
-
 FDataflowEditorToolkit::FDataflowEditorToolkit(UAssetEditor* InOwningAssetEditor)
 	: FBaseCharacterFXEditorToolkit(InOwningAssetEditor, FName("DataflowEditor")) 
 {
 	check(Cast<UDataflowEditor>(InOwningAssetEditor));
 
-	StandaloneDefaultLayout = FTabManager::NewLayout("Dataflow_Layout.V3")
+	StandaloneDefaultLayout = FTabManager::NewLayout("Dataflow_Layout.V5")
 		->AddArea
 		(
 			FTabManager::NewPrimaryArea()->SetOrientation(Orient_Vertical)
@@ -124,15 +82,17 @@ FDataflowEditorToolkit::FDataflowEditorToolkit(UAssetEditor* InOwningAssetEditor
 						(
 							FTabManager::NewStack()
 							->SetSizeCoefficient(0.7f)
-							->AddTab(AssetDetailsTabId, ETabState::OpenedTab)
+							->AddTab(DetailsTabID, ETabState::OpenedTab)
 						)
 					)
 				)
 			)
 		);
 
-	// Construction view scene @todo(DataflowMode) : Add a FDataflowPreviewScene to the Toolkit here
-	// ObjectScene = MakeUnique<FAdvancedPreviewScene>(FAdvancedPreviewScene::ConstructionValues().SetCreateDefaultLighting(false));
+	FAdvancedPreviewScene::ConstructionValues PreviewSceneArgs;
+	PreviewSceneArgs.bShouldSimulatePhysics = 1;
+	PreviewSceneArgs.bCreatePhysicsScene = 1;
+	ObjectScene = MakeUnique<FDataflowPreviewScene>(PreviewSceneArgs, ModifyDataflowEditorDatas());
 }
 
 
@@ -158,16 +118,23 @@ FDataflowEditorToolkit::~FDataflowEditorToolkit()
 	// will end up getting destroyed before the mode's Exit() function gets to run, and we'll get some
 	// warnings when we destroy any mode actors.
 	EditorModeManager->DestroyMode(UDataflowEditorMode::EM_DataflowEditorModeId);
+	
+}
+
+void FDataflowEditorToolkit::CreateEditorModeManager()
+{
+	FBaseCharacterFXEditorToolkit::CreateEditorModeManager();
+	static_cast<FDataflowPreviewScene*>(ObjectScene.Get())->ModifyDataflowModeManager()
+		= StaticCastSharedPtr<FAssetEditorModeManager>(EditorModeManager);
 }
 
 bool FDataflowEditorToolkit::CanOpenDataflowEditor(UObject* ObjectToEdit)
 {
-	UDataflow* Dataflow = Cast<UDataflow>(ObjectToEdit);
-	if (Dataflow == nullptr)
+	if (const UClass* Class = ObjectToEdit->GetClass())
 	{
-		Dataflow = Private::GetDataflowFrom(ObjectToEdit);
+		return Class->FindPropertyByName(FName("DataflowAsset")) != nullptr;
 	}
-	return Dataflow != nullptr;
+	return false;
 }
 
 bool FDataflowEditorToolkit::HasDataflowAsset(UObject* ObjectToEdit)
@@ -182,6 +149,40 @@ bool FDataflowEditorToolkit::HasDataflowAsset(UObject* ObjectToEdit)
 	return false;
 }
 
+UDataflow* FDataflowEditorToolkit::GetDataflowAsset(UObject* ObjectToEdit)
+{
+	UDataflow* DataflowObject = Cast<UDataflow>(ObjectToEdit);
+
+	if (!DataflowObject)
+	{
+		if (const UClass* Class = ObjectToEdit->GetClass())
+		{
+			if (FProperty* Property = Class->FindPropertyByName(FName("DataflowAsset")))
+			{
+				DataflowObject = *Property->ContainerPtrToValuePtr<UDataflow*>(ObjectToEdit);
+			}
+		}
+	}
+	return DataflowObject;
+}
+
+const UDataflow* FDataflowEditorToolkit::GetDataflowAsset(const UObject* ObjectToEdit)
+{
+	const UDataflow* DataflowObject = Cast<UDataflow>(ObjectToEdit);
+
+	if (!DataflowObject)
+	{
+		if (const UClass* Class = ObjectToEdit->GetClass())
+		{
+			if (const FProperty* Property = Class->FindPropertyByName(FName("DataflowAsset")))
+			{
+				DataflowObject = *Property->ContainerPtrToValuePtr<const UDataflow*>(ObjectToEdit);
+			}
+		}
+	}
+	return DataflowObject;
+}
+
 //~ Begin FBaseCharacterFXEditorToolkit overrides
 
 FEditorModeID FDataflowEditorToolkit::GetEditorModeId() const
@@ -189,6 +190,16 @@ FEditorModeID FDataflowEditorToolkit::GetEditorModeId() const
 	return UDataflowEditorMode::EM_DataflowEditorModeId;
 }
 
+const FDataflowEditorDatas& FDataflowEditorToolkit::GetDataflowEditorDatas() const
+{
+	return Cast<UDataflowEditor>(OwningAssetEditor)->DataflowDatas;
+}
+
+FDataflowEditorDatas& FDataflowEditorToolkit::ModifyDataflowEditorDatas()
+{
+	return Cast<UDataflowEditor>(OwningAssetEditor)->DataflowDatas;
+}
+
 bool FDataflowEditorToolkit::OnRequestClose(EAssetEditorCloseReason InCloseReason)
 {
 	// Note: This needs a bit of adjusting, because currently OnRequestClose seems to be 
@@ -215,7 +226,7 @@ bool FDataflowEditorToolkit::OnRequestClose(EAssetEditorCloseReason InCloseReaso
 void FDataflowEditorToolkit::PostInitAssetEditor()
 {
 	// @todo(DataflowMode) : Enable base class functionality when the mode is ready.
-	// FBaseCharacterFXEditorToolkit::PostInitAssetEditor();
+	FBaseCharacterFXEditorToolkit::PostInitAssetEditor();
 	
 	// @todo(DataflowMode) : Access to mode in Toolkit construction
 	//UDataflowEditorMode* const DataflowEditorMode = CastChecked<UDataflowEditorMode>(EditorModeManager->GetActiveScriptableMode(UDataflowEditorMode::EM_DataflowEditorModeId));
@@ -250,6 +261,10 @@ void FDataflowEditorToolkit::PostInitAssetEditor()
 
 		// If exposure isn't set to fixed, it will flash as we stare into the void
 		Client->ExposureSettings.bFixed = true;
+
+		// We need the viewport client to start out focused, or else it won't get ticked until
+		// we click inside it.
+		Client->ReceivedFocus(Client->Viewport);
 	};
 	SetCommonViewportClientOptions(ViewportClient.Get());
 
@@ -257,6 +272,8 @@ void FDataflowEditorToolkit::PostInitAssetEditor()
 	TSharedPtr<FDataflowEditorViewportClient> DataflowViewportClient = StaticCastSharedPtr<FDataflowEditorViewportClient>(ViewportClient);
 	DataflowViewportClient->SetDataflowEditorToolkit(StaticCastSharedRef<FDataflowEditorToolkit>(this->AsShared()));
 
+	//SetEditingObject(GetDataflowEditorDatas().DataflowOwner);
+
 	// @todo(DataflowMode) : Do this in UDataflowEditorMode::RefocusViewportClient and use the DataflowComponents bounds
 	//		... see ChaosClothAssetEditorMode::RefocusRestSpaceViewportClient() for reference
 	// FBoxSphereBounds SphereBounds = FBoxSphereBounds(EForceInit::ForceInitToZero);
@@ -266,8 +283,15 @@ void FDataflowEditorToolkit::PostInitAssetEditor()
 
 void FDataflowEditorToolkit::InitializeEdMode(UBaseCharacterFXEditorMode* EdMode)
 {
-	UDataflowEditorMode* DataflowMode = Cast<UDataflowEditorMode>(EdMode);
-	check(DataflowMode);
+	// We first set the preview scene in order to store the dynamic mesh elements
+	// generated by the tools
+	Cast<UDataflowEditorMode>(EdMode)->SetDataflowPreviewScene(
+		static_cast<FDataflowPreviewScene*>(ObjectScene.Get()));
+
+	// Set of the graph editor to be able to add nodes
+	Cast<UDataflowEditorMode>(EdMode)->SetDataflowGraphEditor(GraphEditor);
+	
+	FBaseCharacterFXEditorToolkit::InitializeEdMode(EdMode);
 }
 
 void FDataflowEditorToolkit::CreateEditorModeUILayer()
@@ -279,10 +303,10 @@ void FDataflowEditorToolkit::GetSaveableObjects(TArray<UObject*>& OutObjects) co
 {
 	FBaseCharacterFXEditorToolkit::GetSaveableObjects(OutObjects);
 
-	if (Dataflow)
+	if (UDataflow* DataflowAsset = GetDataflowEditorDatas().DataflowAsset)
 	{
-		check(Dataflow->IsAsset());
-		OutObjects.Add(Dataflow);
+		check(DataflowAsset->IsAsset());
+		OutObjects.Add(DataflowAsset);
 	}
 }
 
@@ -295,38 +319,13 @@ void FDataflowEditorToolkit::CreateWidgets()
 {
 	FBaseCharacterFXEditorToolkit::CreateWidgets();
 
-	Asset = nullptr;
-	UObject* const ObjectToEdit = GetAsset();
-
-	if (ObjectToEdit)
-	{
-		Dataflow = Cast<UDataflow>(ObjectToEdit);
-		if (Dataflow == nullptr)
-		{
-			Dataflow = Private::GetDataflowFrom(ObjectToEdit);
-
-			// TODO: Figure out how to create the GraphEditor widgets when the Input Asset doesn't have a Dataflow property set
-			if (Dataflow)
-			{
-				DataflowTerminalPath = Private::GetDataflowTerminalFrom(Dataflow);
-
-				Dataflow->Schema = UDataflowSchema::StaticClass();
-			}
-		}
-		Asset = ObjectToEdit;
-	}
-
-	if (Dataflow != nullptr)
+	const FDataflowEditorDatas& DataflowEditorDatas = GetDataflowEditorDatas();
+	if (UDataflow* DataflowAsset = DataflowEditorDatas.DataflowAsset)
 	{
-		Context = TSharedPtr< Dataflow::FEngineContext>(new Dataflow::FAssetContext(Asset, Dataflow, FPlatformTime::Cycles64()));
-		LastNodeTimestamp = Context->GetTimestamp();
-
-		Dataflow->Schema = UDataflowSchema::StaticClass();
-
-		NodeDetailsEditor = CreateNodeDetailsEditorWidget(ObjectToEdit);
-		AssetDetailsEditor = CreateAssetDetailsEditorWidget(ObjectToEdit);
-		GraphEditor = CreateGraphEditorWidget(Dataflow, NodeDetailsEditor);
-		SkeletalEditor = CreateSkeletalEditorWidget(ObjectToEdit);
+		NodeDetailsEditor = CreateNodeDetailsEditorWidget(DataflowEditorDatas.DataflowOwner);
+		AssetDetailsEditor = CreateAssetDetailsEditorWidget(DataflowEditorDatas.DataflowOwner);
+		GraphEditor = CreateGraphEditorWidget(DataflowAsset, NodeDetailsEditor);
+		SkeletalEditor = CreateSkeletalEditorWidget(DataflowEditorDatas.DataflowOwner);
 	}
 }
 
@@ -360,12 +359,14 @@ TSharedPtr<FEditorViewportClient> FDataflowEditorToolkit::CreateEditorViewportCl
 
 void FDataflowEditorToolkit::OnPropertyValueChanged(const FPropertyChangedEvent& PropertyChangedEvent)
 {
-	FDataflowEditorCommands::OnPropertyValueChanged(this->GetDataflow(), Context, LastNodeTimestamp, PropertyChangedEvent, PrevNodeSelection);
+	FDataflowEditorDatas& DataflowEditorDatas = ModifyDataflowEditorDatas();
+	FDataflowEditorCommands::OnPropertyValueChanged(DataflowEditorDatas.DataflowAsset, DataflowEditorDatas.DataflowContext, DataflowEditorDatas.LastNodeTimestamp, PropertyChangedEvent, PrevNodeSelection);
 }
 
 void FDataflowEditorToolkit::OnAssetPropertyValueChanged(const FPropertyChangedEvent& PropertyChangedEvent)
 {
-	FDataflowEditorCommands::OnAssetPropertyValueChanged(this->GetDataflow(), Context, LastNodeTimestamp, PropertyChangedEvent);
+	FDataflowEditorDatas& DataflowEditorDatas = ModifyDataflowEditorDatas();
+	FDataflowEditorCommands::OnAssetPropertyValueChanged(DataflowEditorDatas.DataflowAsset, DataflowEditorDatas.DataflowContext, DataflowEditorDatas.LastNodeTimestamp, PropertyChangedEvent);
 }
 
 bool FDataflowEditorToolkit::OnNodeVerifyTitleCommit(const FText& NewText, UEdGraphNode* GraphNode, FText& OutErrorMessage)
@@ -380,7 +381,7 @@ void FDataflowEditorToolkit::OnNodeTitleCommitted(const FText& InNewText, ETextC
 
 void FDataflowEditorToolkit::OnNodeSelectionChanged(const TSet<UObject*>& NewSelection)
 {
-	if (Dataflow)
+	if(GetDataflowEditorDatas().DataflowAsset)
 	{
 		// Only keep UDataflowEdNode from NewSelection
 		TSet<UObject*> ValidatedSelection;
@@ -448,15 +449,16 @@ void FDataflowEditorToolkit::OnNodeDeleted(const TSet<UObject*>& NewSelection)
 
 void FDataflowEditorToolkit::Tick(float DeltaTime)
 {
-	if (Dataflow && Asset)
+	FDataflowEditorDatas& DataflowEditorDatas = ModifyDataflowEditorDatas();
+	if(DataflowEditorDatas.IsValid())
 	{
-		if (!Context)
+		if (!DataflowEditorDatas.DataflowContext)
 		{
-			Context = TSharedPtr< Dataflow::FEngineContext>(new Dataflow::FAssetContext(Asset, Dataflow, Dataflow::FTimestamp::Invalid));
-			LastNodeTimestamp = Dataflow::FTimestamp::Invalid;
+			DataflowEditorDatas.DataflowContext = MakeShared<Dataflow::FEngineContext>(DataflowEditorDatas.DataflowOwner, DataflowEditorDatas.DataflowAsset, Dataflow::FTimestamp::Invalid);
+			DataflowEditorDatas.LastNodeTimestamp = Dataflow::FTimestamp::Invalid;
 		}
-		DataflowTerminalPath = Private::GetDataflowTerminalFrom(Asset);
-		FDataflowEditorCommands::EvaluateTerminalNode(*Context.Get(), LastNodeTimestamp, Dataflow, nullptr, nullptr, Asset, DataflowTerminalPath);
+		FDataflowEditorCommands::EvaluateTerminalNode(*DataflowEditorDatas.DataflowContext.Get(), DataflowEditorDatas.LastNodeTimestamp,
+			DataflowEditorDatas.DataflowAsset, nullptr, nullptr, DataflowEditorDatas.DataflowOwner, DataflowEditorDatas.DataflowTerminal);
 	}
 }
 
@@ -465,8 +467,6 @@ TStatId FDataflowEditorToolkit::GetStatId() const
 	RETURN_QUICK_DECLARE_CYCLE_STAT(FDataflowEditorToolkit, STATGROUP_Tickables);
 }
 
-
-
 TSharedRef<SDataflowGraphEditor> FDataflowEditorToolkit::CreateGraphEditorWidget(UDataflow* DataflowToEdit, TSharedPtr<IStructureDetailsView> InNodeDetailsEditor)
 {
 	ensure(DataflowToEdit);
@@ -474,14 +474,19 @@ TSharedRef<SDataflowGraphEditor> FDataflowEditorToolkit::CreateGraphEditorWidget
 
 	FDataflowEditorCommands::FGraphEvaluationCallback Evaluate = [&](FDataflowNode* Node, FDataflowOutput* Out)
 	{
-		if (!Context)
+		FDataflowEditorDatas& DataflowEditorDatas = ModifyDataflowEditorDatas();
+		if(DataflowEditorDatas.IsValid())
 		{
-			Context = TSharedPtr< Dataflow::FEngineContext>(new Dataflow::FAssetContext(Asset, Dataflow, Dataflow::FTimestamp::Invalid));
+			if (!DataflowEditorDatas.DataflowContext)
+			{
+				DataflowEditorDatas.DataflowContext = MakeShared<Dataflow::FEngineContext>(DataflowEditorDatas.DataflowOwner, DataflowEditorDatas.DataflowAsset, Dataflow::FTimestamp::Invalid);
+			}
+			Node->Invalidate();
+			DataflowEditorDatas.LastNodeTimestamp = Dataflow::FTimestamp::Invalid;
+			
+			FDataflowEditorCommands::EvaluateTerminalNode(*DataflowEditorDatas.DataflowContext.Get(), DataflowEditorDatas.LastNodeTimestamp,
+				DataflowEditorDatas.DataflowAsset, Node, Out, DataflowEditorDatas.DataflowOwner, DataflowEditorDatas.DataflowTerminal);
 		}
-		Node->Invalidate();
-		LastNodeTimestamp = Dataflow::FTimestamp::Invalid;
-
-		FDataflowEditorCommands::EvaluateTerminalNode(*Context.Get(), LastNodeTimestamp, Dataflow, Node, Out, Asset, DataflowTerminalPath);
 	};
 
 	SGraphEditor::FGraphEditorEvents InEvents;
@@ -557,57 +562,35 @@ TSharedPtr<IDetailsView> FDataflowEditorToolkit::CreateAssetDetailsEditorWidget(
 
 TSharedPtr<ISkeletonTree> FDataflowEditorToolkit::CreateSkeletalEditorWidget(UObject* ObjectToEdit)
 {
-	if (Dataflow)
+	const FDataflowEditorDatas& DataflowEditorDatas = GetDataflowEditorDatas();
+	if(DataflowEditorDatas.DataflowAsset)
 	{
-		if (!StubSkeletalMesh)
-		{
-			const FName NodeName = MakeUniqueObjectName(Dataflow, UDataflow::StaticClass(), FName("USkeleton"));
-			StubSkeleton = NewObject<USkeleton>(Dataflow, NodeName);
-			const FName NodeName2 = MakeUniqueObjectName(Dataflow, UDataflow::StaticClass(), FName("USkeleton"));
-			StubSkeletalMesh = NewObject<USkeletalMesh>(Dataflow, NodeName2);
-			StubSkeletalMesh->SetSkeleton(StubSkeleton);
-		}
-
 		FSkeletonTreeArgs SkeletonTreeArgs;
-		//SkeletonTreeArgs.OnSelectionChanged = FOnSkeletonTreeSelectionChanged::CreateSP(this, &FAnimationBlueprintEditor::HandleSelectionChanged);
-		//SkeletonTreeArgs.PreviewScene = GetPreviewScene();
-		//SkeletonTreeArgs.ContextName = GetToolkitFName();
-
-		USkeleton* Skeleton = StubSkeleton;
-		if (Asset)
-		{
-			if (USkeletalMesh* SkeletalMesh = Private::GetSkeletalMeshFrom(Asset))
-			{
-				Skeleton = SkeletalMesh->GetSkeleton();
-			}
-		}
 		ISkeletonEditorModule& SkeletonEditorModule = FModuleManager::LoadModuleChecked<ISkeletonEditorModule>("SkeletonEditor");
-		TSharedPtr<ISkeletonTree> SkeletonTree = SkeletonEditorModule.CreateSkeletonTree(Skeleton, SkeletonTreeArgs);
+		TSharedPtr<ISkeletonTree> SkeletonTree = SkeletonEditorModule.CreateSkeletonTree(DataflowEditorDatas.Skeleton, SkeletonTreeArgs);
 		return SkeletonTree;
 	}
 	return TSharedPtr<ISkeletonTree>(nullptr);
 }
 
-UObject* FDataflowEditorToolkit::GetAsset() const
+TSharedRef<SDockTab> FDataflowEditorToolkit::SpawnTab_AssetDetails(const FSpawnTabArgs& Args)
 {
-	TArray<TObjectPtr<UObject>> ObjectsToEdit;
-	OwningAssetEditor->GetObjectsToEdit(MutableView(ObjectsToEdit));
-
-	UObject* ObjectToEdit = nullptr;
-	if (ensure(ObjectsToEdit.Num() == 1))
-	{
-		ObjectToEdit = ObjectsToEdit[0];
-	}
+	check(Args.GetTabId() == DetailsTabID);
 
-	return Cast<UObject>(ObjectToEdit);
+	return SNew(SDockTab)
+		.Label(LOCTEXT("DataflowEditor_AssetDetails_TabTitle", "Asset Details"))
+		[
+			AssetDetailsEditor->AsShared()
+		];
 }
 
+
 TSharedRef<SDockTab> FDataflowEditorToolkit::SpawnTab_GraphCanvas(const FSpawnTabArgs& Args)
 {
 	check(Args.GetTabId() == GraphCanvasTabId);
 
 	TSharedRef<SDockTab> SpawnedTab = SNew(SDockTab)
-		.Label(LOCTEXT("DataflowEditor_Dataflow_TabTitle", "Graph"))
+		.Label(LOCTEXT("DataflowEditor_Dataflow_TabTitle", "Dataflow Graph"))
 		[
 			GraphEditor.ToSharedRef()
 		];
@@ -615,17 +598,6 @@ TSharedRef<SDockTab> FDataflowEditorToolkit::SpawnTab_GraphCanvas(const FSpawnTa
 	return SpawnedTab;
 }
 
-TSharedRef<SDockTab> FDataflowEditorToolkit::SpawnTab_AssetDetails(const FSpawnTabArgs& Args)
-{
-	check(Args.GetTabId() == AssetDetailsTabId);
-
-	return SNew(SDockTab)
-		.Label(LOCTEXT("DataflowEditor_AssetDetails_TabTitle", "Details"))
-		[
-			AssetDetailsEditor->AsShared()
-		];
-}
-
 TSharedRef<SDockTab> FDataflowEditorToolkit::SpawnTab_NodeDetails(const FSpawnTabArgs& Args)
 {
 	check(Args.GetTabId() == NodeDetailsTabId);
@@ -640,23 +612,16 @@ TSharedRef<SDockTab> FDataflowEditorToolkit::SpawnTab_NodeDetails(const FSpawnTa
 TSharedRef<SDockTab> FDataflowEditorToolkit::SpawnTab_Skeletal(const FSpawnTabArgs& Args)
 {
 	check(Args.GetTabId() == SkeletalTabId);
-
-	USkeletalMesh* SkeletalMesh = StubSkeletalMesh;
-	if (Asset)
-	{
-		SkeletalMesh = Private::GetSkeletalMeshFrom(Asset);
-	}
-
-	SkeletalEditor->SetSkeletalMesh(SkeletalMesh);
+	const FDataflowEditorDatas& DataflowEditorDatas = GetDataflowEditorDatas();
+	SkeletalEditor->SetSkeletalMesh(DataflowEditorDatas.SkeletalMesh);
 
 	return SNew(SDockTab)
-		.Label(LOCTEXT("FleshEditorSkeletal_TabTitle", "Skeletal Hierarchy"))
+		.Label(LOCTEXT("DataflowEditor_Skeletal_TabTitle", "Skeletal Hierarchy"))
 		[
 			SkeletalEditor.ToSharedRef()
 		];
 }
 
-
 TSharedRef<SDockTab> FDataflowEditorToolkit::SpawnTab_SelectionView(const FSpawnTabArgs& Args)
 {
 //	check(Args.GetTabId().TabType == SelectionViewTabId_1);
@@ -703,12 +668,13 @@ TSharedRef<SDockTab> FDataflowEditorToolkit::SpawnTab_SelectionView(const FSpawn
 
 	if (SelectionViewWidget)
 	{
+		const FDataflowEditorDatas& DataflowEditorDatas = GetDataflowEditorDatas();
 		if (Args.GetTabId() == SelectionViewTabId_1)
 		{
 			DataflowSelectionView_1->SetSelectionView(SelectionViewWidget);
 
 			// Set the Context on the interface
-			if (TSharedPtr<Dataflow::FContext> CurrentContext = this->GetContext())
+			if (TSharedPtr<Dataflow::FContext> CurrentContext = DataflowEditorDatas.DataflowContext)
 			{
 				DataflowSelectionView_1->SetContext(CurrentContext);
 			}
@@ -718,7 +684,7 @@ TSharedRef<SDockTab> FDataflowEditorToolkit::SpawnTab_SelectionView(const FSpawn
 			DataflowSelectionView_2->SetSelectionView(SelectionViewWidget);
 
 			// Set the Context on the interface
-			if (TSharedPtr<Dataflow::FContext> CurrentContext = this->GetContext())
+			if (TSharedPtr<Dataflow::FContext> CurrentContext = DataflowEditorDatas.DataflowContext)
 			{
 				DataflowSelectionView_2->SetContext(CurrentContext);
 			}
@@ -728,7 +694,7 @@ TSharedRef<SDockTab> FDataflowEditorToolkit::SpawnTab_SelectionView(const FSpawn
 			DataflowSelectionView_3->SetSelectionView(SelectionViewWidget);
 
 			// Set the Context on the interface
-			if (TSharedPtr<Dataflow::FContext> CurrentContext = this->GetContext())
+			if (TSharedPtr<Dataflow::FContext> CurrentContext = DataflowEditorDatas.DataflowContext)
 			{
 				DataflowSelectionView_3->SetContext(CurrentContext);
 			}
@@ -738,7 +704,7 @@ TSharedRef<SDockTab> FDataflowEditorToolkit::SpawnTab_SelectionView(const FSpawn
 			DataflowSelectionView_4->SetSelectionView(SelectionViewWidget);
 
 			// Set the Context on the interface
-			if (TSharedPtr<Dataflow::FContext> CurrentContext = this->GetContext())
+			if (TSharedPtr<Dataflow::FContext> CurrentContext = DataflowEditorDatas.DataflowContext)
 			{
 				DataflowSelectionView_4->SetContext(CurrentContext);
 			}
@@ -752,8 +718,6 @@ TSharedRef<SDockTab> FDataflowEditorToolkit::SpawnTab_SelectionView(const FSpawn
 
 TSharedRef<SDockTab> FDataflowEditorToolkit::SpawnTab_CollectionSpreadSheet(const FSpawnTabArgs& Args)
 {
-//	check(Args.GetTabId() == CollectionSpreadSheetTabId_1);
-
 	if (Args.GetTabId() == CollectionSpreadSheetTabId_1)
 	{
 		DataflowCollectionSpreadSheet_1 = MakeShared<FDataflowCollectionSpreadSheet>(FDataflowCollectionSpreadSheet());
@@ -796,12 +760,13 @@ TSharedRef<SDockTab> FDataflowEditorToolkit::SpawnTab_CollectionSpreadSheet(cons
 
 	if (CollectionSpreadSheetWidget)
 	{
+		const FDataflowEditorDatas& DataflowEditorDatas = GetDataflowEditorDatas();
 		if (Args.GetTabId() == CollectionSpreadSheetTabId_1)
 		{
 			DataflowCollectionSpreadSheet_1->SetCollectionSpreadSheet(CollectionSpreadSheetWidget);
 
 			// Set the Context on the interface
-			if (TSharedPtr<Dataflow::FContext> CurrentContext = this->GetContext())
+			if (TSharedPtr<Dataflow::FContext> CurrentContext = DataflowEditorDatas.DataflowContext)
 			{
 				DataflowCollectionSpreadSheet_1->SetContext(CurrentContext);
 			}
@@ -811,7 +776,7 @@ TSharedRef<SDockTab> FDataflowEditorToolkit::SpawnTab_CollectionSpreadSheet(cons
 			DataflowCollectionSpreadSheet_2->SetCollectionSpreadSheet(CollectionSpreadSheetWidget);
 
 			// Set the Context on the interface
-			if (TSharedPtr<Dataflow::FContext> CurrentContext = this->GetContext())
+			if (TSharedPtr<Dataflow::FContext> CurrentContext = DataflowEditorDatas.DataflowContext)
 			{
 				DataflowCollectionSpreadSheet_2->SetContext(CurrentContext);
 			}
@@ -821,7 +786,7 @@ TSharedRef<SDockTab> FDataflowEditorToolkit::SpawnTab_CollectionSpreadSheet(cons
 			DataflowCollectionSpreadSheet_3->SetCollectionSpreadSheet(CollectionSpreadSheetWidget);
 
 			// Set the Context on the interface
-			if (TSharedPtr<Dataflow::FContext> CurrentContext = this->GetContext())
+			if (TSharedPtr<Dataflow::FContext> CurrentContext = DataflowEditorDatas.DataflowContext)
 			{
 				DataflowCollectionSpreadSheet_3->SetContext(CurrentContext);
 			}
@@ -831,7 +796,7 @@ TSharedRef<SDockTab> FDataflowEditorToolkit::SpawnTab_CollectionSpreadSheet(cons
 			DataflowCollectionSpreadSheet_4->SetCollectionSpreadSheet(CollectionSpreadSheetWidget);
 
 			// Set the Context on the interface
-			if (TSharedPtr<Dataflow::FContext> CurrentContext = this->GetContext())
+			if (TSharedPtr<Dataflow::FContext> CurrentContext = DataflowEditorDatas.DataflowContext)
 			{
 				DataflowCollectionSpreadSheet_4->SetContext(CurrentContext);
 			}
@@ -845,29 +810,27 @@ TSharedRef<SDockTab> FDataflowEditorToolkit::SpawnTab_CollectionSpreadSheet(cons
 
 void FDataflowEditorToolkit::RegisterTabSpawners(const TSharedRef<FTabManager>& InTabManager)
 {
-	// We bypass FBaseAssetToolkit::RegisterTabSpawners because it doesn't seem to provide us with
-	// anything except tabs that we don't want.
 	FAssetEditorToolkit::RegisterTabSpawners(InTabManager);
 
 	EditorMenuCategory = InTabManager->AddLocalWorkspaceMenuCategory(LOCTEXT("WorkspaceMenu_DataflowEditor", "Dataflow Editor"));
-	TSharedRef<FWorkspaceItem> SelectionViewWorkspaceMenuCategoryRef = EditorMenuCategory->AddGroup(LOCTEXT("WorkspaceMenu_SelectionView", "Selection View"), FSlateIcon(FAppStyle::GetAppStyleSetName(), "LevelEditor.Tabs.Outliner"));
-	TSharedRef<FWorkspaceItem> CollectionSpreadSheetWorkspaceMenuCategoryRef = EditorMenuCategory->AddGroup(LOCTEXT("WorkspaceMenu_CollectionSpreadSheet", "Collection SpreadSheet"), FSlateIcon(FAppStyle::GetAppStyleSetName(), "LevelEditor.Tabs.Outliner"));
-
-	InTabManager->RegisterTabSpawner(ViewportTabId, FOnSpawnTab::CreateSP(this, &FDataflowEditorToolkit::SpawnTab_Viewport))
+	const TSharedRef<FWorkspaceItem> SelectionViewWorkspaceMenuCategoryRef = EditorMenuCategory->AddGroup(LOCTEXT("WorkspaceMenu_SelectionView", "Selection View"), FSlateIcon(FAppStyle::GetAppStyleSetName(), "LevelEditor.Tabs.Outliner"));
+	const TSharedRef<FWorkspaceItem> CollectionSpreadSheetWorkspaceMenuCategoryRef = EditorMenuCategory->AddGroup(LOCTEXT("WorkspaceMenu_CollectionSpreadSheet", "Collection SpreadSheet"), FSlateIcon(FAppStyle::GetAppStyleSetName(), "LevelEditor.Tabs.Outliner"));
+	
+	InTabManager->RegisterTabSpawner(ViewportTabID, FOnSpawnTab::CreateSP(this, &FDataflowEditorToolkit::SpawnTab_Viewport))
 		.SetDisplayName(LOCTEXT("DataflowViewportTab", "Dataflow Viewport"))
 		.SetGroup(EditorMenuCategory.ToSharedRef())
-		.SetIcon(FSlateIcon(FAppStyle::GetAppStyleSetName(), "GraphEditor.EventGraph_16x"));
+		.SetIcon(FSlateIcon(FAppStyle::GetAppStyleSetName(), "LevelEditor.Tabs.Viewports"));
 
+	InTabManager->RegisterTabSpawner(DetailsTabID, FOnSpawnTab::CreateSP(this, &FDataflowEditorToolkit::SpawnTab_AssetDetails))
+		.SetDisplayName(LOCTEXT("AssetDetailsTab", "Asset Details"))
+		.SetGroup(EditorMenuCategory.ToSharedRef())
+		.SetIcon(FSlateIcon(FAppStyle::GetAppStyleSetName(), "LevelEditor.Tabs.Details"));
+	
 	InTabManager->RegisterTabSpawner(GraphCanvasTabId, FOnSpawnTab::CreateSP(this, &FDataflowEditorToolkit::SpawnTab_GraphCanvas))
-		.SetDisplayName(LOCTEXT("DataflowTab", "Graph"))
+		.SetDisplayName(LOCTEXT("DataflowTab", "Dataflow Graph"))
 		.SetGroup(EditorMenuCategory.ToSharedRef())
 		.SetIcon(FSlateIcon(FAppStyle::GetAppStyleSetName(), "GraphEditor.EventGraph_16x"));
 
-	InTabManager->RegisterTabSpawner(AssetDetailsTabId, FOnSpawnTab::CreateSP(this, &FDataflowEditorToolkit::SpawnTab_AssetDetails))
-		.SetDisplayName(LOCTEXT("AssetDetailsTab", "Details"))
-		.SetGroup(EditorMenuCategory.ToSharedRef())
-		.SetIcon(FSlateIcon(FAppStyle::GetAppStyleSetName(), "LevelEditor.Tabs.Details"));
-
 	InTabManager->RegisterTabSpawner(NodeDetailsTabId, FOnSpawnTab::CreateSP(this, &FDataflowEditorToolkit::SpawnTab_NodeDetails))
 		.SetDisplayName(LOCTEXT("NodeDetailsTab", "Node Details"))
 		.SetGroup(EditorMenuCategory.ToSharedRef())
@@ -917,16 +880,13 @@ void FDataflowEditorToolkit::RegisterTabSpawners(const TSharedRef<FTabManager>&
 		.SetDisplayName(LOCTEXT("DataflowCollectionSpreadSheetTab4", "Collection SpreadSheet 4"))
 		.SetGroup(CollectionSpreadSheetWorkspaceMenuCategoryRef)
 		.SetIcon(FSlateIcon(FAppStyle::GetAppStyleSetName(), "LevelEditor.Tabs.Outliner"));
-
 }
 
 void FDataflowEditorToolkit::UnregisterTabSpawners(const TSharedRef<class FTabManager>& InTabManager)
 {
-	FAssetEditorToolkit::UnregisterTabSpawners(InTabManager);
+	FBaseAssetToolkit::UnregisterTabSpawners(InTabManager);
 
-	InTabManager->UnregisterTabSpawner(ViewportTabId);
 	InTabManager->UnregisterTabSpawner(GraphCanvasTabId);
-	InTabManager->UnregisterTabSpawner(AssetDetailsTabId);
 	InTabManager->UnregisterTabSpawner(NodeDetailsTabId);
 	InTabManager->UnregisterTabSpawner(SkeletalTabId);
 	InTabManager->UnregisterTabSpawner(SelectionViewTabId_1);
@@ -982,13 +942,14 @@ FName FDataflowEditorToolkit::GetToolkitFName() const
 
 FText FDataflowEditorToolkit::GetToolkitName() const
 {
-	if (Asset)
+	const FDataflowEditorDatas& DataflowEditorDatas = GetDataflowEditorDatas();
+	if(DataflowEditorDatas.DataflowOwner)
 	{
-		return  GetLabelForObject(Asset);
+		return  GetLabelForObject(DataflowEditorDatas.DataflowOwner);
 	}
-	if (Dataflow)
+	else if(DataflowEditorDatas.DataflowAsset)
 	{
-		return  GetLabelForObject(Dataflow);
+		return  GetLabelForObject(DataflowEditorDatas.DataflowAsset);
 	}
 	return  LOCTEXT("ToolkitName", "Empty Dataflow Editor");
 }
diff --git a/Engine/Plugins/Experimental/Dataflow/Source/DataflowEditor/Private/Dataflow/DataflowEditorViewport.cpp b/Engine/Plugins/Experimental/Dataflow/Source/DataflowEditor/Private/Dataflow/DataflowEditorViewport.cpp
--- a/Engine/Plugins/Experimental/Dataflow/Source/DataflowEditor/Private/Dataflow/DataflowEditorViewport.cpp
+++ b/Engine/Plugins/Experimental/Dataflow/Source/DataflowEditor/Private/Dataflow/DataflowEditorViewport.cpp
@@ -8,7 +8,7 @@
 #include "Dataflow/DataflowEditorMode.h"
 #include "Dataflow/DataflowEditorViewportClient.h"
 #include "EditorModeManager.h"
-#include "SViewportToolBar.h"
+#include "Dataflow/DataflowEditorViewportToolbar.h"
 
 #define LOCTEXT_NAMESPACE "SDataflowEditorViewport"
 
@@ -25,6 +25,20 @@ void SDataflowEditorViewport::Construct(const FArguments& InArgs, const FAssetEd
 	Client->VisibilityDelegate.BindSP(this, &SDataflowEditorViewport::IsVisible);
 }
 
+TSharedPtr<SWidget> SDataflowEditorViewport::MakeViewportToolbar()
+{
+	return SNew(SDataflowViewportSelectionToolBar, SharedThis(this));
+}
+
+void SDataflowEditorViewport::OnFocusViewportToSelection()
+{
+	if(UDataflowEditorMode* DataflowMode = GetEdMode())
+	{
+		const FBox SceneBoundingBox = DataflowMode->SceneBoundingBox();
+		Client->FocusViewportOnBox(SceneBoundingBox);
+	}
+}
+
 UDataflowEditorMode* SDataflowEditorViewport::GetEdMode() const
 {
 	if (const FEditorModeTools* const EditorModeTools = Client->GetModeTools())
@@ -42,7 +56,6 @@ void SDataflowEditorViewport::BindCommands()
 	SAssetEditorViewport::BindCommands();
 }
 
-
 bool SDataflowEditorViewport::IsVisible() const
 {
 	// Intentionally not calling SEditorViewport::IsVisible because it will return false if our simulation is more than 250ms.
diff --git a/Engine/Plugins/Experimental/Dataflow/Source/DataflowEditor/Private/Dataflow/DataflowEditorViewportClient.cpp b/Engine/Plugins/Experimental/Dataflow/Source/DataflowEditor/Private/Dataflow/DataflowEditorViewportClient.cpp
--- a/Engine/Plugins/Experimental/Dataflow/Source/DataflowEditor/Private/Dataflow/DataflowEditorViewportClient.cpp
+++ b/Engine/Plugins/Experimental/Dataflow/Source/DataflowEditor/Private/Dataflow/DataflowEditorViewportClient.cpp
@@ -1,23 +1,29 @@
 // Copyright Epic Games, Inc. All Rights Reserved.
 #include "Dataflow/DataflowEditorViewportClient.h"
-
-#include "AdvancedPreviewScene.h"
-#include "Components/DirectionalLightComponent.h"
 #include "Dataflow/DataflowActor.h"
 #include "Dataflow/DataflowEditorToolkit.h"
 #include "Dataflow/DataflowEngineSceneHitProxies.h"
 #include "Dataflow/DataflowXml.h"
+#include "Dataflow/DataflowEditor.h"
+#include "Dataflow/DataflowPreviewScene.h"
+#include "Dataflow/DataflowEditorMode.h"
+#include "Selection.h"
 #include "HAL/PlatformApplicationMisc.h"
 #include "PreviewScene.h"
+#include "EditorModeManager.h"
 
 FDataflowEditorViewportClient::FDataflowEditorViewportClient(FEditorModeTools* InModeTools,
-															 FPreviewScene* InPreviewScene,
-															 const TWeakPtr<SEditorViewport> InEditorViewportWidget)
+                                                             FPreviewScene* InPreviewScene,
+                                                             const TWeakPtr<SEditorViewport> InEditorViewportWidget)
 	: FEditorViewportClient(InModeTools, InPreviewScene, InEditorViewportWidget)
 {
-	// @todo(DataflowMode) : Move this to the Mode as a Component only (see UChaosClothAssetEditorMode)
-	DataflowActor = CastChecked<ADataflowActor>(PreviewScene->GetWorld()->SpawnActor(ADataflowActor::StaticClass()));
-	((FAdvancedPreviewScene*)InPreviewScene)->SetFloorVisibility(false, true);
+	// We want our near clip plane to be quite close so that we can zoom in further.
+	OverrideNearClipPlane(KINDA_SMALL_NUMBER);
+
+	EngineShowFlags.SetSelectionOutline(true);
+	EngineShowFlags.EnableAdvancedFeatures();
+
+	PreviewScene = static_cast<FDataflowPreviewScene*>(InPreviewScene);
 }
 
 void FDataflowEditorViewportClient::SetDataflowEditorToolkit(TWeakPtr<FDataflowEditorToolkit> InDataflowEditorToolkitPtr)
@@ -25,10 +31,9 @@ void FDataflowEditorViewportClient::SetDataflowEditorToolkit(TWeakPtr<FDataflowE
 	DataflowEditorToolkitPtr = InDataflowEditorToolkitPtr;
 }
 
-
 void FDataflowEditorViewportClient::SetSelectionMode(FDataflowSelectionState::EMode InState)
 {
-	FDataflowSelectionState State = DataflowActor->DataflowComponent->GetSelectionState();
+	FDataflowSelectionState State = PreviewScene->GetDataflowComponent()->GetSelectionState();
 
 	if (SelectionMode == InState)
 	{
@@ -40,13 +45,13 @@ void FDataflowEditorViewportClient::SetSelectionMode(FDataflowSelectionState::EM
 	}
 
 	State.Mode = SelectionMode;
-	DataflowActor->DataflowComponent->SetSelectionState(State);
+	PreviewScene->ModifyDataflowComponent()->SetSelectionState(State);
 
 	if (SelectionMode == FDataflowSelectionState::EMode::DSS_Dataflow_None)
 	{
-		if (!DataflowActor->DataflowComponent->GetSelectionState().IsEmpty())
+		if (!PreviewScene->GetDataflowComponent()->GetSelectionState().IsEmpty())
 		{
-			DataflowActor->DataflowComponent->SetSelectionState(FDataflowSelectionState(SelectionMode));
+			PreviewScene->ModifyDataflowComponent()->SetSelectionState(FDataflowSelectionState(SelectionMode));
 		}
 	}
 }
@@ -55,7 +60,8 @@ bool FDataflowEditorViewportClient::CanSetSelectionMode(FDataflowSelectionState:
 	TSharedPtr<FDataflowEditorToolkit> DataflowEditorToolkit = DataflowEditorToolkitPtr.Pin();
 	if (DataflowEditorToolkitPtr.IsValid())
 	{
-		if (const UDataflow* Dataflow = DataflowEditorToolkit->GetDataflow())
+		const FDataflowEditorDatas& DataflowEditorDatas = DataflowEditorToolkit->GetDataflowEditorDatas();
+		if (const UDataflow* Dataflow = DataflowEditorDatas.DataflowAsset)
 		{
 			if (Dataflow->GetRenderTargets().Num())
 			{
@@ -65,7 +71,7 @@ bool FDataflowEditorViewportClient::CanSetSelectionMode(FDataflowSelectionState:
 				}
 
 				if (InState == FDataflowSelectionState::EMode::DSS_Dataflow_Vertex
-					&& !DataflowActor->DataflowComponent->GetSelectionState().Nodes.IsEmpty())
+					&& !PreviewScene->GetDataflowComponent()->GetSelectionState().Nodes.IsEmpty())
 				{
 					return true;
 				}
@@ -96,11 +102,11 @@ bool FDataflowEditorViewportClient::InputKey(const FInputKeyEventArgs& EventArgs
 	FInputEventState InputState(EventArgs.Viewport, EventArgs.Key, EventArgs.Event);
 	if (InputState.IsCtrlButtonPressed() && EventArgs.Key == EKeys::C)
 	{
-		if (DataflowActor->DataflowComponent->GetSelectionState().Vertices.Num())
+		if (PreviewScene->GetDataflowComponent()->GetSelectionState().Vertices.Num())
 		{
 			FString XmlBuffer = FDataflowXmlWrite()
 				.Begin()
-				.MakeVertexSelectionBlock(DataflowActor->DataflowComponent->GetSelectionState().Vertices)
+				.MakeVertexSelectionBlock(PreviewScene->GetDataflowComponent()->GetSelectionState().Vertices)
 				.End()
 				.ToString();
 			FPlatformApplicationMisc::ClipboardCopy(*XmlBuffer);
@@ -119,12 +125,44 @@ bool FDataflowEditorViewportClient::InputKey(const FInputKeyEventArgs& EventArgs
 void FDataflowEditorViewportClient::ProcessClick(FSceneView& View, HHitProxy* HitProxy, FKey Key, EInputEvent Event, uint32 HitX, uint32 HitY)
 {
 	Super::ProcessClick(View, HitProxy, Key, Event, HitX, HitY);
-	if (DataflowActor && DataflowActor->DataflowComponent)
+	if (PreviewScene->GetDataflowComponent())
 	{
 		const bool bIsShiftKeyDown = Viewport->KeyState(EKeys::LeftShift) || Viewport->KeyState(EKeys::RightShift);
 		const bool bIsCtrltKeyDown = Viewport->KeyState(EKeys::LeftControl) || Viewport->KeyState(EKeys::RightControl);
 
-		FDataflowSelectionState SelectionState = DataflowActor->DataflowComponent->GetSelectionState();
+		USelection* SelectedComponents = ModeTools->GetSelectedComponents();
+
+		TArray<UPrimitiveComponent*> PreviouslySelectedComponents;
+		SelectedComponents->GetSelectedObjects<UPrimitiveComponent>(PreviouslySelectedComponents);
+
+		SelectedComponents->Modify();
+		SelectedComponents->BeginBatchSelectOperation();
+
+		SelectedComponents->DeselectAll();
+
+		if (HitProxy && HitProxy->IsA(HActor::StaticGetType()))
+		{
+			const HActor* const ActorProxy = static_cast<HActor*>(HitProxy);
+			if (ActorProxy && ActorProxy->Actor)
+			{
+				const AActor* const Actor = ActorProxy->Actor;
+			
+				Actor->ForEachComponent<UPrimitiveComponent>(true, [&](UPrimitiveComponent* Component)
+				{
+					SelectedComponents->Select(Component);
+					Component->PushSelectionToProxy();
+				});
+			}
+		}
+
+		SelectedComponents->EndBatchSelectOperation();
+
+		for (UPrimitiveComponent* const Component : PreviouslySelectedComponents)
+		{
+			Component->PushSelectionToProxy();
+		}
+
+		FDataflowSelectionState SelectionState = PreviewScene->GetDataflowComponent()->GetSelectionState();
 		FDataflowSelectionState PreState = SelectionState;
 
 		if (SelectionMode == FDataflowSelectionState::EMode::DSS_Dataflow_Object)
@@ -194,7 +232,7 @@ void FDataflowEditorViewportClient::ProcessClick(FSceneView& View, HHitProxy* Hi
 
 		if (PreState != SelectionState)
 		{
-			DataflowActor->DataflowComponent->SetSelectionState(SelectionState);
+			PreviewScene->ModifyDataflowComponent()->SetSelectionState(SelectionState);
 		}
 	}
 }
@@ -203,50 +241,28 @@ void FDataflowEditorViewportClient::Tick(float DeltaSeconds)
 {
 	Super::Tick(DeltaSeconds);
 
-	TSharedPtr<FDataflowEditorToolkit> DataflowEditorToolkit = DataflowEditorToolkitPtr.Pin();
+	const TSharedPtr<FDataflowEditorToolkit> DataflowEditorToolkit = DataflowEditorToolkitPtr.Pin();
 
-	if (DataflowActor && DataflowEditorToolkitPtr.IsValid())
+	if (PreviewScene->GetDataflowComponent())
 	{
-		if (TSharedPtr<Dataflow::FContext> Context = DataflowEditorToolkit->GetContext())
+		const FDataflowEditorDatas& DataflowEditorDatas = PreviewScene->GetDataflowDatas();
+		if (TSharedPtr<Dataflow::FContext> Context = DataflowEditorDatas.DataflowContext)
 		{
-			if (const UDataflow* Dataflow = DataflowEditorToolkit->GetDataflow())
+			if (const UDataflow* Dataflow = DataflowEditorDatas.DataflowAsset)
 			{
-				if (UDataflowComponent* DataflowComponent = DataflowActor->GetDataflowComponent())
+				const Dataflow::FTimestamp SystemTimestamp = LatestTimestamp(Dataflow, Context.Get());
+				if (SystemTimestamp >= LastModifiedTimestamp)
 				{
-					Dataflow::FTimestamp SystemTimestamp = LatestTimestamp(Dataflow, Context.Get());
-					if (SystemTimestamp >= LastModifiedTimestamp)
-					{
-						if (Dataflow->GetRenderTargets().Num())
-						{
-							// Component Object Rendering
-							DataflowComponent->ResetRenderTargets();
-							DataflowComponent->SetDataflow(Dataflow);
-							DataflowComponent->SetContext(Context);
-							for (const UDataflowEdNode* Node : Dataflow->GetRenderTargets())
-							{
-								DataflowComponent->AddRenderTarget(Node);
-							}
-						}
-						else
-						{
-							DataflowComponent->ResetRenderTargets();
-						}
-
-						LastModifiedTimestamp = LatestTimestamp(Dataflow, Context.Get()).Value + 1;
-					}
+					PreviewScene->UpdateDataflowComponent();
+					LastModifiedTimestamp = LatestTimestamp(Dataflow, Context.Get()).Value + 1;
 				}
 			}
 		}
 	}
-
-	// Tick the preview scene world.
-	PreviewScene->GetWorld()->Tick(IsRealtime() ? LEVELTICK_All : LEVELTICK_TimeOnly, DeltaSeconds);
 }
 
-
 void FDataflowEditorViewportClient::AddReferencedObjects(FReferenceCollector& Collector)
 {
 	Super::AddReferencedObjects(Collector);
-	Collector.AddReferencedObject(DataflowActor);
 }
 
diff --git a/Engine/Plugins/Experimental/Dataflow/Source/DataflowEditor/Private/Dataflow/DataflowEditorViewportToolbar.cpp b/Engine/Plugins/Experimental/Dataflow/Source/DataflowEditor/Private/Dataflow/DataflowEditorViewportToolbar.cpp
--- a/Engine/Plugins/Experimental/Dataflow/Source/DataflowEditor/Private/Dataflow/DataflowEditorViewportToolbar.cpp
+++ b/Engine/Plugins/Experimental/Dataflow/Source/DataflowEditor/Private/Dataflow/DataflowEditorViewportToolbar.cpp
@@ -12,9 +12,9 @@ bool bDataflowPrototypeSelectionMenu = false;
 FAutoConsoleVariableRef CVarDataflowPrototypeSelectionMenu(TEXT("p.Dataflow.prototype.selection"), bDataflowPrototypeSelectionMenu, TEXT("Work in progress development for selection controls in the 3D Viewport.[def:false]"));
 
 
-void SDataflowViewportSelectionToolBar::Construct(const FArguments& InArgs)
+void SDataflowViewportSelectionToolBar::Construct(const FArguments& InArgs, TSharedPtr<SDataflowEditorViewport> InDataflowViewport)
 {
-	EditorViewport = InArgs._EditorViewport;
+	EditorViewport = InDataflowViewport;
 	static const FName DefaultForegroundName("DefaultForeground");
 
 	this->ChildSlot
diff --git a/Engine/Plugins/Experimental/Dataflow/Source/DataflowEditor/Private/Dataflow/DataflowPreviewScene.cpp b/Engine/Plugins/Experimental/Dataflow/Source/DataflowEditor/Private/Dataflow/DataflowPreviewScene.cpp
new file mode 100644
--- /dev/null
+++ b/Engine/Plugins/Experimental/Dataflow/Source/DataflowEditor/Private/Dataflow/DataflowPreviewScene.cpp
@@ -0,0 +1,179 @@
+// Copyright Epic Games, Inc. All Rights Reserved.
+
+#include "Dataflow/DataflowPreviewScene.h"
+
+#include "Animation/AnimSingleNodeInstance.h"
+#include "AssetEditorModeManager.h"
+#include "Components/SkeletalMeshComponent.h"
+#include "Components/DynamicMeshComponent.h"
+#include "Dataflow/DataflowComponent.h"
+#include "Dataflow/DataflowEditor.h"
+#include "Elements/Framework/EngineElementsLibrary.h"
+
+#define LOCTEXT_NAMESPACE "FDataflowPreviewScene"
+
+FDataflowPreviewScene::FDataflowPreviewScene(FPreviewScene::ConstructionValues ConstructionValues, FDataflowEditorDatas& DataflowAssetDatas) :
+	FAdvancedPreviewScene(ConstructionValues), DataflowDatas(DataflowAssetDatas)
+{
+	SkeletalMeshActor = GetWorld()->SpawnActor<AActor>(AActor::StaticClass());
+
+	if(DataflowDatas.SkeletalMesh && DataflowDatas.bHasValidSkeletalMesh)
+	{
+		SkeletalMeshComponent = NewObject<USkeletalMeshComponent>(SkeletalMeshActor);
+		SkeletalMeshComponent->SelectionOverrideDelegate = UPrimitiveComponent::FSelectionOverride::CreateRaw(this, &FDataflowPreviewScene::IsComponentSelected);
+		SkeletalMeshComponent->SetDisablePostProcessBlueprint(true);
+		UpdateSkeletalMeshComponent();
+	}
+	
+	DataflowActor = GetWorld()->SpawnActor<AActor>(AActor::StaticClass());
+
+	if(DataflowDatas.DataflowAsset)
+	{
+		DataflowComponent = NewObject<UDataflowComponent>(DataflowActor);
+		//DataflowComponent->SelectionOverrideDelegate = UPrimitiveComponent::FSelectionOverride::CreateRaw(this, &FDataflowPreviewScene::IsComponentSelected);
+		
+		// @todo(DynamicMeshRendering) : Enable Dynamic Mesh Rendering for dataflow terminals. Hide the dataflow
+		//DataflowComponent->SetVisibility(false);
+		UpdateDataflowComponent();
+	}
+	
+	DynamicMeshActor = GetWorld()->SpawnActor<AActor>(AActor::StaticClass());
+
+	
+	SkeletalMeshActor->RegisterAllComponents();
+	DataflowActor->RegisterAllComponents();
+	DynamicMeshActor->RegisterAllComponents();
+
+	SetFloorVisibility(false, true);
+}
+
+FDataflowPreviewScene::~FDataflowPreviewScene()
+{
+	if (SkeletalMeshComponent)
+	{
+		SkeletalMeshComponent->TransformUpdated.RemoveAll(this);
+		SkeletalMeshComponent->SelectionOverrideDelegate.Unbind();
+		SkeletalMeshComponent->UnregisterComponent();
+		SkeletalMeshComponent->DestroyComponent();
+	}
+
+	if (DataflowComponent)
+	{
+		DataflowComponent->SelectionOverrideDelegate.Unbind();
+		DataflowComponent->UnregisterComponent();
+		DataflowComponent->DestroyComponent();
+	}
+
+	ResetDynamicMeshComponents();
+}
+
+void FDataflowPreviewScene::AddReferencedObjects(FReferenceCollector& Collector)
+{
+	FAdvancedPreviewScene::AddReferencedObjects(Collector);
+
+	Collector.AddReferencedObject(DataflowComponent);
+	Collector.AddReferencedObject(SkeletalMeshComponent);
+	Collector.AddReferencedObject(SkeletalMeshActor);
+	Collector.AddReferencedObject(DataflowActor);
+	Collector.AddReferencedObject(DynamicMeshActor);
+	Collector.AddReferencedObject(PreviewAnimInstance);
+	Collector.AddReferencedObjects(DynamicMeshComponents);
+}
+
+void FDataflowPreviewScene::UpdateDataflowComponent()
+{
+	DataflowComponent->ResetRenderTargets();
+	
+	DataflowComponent->SetDataflow(DataflowDatas.DataflowAsset);
+	DataflowComponent->SetContext(DataflowDatas.DataflowContext);
+	
+	for (const UDataflowEdNode* const Node : DataflowDatas.DataflowAsset->GetRenderTargets())
+	{
+		DataflowComponent->AddRenderTarget(Node);
+	}
+	DataflowComponent->UpdateBounds();
+}
+
+void FDataflowPreviewScene::UpdateSkeletalMeshComponent()
+{
+	SkeletalMeshComponent->SetSkeletalMeshAsset(DataflowDatas.SkeletalMesh);
+
+	if (DataflowDatas.AnimationAsset)
+	{
+		PreviewAnimInstance = NewObject<UAnimSingleNodeInstance>(SkeletalMeshComponent);
+		PreviewAnimInstance->SetAnimationAsset(DataflowDatas.AnimationAsset);
+
+		SkeletalMeshComponent->SetAnimationMode(EAnimationMode::AnimationSingleNode);
+		SkeletalMeshComponent->InitAnim(true);
+		SkeletalMeshComponent->AnimationData.PopulateFrom(PreviewAnimInstance);
+		SkeletalMeshComponent->AnimScriptInstance = PreviewAnimInstance;
+		SkeletalMeshComponent->AnimScriptInstance->InitializeAnimation();
+		SkeletalMeshComponent->ValidateAnimation();
+	}
+	else
+	{
+		SkeletalMeshComponent->Stop();
+		SkeletalMeshComponent->AnimationData = FSingleAnimationPlayData();
+		SkeletalMeshComponent->AnimScriptInstance = nullptr;
+	}
+	SkeletalMeshComponent->UpdateBounds();
+}
+
+void FDataflowPreviewScene::ResetDynamicMeshComponents()
+{
+	for(const TObjectPtr<UDynamicMeshComponent>& DynamicMeshComponent : DynamicMeshComponents)
+	{
+		DynamicMeshComponent->SelectionOverrideDelegate.Unbind();
+		DynamicMeshComponent->UnregisterComponent();
+		DynamicMeshComponent->DestroyComponent();
+	}
+	DynamicMeshComponents.Reset();
+}
+
+TObjectPtr<UDynamicMeshComponent>& FDataflowPreviewScene::AddDynamicMeshComponent(UE::Geometry::FDynamicMesh3&& DynamicMesh, const TArray<UMaterialInterface*>& MaterialSet)
+{
+	TObjectPtr<UDynamicMeshComponent> DynamicMeshComponent = NewObject<UDynamicMeshComponent>(DynamicMeshActor);
+		
+	DynamicMeshComponent->SetMesh(MoveTemp(DynamicMesh));
+	DynamicMeshComponent->ConfigureMaterialSet(MaterialSet);
+	check(DynamicMeshComponent->ValidateMaterialSlots(false, false));
+	DynamicMeshComponent->SelectionOverrideDelegate = UPrimitiveComponent::FSelectionOverride::CreateRaw(this, &FDataflowPreviewScene::IsComponentSelected);
+	
+	DynamicMeshComponent->RegisterComponentWithWorld(GetWorld());
+	DynamicMeshComponent->UpdateBounds();
+		
+	const int32 ElementIndex = DynamicMeshComponents.Emplace(DynamicMeshComponent);
+	return DynamicMeshComponents[ElementIndex];
+}
+
+FBox FDataflowPreviewScene::GetBoundingBox() const
+{
+	FBox SceneBounds(ForceInitToZero);
+	for (const TObjectPtr<UDynamicMeshComponent>& MeshComponent : DynamicMeshComponents)
+	{
+		if (MeshComponent)
+		{
+			SceneBounds += MeshComponent->Bounds.GetBox();
+		}
+	}
+	return SceneBounds;
+}
+
+bool FDataflowPreviewScene::IsComponentSelected(const UPrimitiveComponent* InComponent) const
+{
+	if(DataflowModeManager.IsValid())
+	{
+		if (const UTypedElementSelectionSet* const TypedElementSelectionSet = DataflowModeManager->GetEditorSelectionSet())
+		{
+			if (const FTypedElementHandle ComponentElement = UEngineElementsLibrary::AcquireEditorComponentElementHandle(InComponent))
+			{
+				const bool bElementSelected = TypedElementSelectionSet->IsElementSelected(ComponentElement, FTypedElementIsSelectedOptions());
+				return bElementSelected;
+			}
+		}
+	}
+	return false;
+}
+
+#undef LOCTEXT_NAMESPACE
+
diff --git a/Engine/Plugins/Experimental/Dataflow/Source/DataflowEditor/Private/Dataflow/DataflowToolTarget.cpp b/Engine/Plugins/Experimental/Dataflow/Source/DataflowEditor/Private/Dataflow/DataflowToolTarget.cpp
new file mode 100644
--- /dev/null
+++ b/Engine/Plugins/Experimental/Dataflow/Source/DataflowEditor/Private/Dataflow/DataflowToolTarget.cpp
@@ -0,0 +1,288 @@
+// Copyright Epic Games, Inc. All Rights Reserved.
+
+#include "Dataflow/DataflowToolTarget.h"
+
+#include "ConversionUtils/DynamicMeshViaMeshDescriptionUtil.h"
+#include "Dataflow/DataflowEditorToolkit.h"
+#include "Dataflow/DataflowEdNode.h"
+#include "Dataflow/DataflowObject.h"
+#include "Dataflow/DataflowObjectInterface.h"
+#include "Dataflow/DataflowNodeParameters.h"
+#include "DynamicMeshToMeshDescription.h"
+#include "GeometryCollection/Facades/CollectionRenderingFacade.h"
+#include "Materials/Material.h"
+#include "MeshDescriptionToDynamicMesh.h"
+#include "MaterialDomain.h"
+#include "StaticMeshAttributes.h"
+#include "UObject/Package.h"
+
+#include UE_INLINE_GENERATED_CPP_BY_NAME(DataflowToolTarget)
+
+using namespace UE::Geometry;
+
+namespace UE
+{
+	namespace Conversion
+	{
+		// Convert a rendering facade to a dynamic mesh
+		void RenderingFacadeToDynamicMesh(const GeometryCollection::Facades::FRenderingFacade& Facade, FDynamicMesh3& DynamicMesh)
+		{
+			if(Facade.CanRenderSurface())
+			{
+				const int32 NumTriangles = Facade.NumTriangles();
+				const int32 NumVertices = Facade.NumVertices();
+				
+				const TManagedArray<FIntVector>& Indices = Facade.GetIndices();
+				const TManagedArray<FVector3f>& Positions = Facade.GetVertices();
+				const TManagedArray<FVector3f>& Normals = Facade.GetNormals();
+				const TManagedArray<FLinearColor>& Colors = Facade.GetVertexColor();
+
+				DynamicMesh.Clear();
+				for(int32 VertexIndex = 0; VertexIndex < NumVertices; ++VertexIndex)
+				{
+					DynamicMesh.AppendVertex(FVertexInfo(FVector3d(Positions[VertexIndex]), Normals[VertexIndex],
+						FVector3f(Colors[VertexIndex].R, Colors[VertexIndex].G, Colors[VertexIndex].B)));
+				}
+				for(int32 TriangleIndex = 0; TriangleIndex < NumTriangles; ++TriangleIndex)
+				{
+					DynamicMesh.AppendTriangle(FIndex3i(Indices[TriangleIndex].X, Indices[TriangleIndex].Y, Indices[TriangleIndex].Z));
+				}
+			}
+		}
+		
+		// Convert a dataflow component to a dynamic mesh
+		void DataflowToDynamicMesh(UObject* Asset, UDataflow* Dataflow, FDynamicMesh3& DynamicMesh)
+		{
+			FManagedArrayCollection RenderCollection;
+			GeometryCollection::Facades::FRenderingFacade Facade(RenderCollection);
+			Facade.DefineSchema();
+			
+			const TSharedPtr<::Dataflow::FEngineContext> DataflowContext = 
+				MakeShared<::Dataflow::FEngineContext>(	Asset, Dataflow, ::Dataflow::FTimestamp::Invalid);
+			
+			for (const UDataflowEdNode* Target : Dataflow->GetRenderTargets())
+			{
+				if (Target)
+				{
+					Target->Render(Facade, DataflowContext);
+				}
+			}
+			RenderingFacadeToDynamicMesh(Facade, DynamicMesh);
+		}
+
+		// Convert a dynamic mesh to a rendering facade
+		void DynamicMeshToRenderingFacade(const FDynamicMesh3& DynamicMesh, GeometryCollection::Facades::FRenderingFacade& Facade)
+		{
+			if(Facade.CanRenderSurface())
+			{
+				const int32 NumTriangles = Facade.NumTriangles();
+				const int32 NumVertices = Facade.NumVertices();
+
+				// We can only override vertices attributes (position, normals, colors)
+				if((NumTriangles == DynamicMesh.TriangleCount()) && (NumVertices == DynamicMesh.VertexCount()))
+				{
+					TManagedArray<FVector3f>& Positions = Facade.ModifyVertices();
+					TManagedArray<FVector3f>& Normals = Facade.ModifyNormals();
+					TManagedArray<FLinearColor>& Colors = Facade.ModifyVertexColor();
+
+					for(int32 VertexIndex = 0; VertexIndex < NumVertices; ++VertexIndex)
+					{
+						Positions[VertexIndex] = FVector3f(DynamicMesh.GetVertex(VertexIndex));
+						Normals[VertexIndex] = DynamicMesh.GetVertexNormal(VertexIndex);
+						Colors[VertexIndex] = DynamicMesh.GetVertexColor(VertexIndex);
+					}
+				}
+			}
+		}
+		
+		// Convert a dynamic mesh to a dataflow component
+		void DynamicMeshToDataflow(const FDynamicMesh3& DynamicMesh, UDataflow* Dataflow)
+		{
+			//todo
+		}
+	}
+}
+
+//
+// UDataflowReadOnlyToolTarget
+//
+
+bool UDataflowReadOnlyToolTarget::IsValid() const
+{
+	return Dataflow != nullptr && Asset != nullptr;
+}
+
+int32 UDataflowReadOnlyToolTarget::GetNumMaterials() const
+{
+	return (ensure(IsValid()) && Dataflow->Material) ? 1 : 0;
+}
+
+UMaterialInterface* UDataflowReadOnlyToolTarget::GetMaterial(int32 MaterialIndex) const
+{
+	return (ensure(IsValid()) &&  MaterialIndex == 0)? (UMaterialInterface*)Dataflow->Material : nullptr;
+}
+
+void UDataflowReadOnlyToolTarget::GetMaterialSet(FComponentMaterialSet& MaterialSetOut, bool bPreferAssetMaterials) const
+{
+	if (!ensure(IsValid())) return;
+	if (Dataflow->Material)
+	{
+		MaterialSetOut.Materials.SetNum(1);
+		MaterialSetOut.Materials[0] = (UMaterialInterface*)Dataflow->Material;
+	}
+}
+
+bool UDataflowReadOnlyToolTarget::CommitMaterialSetUpdate(const FComponentMaterialSet& MaterialSet, bool bApplyToAsset)
+{
+	// filter out any Engine materials that we don't want to be permanently assigning
+	TArray<UMaterialInterface*> FilteredMaterials = MaterialSet.Materials;
+	for (int32 k = 0; k < FilteredMaterials.Num(); ++k)
+	{
+		FString AssetPath = FilteredMaterials[k]->GetPathName();
+		if (AssetPath.StartsWith(TEXT("/MeshModelingToolsetExp/")))
+		{
+			FilteredMaterials[k] = UMaterial::GetDefaultMaterial(MD_Surface);
+		}
+	}
+
+	if(IsValid() && FilteredMaterials.Num() == 0)
+	{
+		if(UMaterial* Material = Cast<UMaterial>(FilteredMaterials[0]))
+		Dataflow->Material = Material;
+		return true;
+	}
+	return false;
+}
+
+const FMeshDescription* UDataflowReadOnlyToolTarget::GetMeshDescription(const FGetMeshParameters& GetMeshParams)
+{
+	if (ensure(IsValid()) && !DataflowMeshDescription.IsValid())
+	{
+		DataflowMeshDescription = MakeUnique<FMeshDescription>();
+		FStaticMeshAttributes Attributes(*DataflowMeshDescription);
+		Attributes.Register();
+
+		const FDynamicMesh3 DynamicMesh = GetDynamicMesh();
+		
+		FDynamicMeshToMeshDescription Converter;
+		Converter.Convert(&DynamicMesh, *DataflowMeshDescription, true);
+	}
+	return DataflowMeshDescription.Get();
+}
+
+FMeshDescription UDataflowReadOnlyToolTarget::GetEmptyMeshDescription()
+{
+	FMeshDescription EmptyMeshDescription;
+	FStaticMeshAttributes MeshAttributes(EmptyMeshDescription);
+	MeshAttributes.Register();
+
+	return EmptyMeshDescription;
+}
+
+FDynamicMesh3 UDataflowReadOnlyToolTarget::GetDynamicMesh()
+{
+	FDynamicMesh3 DynamicMesh;
+	if(IsValid())
+	{
+		UE::Conversion::DataflowToDynamicMesh(Asset,Dataflow, DynamicMesh);
+	}
+	return DynamicMesh;
+}
+
+FDynamicMesh3 UDataflowReadOnlyToolTarget::GetDynamicMesh(bool bRequestTangents)
+{
+	return GetDynamicMesh();
+}
+
+//
+// UDataflowToolTarget
+//
+
+void UDataflowToolTarget::CommitMeshDescription(const FCommitter& Committer, const FCommitMeshParameters& CommitMeshParams)
+{
+	if (ensure(IsValid()) == false) return;
+
+	// Let the user fill our mesh description with the Committer
+	if (!DataflowMeshDescription.IsValid())
+	{
+		DataflowMeshDescription = MakeUnique<FMeshDescription>();
+		FStaticMeshAttributes Attributes(*DataflowMeshDescription);
+		Attributes.Register();
+	}
+	FCommitterParams CommitParams;
+	CommitParams.MeshDescriptionOut = DataflowMeshDescription.Get();
+	Committer(CommitParams);
+
+	// The conversion we have right now is from dynamic mesh to volume, so we convert
+	// to dynamic mesh first.
+	FDynamicMesh3 DynamicMesh;
+	FMeshDescriptionToDynamicMesh Converter;
+	Converter.Convert(DataflowMeshDescription.Get(), DynamicMesh);
+
+	CommitDynamicMesh(DynamicMesh);
+}
+
+void UDataflowToolTarget::CommitDynamicMesh(const FDynamicMesh3& DynamicMesh, const FDynamicMeshCommitInfo& CommitInfo)
+{
+	if(IsValid())
+	{
+		UE::Conversion::DynamicMeshToDataflow(DynamicMesh, Dataflow);
+	}
+}
+
+// Factory
+
+bool UDataflowReadOnlyToolTargetFactory::CanBuildTarget(UObject* SourceObject, const FToolTargetTypeRequirements& Requirements) const
+{
+	// We are using an exact cast here to prevent subclasses, which might not meet all
+	// requirements for functionality such as the deprecated DestructibleMeshComponent, from 
+	// being caught up as valid targets.
+	// If you want to make the tool target work with some subclass of UDataflow,
+	// just add another factory that allows that class specifically(but make sure that
+	// GetMeshDescription and such work properly)
+	const UDataflow* Dataflow = FDataflowEditorToolkit::GetDataflowAsset(SourceObject);
+
+	return Dataflow &&
+		ExactCast<UDataflow>(Cast<UDataflow>(Dataflow)) &&
+		!ExactCast<UDataflow>(Cast<UDataflow>(Dataflow))->GetOutermost()->bIsCookedForEditor &&
+		Requirements.AreSatisfiedBy(UDataflowReadOnlyToolTarget::StaticClass());
+}
+
+UToolTarget* UDataflowReadOnlyToolTargetFactory::BuildTarget(UObject* SourceObject, const FToolTargetTypeRequirements& Requirements)
+{
+	UDataflowReadOnlyToolTarget* Target = NewObject<UDataflowReadOnlyToolTarget>();
+	Target->Asset = SourceObject;
+	Target->Dataflow = FDataflowEditorToolkit::GetDataflowAsset(SourceObject);
+
+	// @todo(brice) : I needed to comment this out?
+	//checkSlow(Target->Component.IsValid() && Requirements.AreSatisfiedBy(Target));
+
+	return Target;
+}
+
+bool UDataflowToolTargetFactory::CanBuildTarget(UObject* SourceObject, const FToolTargetTypeRequirements& Requirements) const
+{
+	// We are using an exact cast here to prevent subclasses, which might not meet all
+	// requirements for functionality such as the deprecated DestructibleMeshComponent, from 
+	// being caught up as valid targets.
+	// If you want to make the tool target work with some subclass of UDataflow,
+	// just add another factory that allows that class specifically(but make sure that
+	// GetMeshDescription and such work properly)
+	const UDataflow* Dataflow = FDataflowEditorToolkit::GetDataflowAsset(SourceObject);
+
+	return Dataflow &&
+		ExactCast<UDataflow>(Cast<UDataflow>(Dataflow)) &&
+		!ExactCast<UDataflow>(Cast<UDataflow>(Dataflow))->GetOutermost()->bIsCookedForEditor &&
+		Requirements.AreSatisfiedBy(UDataflowToolTarget::StaticClass());
+}
+
+UToolTarget* UDataflowToolTargetFactory::BuildTarget(UObject* SourceObject, const FToolTargetTypeRequirements& Requirements)
+{
+	UDataflowToolTarget* Target = NewObject<UDataflowToolTarget>();
+	Target->Asset = SourceObject;
+	Target->Dataflow = FDataflowEditorToolkit::GetDataflowAsset(SourceObject);
+	//checkSlow(Target->Component.IsValid() && Requirements.AreSatisfiedBy(Target));
+
+	return Target;
+}
+
diff --git a/Engine/Plugins/Experimental/Dataflow/Source/DataflowEditor/Public/Dataflow/DataflowComponentToolTarget.h b/Engine/Plugins/Experimental/Dataflow/Source/DataflowEditor/Public/Dataflow/DataflowComponentToolTarget.h
new file mode 100644
--- /dev/null
+++ b/Engine/Plugins/Experimental/Dataflow/Source/DataflowEditor/Public/Dataflow/DataflowComponentToolTarget.h
@@ -0,0 +1,109 @@
+﻿// Copyright Epic Games, Inc. All Rights Reserved.
+
+#pragma once
+
+#include "CoreMinimal.h"
+
+#include "TargetInterfaces/DynamicMeshCommitter.h"
+#include "TargetInterfaces/DynamicMeshProvider.h"
+#include "TargetInterfaces/MaterialProvider.h"
+#include "TargetInterfaces/MeshDescriptionCommitter.h"
+#include "TargetInterfaces/MeshDescriptionProvider.h"
+#include "ToolTargets/PrimitiveComponentToolTarget.h"
+
+#include "DataflowComponentToolTarget.generated.h"
+
+/**
+ * A tool target backed by a read-only dataflow component that can provide and take a mesh
+ * description.
+ */
+UCLASS(Transient)
+class DATAFLOWEDITOR_API UDataflowComponentReadOnlyToolTarget :
+	public UPrimitiveComponentToolTarget,
+	public IMeshDescriptionProvider,
+	public IDynamicMeshProvider, 
+	public IMaterialProvider
+{
+	GENERATED_BODY()
+
+public:
+
+	// UToolTarget implementation
+	virtual bool IsValid() const override;
+
+	// IMeshDescriptionProvider implementation
+	virtual const FMeshDescription* GetMeshDescription(const FGetMeshParameters& GetMeshParams = FGetMeshParameters()) override;
+	virtual FMeshDescription GetEmptyMeshDescription() override;
+
+	// IMaterialProvider implementation
+	virtual int32 GetNumMaterials() const override;
+	virtual UMaterialInterface* GetMaterial(int32 MaterialIndex) const override;
+	virtual void GetMaterialSet(FComponentMaterialSet& MaterialSetOut, bool bPreferAssetMaterials) const override;
+	virtual bool CommitMaterialSetUpdate(const FComponentMaterialSet& MaterialSet, bool bApplyToAsset) override;
+
+	// IDynamicMeshProvider
+	virtual UE::Geometry::FDynamicMesh3 GetDynamicMesh() override;
+	virtual UE::Geometry::FDynamicMesh3 GetDynamicMesh(bool bRequestTangents) override;
+
+protected:
+	// So that the tool target factory can poke into Component.
+	friend class UDataflowComponentReadOnlyToolTargetFactory;
+
+	// Until UDataflow stores its internal representation as FMeshDescription, we need to
+	// retain the storage here to cover the lifetime of the pointer returned by GetMeshDescription(). 
+	TUniquePtr<FMeshDescription> DataflowMeshDescription;
+};
+
+/**
+ * A tool target backed by a dataflow component that can provide and take a mesh
+ * description.
+ */
+UCLASS(Transient)
+class DATAFLOWEDITOR_API UDataflowComponentToolTarget :
+	public UDataflowComponentReadOnlyToolTarget,
+	public IMeshDescriptionCommitter,
+	public IDynamicMeshCommitter
+{
+	GENERATED_BODY()
+
+public:
+	// IMeshDescriptionCommitter implementation
+	virtual void CommitMeshDescription(const FCommitter& Committer, const FCommitMeshParameters& CommitParams = FCommitMeshParameters()) override;
+	using IMeshDescriptionCommitter::CommitMeshDescription; // unhide the other overload
+
+	// IDynamicMeshCommitter
+	virtual void CommitDynamicMesh(const UE::Geometry::FDynamicMesh3& Mesh, const FDynamicMeshCommitInfo& CommitInfo) override;
+	using IDynamicMeshCommitter::CommitDynamicMesh; // unhide the other overload
+
+protected:
+	// So that the tool target factory can poke into Component.
+	friend class UDataflowComponentToolTargetFactory;
+};
+
+
+/** Factory for UDataflowComponentReadOnlyToolTarget to be used by the target manager. */
+UCLASS(Transient)
+class DATAFLOWEDITOR_API UDataflowComponentReadOnlyToolTargetFactory : public UToolTargetFactory
+{
+	GENERATED_BODY()
+
+public:
+
+	// UToolTargetFactory implementation
+	virtual bool CanBuildTarget(UObject* SourceObject, const FToolTargetTypeRequirements& TargetTypeInfo) const override;
+	virtual UToolTarget* BuildTarget(UObject* SourceObject, const FToolTargetTypeRequirements& TargetTypeInfo) override;
+};
+
+
+/** Factory for UDataflowComponentToolTarget to be used by the target manager. */
+UCLASS(Transient)
+class DATAFLOWEDITOR_API UDataflowComponentToolTargetFactory : public UToolTargetFactory
+{
+	GENERATED_BODY()
+
+public:
+
+	// UToolTargetFactory implementation
+	virtual bool CanBuildTarget(UObject* SourceObject, const FToolTargetTypeRequirements& TargetTypeInfo) const override;
+	virtual UToolTarget* BuildTarget(UObject* SourceObject, const FToolTargetTypeRequirements& TargetTypeInfo) override;
+};
diff --git a/Engine/Plugins/Experimental/Dataflow/Source/DataflowEditor/Public/Dataflow/DataflowEditor.h b/Engine/Plugins/Experimental/Dataflow/Source/DataflowEditor/Public/Dataflow/DataflowEditor.h
--- a/Engine/Plugins/Experimental/Dataflow/Source/DataflowEditor/Public/Dataflow/DataflowEditor.h
+++ b/Engine/Plugins/Experimental/Dataflow/Source/DataflowEditor/Public/Dataflow/DataflowEditor.h
@@ -17,6 +17,46 @@
 #include "DataflowEditor.generated.h"
 
 class FDataflowEditorToolkit;
+class USkeletalMesh;
+class USkeleton;
+class UAnimationAsset;
+
+/** 
+ * Dataflow datas that will be used within the editor classes to evaluate the graph
+ */
+struct DATAFLOWEDITOR_API FDataflowEditorDatas
+{
+	/** Check if the datas flow datas are valid */
+	bool IsValid() const { return DataflowOwner && DataflowAsset;}
+	
+	/** Data flow object owner */
+	TObjectPtr<UObject> DataflowOwner = nullptr;
+
+	/** Data flow asset that we will edit */
+	TObjectPtr<UDataflow> DataflowAsset = nullptr;
+
+	/** Data flow terminal path for evaluation */
+	FString DataflowTerminal = "";
+
+	/** Data flow skeletal mesh*/
+	TObjectPtr<USkeletalMesh> SkeletalMesh = nullptr;
+
+	/** Data flow skeleton*/
+	TObjectPtr<USkeleton> Skeleton = nullptr;
+
+	/** Animation asset to be used to preview simulation */
+	TObjectPtr<UAnimationAsset> AnimationAsset;
+
+	/**  Engine context to be used for dataflow evaluation */
+	TSharedPtr<Dataflow::FEngineContext> DataflowContext = nullptr;
+
+	/** Last data flow evaluated node time stamp */
+	Dataflow::FTimestamp LastNodeTimestamp = Dataflow::FTimestamp::Invalid;
+
+	/** Boolean to check if the skelmesh is valid*/
+	bool bHasValidSkeletalMesh = false;
+};
+
 /** 
  * The actual asset editor class doesn't have that much in it, intentionally. 
  * 
@@ -37,7 +77,19 @@ class DATAFLOWEDITOR_API UDataflowEditor : public UBaseCharacterFXEditor
 
 public:
 
+	// UBaseCharacterFXEditor interface
 	virtual TSharedPtr<FBaseAssetToolkit> CreateToolkit() override;
+	virtual void Initialize(const TArray<TObjectPtr<UObject>>& InObjects) override;
+
+private :
+
+	friend class FDataflowEditorToolkit;
+	
+	// Dataflow editor is the owner of the object list to edit/process and the dataflow mode
+	// is the one holding the dynamic mesh components to be rendered in the viewport
+	// It is why the data flow asset/owner/skelmesh have been added here. Could be added
+	// in the subsystem if necessary
+	FDataflowEditorDatas DataflowDatas;
 };
 
 DECLARE_LOG_CATEGORY_EXTERN(LogDataflowEditor, Log, All);
\ No newline at end of file
diff --git a/Engine/Plugins/Experimental/Dataflow/Source/DataflowEditor/Public/Dataflow/DataflowEditorCommands.h b/Engine/Plugins/Experimental/Dataflow/Source/DataflowEditor/Public/Dataflow/DataflowEditorCommands.h
--- a/Engine/Plugins/Experimental/Dataflow/Source/DataflowEditor/Public/Dataflow/DataflowEditorCommands.h
+++ b/Engine/Plugins/Experimental/Dataflow/Source/DataflowEditor/Public/Dataflow/DataflowEditorCommands.h
@@ -46,6 +46,13 @@ public:
 	TSharedPtr< FUICommandInfo > ZoomToFitGraph;
 
 	TMap< FName, TSharedPtr<FUICommandInfo> > CreateNodesMap;
+
+	// Temporary, just to test the editor, will be replace by the selection tool
+	const static FString BeginAttributeEditorToolIdentifier;
+	TSharedPtr<FUICommandInfo> BeginAttributeEditorTool;
+
+	const static FString BeginMeshSelectionToolIdentifier;
+	TSharedPtr<FUICommandInfo> BeginMeshSelectionTool;
 };
 
 class DATAFLOWEDITOR_API FDataflowEditorCommands
diff --git a/Engine/Plugins/Experimental/Dataflow/Source/DataflowEditor/Public/Dataflow/DataflowEditorMode.h b/Engine/Plugins/Experimental/Dataflow/Source/DataflowEditor/Public/Dataflow/DataflowEditorMode.h
--- a/Engine/Plugins/Experimental/Dataflow/Source/DataflowEditor/Public/Dataflow/DataflowEditorMode.h
+++ b/Engine/Plugins/Experimental/Dataflow/Source/DataflowEditor/Public/Dataflow/DataflowEditorMode.h
@@ -3,9 +3,11 @@
 #pragma once
 
 #include "BaseCharacterFXEditorMode.h"
-
 #include "DataflowEditorMode.generated.h"
 
+class UDataflowComponent;
+class FDataflowPreviewScene;
+class SDataflowGraphEditor;
 
 /**
  * The dataflow editor mode is the mode used in the cloth asset editor. It holds most of the inter-tool state.
@@ -22,15 +24,44 @@ public:
 	const static FEditorModeID EM_DataflowEditorModeId;
 
 	UDataflowEditorMode();
+	
+	/**
+	* Gets the tool target requirements for the mode.
+	*/
+	static const FToolTargetTypeRequirements& GetToolTargetRequirements();
+
+	/**
+	* Gets the scene bounding box
+	*/
+	virtual FBox SceneBoundingBox() const override;
+
+	/** Set the data flow preview scene to store the dynamic meshes generated by the tools */
+	void SetDataflowPreviewScene(FDataflowPreviewScene* DataflowPreviewScene) { PreviewScene = DataflowPreviewScene;}
+
+	/** Set the data flow graph editor to create nodes once the tools have ended*/
+	void SetDataflowGraphEditor(const TSharedPtr<SDataflowGraphEditor>& DataflowGraphEditor) { GraphEditor = DataflowGraphEditor;}
 
 private:
 	// UEdMode
 	virtual void CreateToolkit() override;
+	virtual void BindCommands() override;
+	virtual void Exit() override;
+	virtual void Enter() override;
+	virtual void ModeTick(float DeltaTime) override;
 
 	// UBaseCharacterFXEditorMode
-	virtual void AddToolTargetFactories() override {}
-	virtual void RegisterTools() override {}
-	virtual void CreateToolTargets(const TArray<TObjectPtr<UObject>>& AssetsIn) override {}
+	virtual void AddToolTargetFactories() override;
+	virtual void RegisterTools() override;
+	virtual void CreateToolTargets(const TArray<TObjectPtr<UObject>>& AssetsIn) override;
+	virtual void InitializeTargets(const TArray<TObjectPtr<UObject>>& AssetsIn) override;
+
+	/** Dataflow preview scene from the toolkit */
+	FDataflowPreviewScene* PreviewScene = nullptr;
+
+	/** Correspondence between node types and commands to launch tools */
+	TMap<FName, TSharedPtr<const FUICommandInfo>> NodeToolMap;
 
+	/** Graph editor is required to add nodes once some tools are complete */
+	TWeakPtr<SDataflowGraphEditor> GraphEditor;
 };
 
diff --git a/Engine/Plugins/Experimental/Dataflow/Source/DataflowEditor/Public/Dataflow/DataflowEditorModeToolkit.h b/Engine/Plugins/Experimental/Dataflow/Source/DataflowEditor/Public/Dataflow/DataflowEditorModeToolkit.h
--- a/Engine/Plugins/Experimental/Dataflow/Source/DataflowEditor/Public/Dataflow/DataflowEditorModeToolkit.h
+++ b/Engine/Plugins/Experimental/Dataflow/Source/DataflowEditor/Public/Dataflow/DataflowEditorModeToolkit.h
@@ -15,13 +15,14 @@ class DATAFLOWEDITOR_API FDataflowEditorModeToolkit : public FBaseCharacterFXEdi
 {
 public:
 
+	// FModeToolkit interface
 	virtual void OnToolStarted(UInteractiveToolManager* Manager, UInteractiveTool* Tool) override;
+	virtual void BuildToolPalette(FName PaletteName, class FToolBarBuilder& ToolbarBuilder) override;
 
-
-	// IToolkit
+	// IToolkit interface
 	virtual FName GetToolkitFName() const override;
 	virtual FText GetBaseToolkitName() const override;
-	virtual const FSlateBrush* GetActiveToolIcon(const FString& Identifier) const override;
-
 
+	// FBaseCharacterFXEditorModeToolkit interface
+	virtual const FSlateBrush* GetActiveToolIcon(const FString& Identifier) const override;
 };
diff --git a/Engine/Plugins/Experimental/Dataflow/Source/DataflowEditor/Public/Dataflow/DataflowEditorToolkit.h b/Engine/Plugins/Experimental/Dataflow/Source/DataflowEditor/Public/Dataflow/DataflowEditorToolkit.h
--- a/Engine/Plugins/Experimental/Dataflow/Source/DataflowEditor/Public/Dataflow/DataflowEditorToolkit.h
+++ b/Engine/Plugins/Experimental/Dataflow/Source/DataflowEditor/Public/Dataflow/DataflowEditorToolkit.h
@@ -22,13 +22,14 @@ class IToolkitHost;
 class UDataflow;
 class USkeletalMesh;
 class SDataflowGraphEditor;
+struct FDataflowEditorDatas;
 
 namespace Dataflow
 {
 	class DATAFLOWEDITOR_API FAssetContext : public TEngineContext<FContextSingle>
 	{
 	public:
-		DATAFLOW_CONTEXT_INTERNAL(TEngineContext<FContextSingle>, FAssetContext);
+		DATAFLOW_CONTEXT_INTERNAL(TEngineContext<FContextSingle>, FAssetContext); 
 
 		FAssetContext(UObject* InOwner, UDataflow* InGraph, FTimestamp InTimestamp)
 			: Super(InOwner, InGraph, InTimestamp)
@@ -38,7 +39,6 @@ namespace Dataflow
 
 class DATAFLOWEDITOR_API FDataflowEditorToolkit final : public FBaseCharacterFXEditorToolkit, public FTickableEditorObject, public FNotifyHook
 {
-
 	using FBaseCharacterFXEditorToolkit::ObjectScene;
 
 public:
@@ -48,141 +48,115 @@ public:
 
 	static bool CanOpenDataflowEditor(UObject* ObjectToEdit);
 	static bool HasDataflowAsset(UObject* ObjectToEdit);
+	static UDataflow* GetDataflowAsset(UObject* ObjectToEdit);
+	static const UDataflow* GetDataflowAsset(const UObject* ObjectToEdit);
+
+	/** Get the const dataflow editor datas */
+	const FDataflowEditorDatas& GetDataflowEditorDatas() const;
+
+	/** Return the dataflow editor datas that could be modified */
+	FDataflowEditorDatas& ModifyDataflowEditorDatas();
 
-	// IToolkit Interface
+	// IToolkit interface
 	virtual FName GetToolkitFName() const override;
 	virtual FText GetToolkitName() const override;
 	virtual FText GetBaseToolkitName() const override;
 	virtual FText GetToolkitToolTipText() const override;
 	virtual FString GetWorldCentricTabPrefix() const override;
 	virtual FLinearColor GetWorldCentricTabColorScale() const override;
-
-	// Tab spawners 
 	virtual void RegisterTabSpawners(const TSharedRef<FTabManager>& TabManager) override;
 	virtual void UnregisterTabSpawners(const TSharedRef<class FTabManager>& TabManager) override;
 
-	TSharedRef<SDockTab> SpawnTab_GraphCanvas(const FSpawnTabArgs& Args);
-	TSharedRef<SDockTab> SpawnTab_AssetDetails(const FSpawnTabArgs& Args);
-	TSharedRef<SDockTab> SpawnTab_NodeDetails(const FSpawnTabArgs& Args);
-	TSharedRef<SDockTab> SpawnTab_Skeletal(const FSpawnTabArgs& Args);
-	TSharedRef<SDockTab> SpawnTab_SelectionView(const FSpawnTabArgs& Args);
-	TSharedRef<SDockTab> SpawnTab_CollectionSpreadSheet(const FSpawnTabArgs& Args);
-
-	// Callbacks for Tab
-	void OnTabClosed(TSharedRef<SDockTab> Tab);
-
-	// Member Access
-	UObject* GetAsset() const;
-
-	UDataflow* GetDataflow() { return Dataflow; }
-	const UDataflow* GetDataflow() const { return Dataflow; }
-
-	TSharedPtr<Dataflow::FEngineContext> GetContext() { return Context; }
-	TSharedPtr<const Dataflow::FEngineContext> GetContext() const { return Context; }
-
-	TSharedPtr<IDetailsView> GetAssetDetailsEditor() {return AssetDetailsEditor;}
-	const TSharedPtr<IDetailsView> GetAssetDetailsEditor() const { return AssetDetailsEditor; }
-
-	TSharedPtr<IStructureDetailsView> GetNodeDetailsEditor() { return NodeDetailsEditor; }
-	const TSharedPtr<IStructureDetailsView> GetNodeDetailsEditor() const { return NodeDetailsEditor; }
-
-	TSharedPtr<SDataflowGraphEditor> GetGraphEditor() { return GraphEditor; }
-	const TSharedPtr<SDataflowGraphEditor> GetGraphEditor() const { return GraphEditor; }
-
 protected:
 
-	//~ Begin DataflowEditorActions
+	// List of dataflow actions callbacks
 	void OnPropertyValueChanged(const FPropertyChangedEvent& PropertyChangedEvent);
 	bool OnNodeVerifyTitleCommit(const FText& NewText, UEdGraphNode* GraphNode, FText& OutErrorMessage);
 	void OnNodeTitleCommitted(const FText& InNewText, ETextCommit::Type InCommitType, UEdGraphNode* GraphNode);
 	void OnNodeSelectionChanged(const TSet<UObject*>& NewSelection);
 	void OnNodeDeleted(const TSet<UObject*>& NewSelection);
 	void OnAssetPropertyValueChanged(const FPropertyChangedEvent& PropertyChangedEvent);
-	//~ End DataflowEditorActions
-
+	
+	// Callback to remove the closed one from the listener views
+	void OnTabClosed(TSharedRef<SDockTab> Tab);
+	
 private:
-
+	
+	// Spawning of all the additional tabs (viewport,details ones are coming from the base asset toolkit)
+	TSharedRef<SDockTab> SpawnTab_GraphCanvas(const FSpawnTabArgs& Args);
+	TSharedRef<SDockTab> SpawnTab_NodeDetails(const FSpawnTabArgs& Args);
+	TSharedRef<SDockTab> SpawnTab_Skeletal(const FSpawnTabArgs& Args);
+	TSharedRef<SDockTab> SpawnTab_SelectionView(const FSpawnTabArgs& Args);
+	TSharedRef<SDockTab> SpawnTab_CollectionSpreadSheet(const FSpawnTabArgs& Args);
+	TSharedRef<SDockTab> SpawnTab_AssetDetails(const FSpawnTabArgs& Args);
+	
 	// FTickableEditorObject interface
 	virtual void Tick(float DeltaTime) override;
 	virtual bool IsTickable() const override { return true; }
 	virtual TStatId GetStatId() const override;
 
-	// FBaseCharacterFXEditorToolkit
+	// FBaseCharacterFXEditorToolkit interface
 	virtual FEditorModeID GetEditorModeId() const override;
 	virtual void InitializeEdMode(UBaseCharacterFXEditorMode* EdMode) override;
 	virtual void CreateEditorModeUILayer() override;
 
-	// FAssetEditorToolkit
+	// FAssetEditorToolkit interface
 	virtual bool OnRequestClose(EAssetEditorCloseReason InCloseReason) override;
 	virtual void PostInitAssetEditor() override;
 	virtual void GetSaveableObjects(TArray<UObject*>& OutObjects) const override;
 
-
-	// FBaseAssetToolkit
+	// FBaseAssetToolkit interface
 	virtual void CreateWidgets() override;
 	virtual AssetEditorViewportFactoryFunction GetViewportDelegate() override;
 	virtual TSharedPtr<FEditorViewportClient> CreateEditorViewportClient() const override;
+	virtual void CreateEditorModeManager() override;
 
-	TObjectPtr<UObject> Asset = nullptr;
-	TObjectPtr<UDataflow> Dataflow = nullptr;
-	FString DataflowTerminalPath = "";
+	// List of all the tab names ids that will be used to identify the editor widgets
+	static const FName GraphCanvasTabId;
+	static const FName NodeDetailsTabId;
+	static const FName SkeletalTabId;
+	static const FName SelectionViewTabId_1;
+	static const FName SelectionViewTabId_2;
+	static const FName SelectionViewTabId_3;
+	static const FName SelectionViewTabId_4;
+	static const FName CollectionSpreadSheetTabId_1;
+	static const FName CollectionSpreadSheetTabId_2;
+	static const FName CollectionSpreadSheetTabId_3;
+	static const FName CollectionSpreadSheetTabId_4;
 
-	static const FName ViewportTabId;
+	// List of all the widgets shared ptr that will be built in the editor
 	TSharedPtr<SDataflowEditorViewport> DataflowEditorViewport;
-
-
-	static const FName GraphCanvasTabId;
 	TSharedPtr<SDataflowGraphEditor> GraphEditor;
-	TSharedPtr<FUICommandList> GraphEditorCommands;
-	TSharedRef<SDataflowGraphEditor> CreateGraphEditorWidget(UDataflow* ObjectToEdit, TSharedPtr<IStructureDetailsView> PropertiesEditor);
-
-	static const FName AssetDetailsTabId;
-	TSharedPtr<IDetailsView> AssetDetailsEditor;
-	TSharedPtr<IDetailsView> CreateAssetDetailsEditorWidget(UObject* ObjectToEdit);
-
-	static const FName NodeDetailsTabId;
 	TSharedPtr<IStructureDetailsView> NodeDetailsEditor;
-	TSharedPtr<IStructureDetailsView> CreateNodeDetailsEditorWidget(UObject* ObjectToEdit);
-
-	static const FName SkeletalTabId;
-	TObjectPtr<USkeleton> StubSkeleton;
-	TObjectPtr<USkeletalMesh> StubSkeletalMesh;
 	TSharedPtr<class ISkeletonTree> SkeletalEditor;
-	TSharedPtr<ISkeletonTree> CreateSkeletalEditorWidget(UObject* ObjectToEdit);
-
-	TSet<UObject*> PrevNodeSelection;
-
-	static const FName SelectionViewTabId_1;
+	TSharedPtr<IDetailsView> AssetDetailsEditor;
 	TSharedPtr<FDataflowSelectionView> DataflowSelectionView_1;
-
-	static const FName SelectionViewTabId_2;
 	TSharedPtr<FDataflowSelectionView> DataflowSelectionView_2;
-
-	static const FName SelectionViewTabId_3;
 	TSharedPtr<FDataflowSelectionView> DataflowSelectionView_3;
-
-	static const FName SelectionViewTabId_4;
 	TSharedPtr<FDataflowSelectionView> DataflowSelectionView_4;
-
-	static const FName CollectionSpreadSheetTabId_1;
 	TSharedPtr<FDataflowCollectionSpreadSheet> DataflowCollectionSpreadSheet_1;
-
-	static const FName CollectionSpreadSheetTabId_2;
 	TSharedPtr<FDataflowCollectionSpreadSheet> DataflowCollectionSpreadSheet_2;
-
-	static const FName CollectionSpreadSheetTabId_3;
 	TSharedPtr<FDataflowCollectionSpreadSheet> DataflowCollectionSpreadSheet_3;
-
-	static const FName CollectionSpreadSheetTabId_4;
 	TSharedPtr<FDataflowCollectionSpreadSheet> DataflowCollectionSpreadSheet_4;
 
-	TArray<IDataflowViewListener*> ViewListeners;
+	// Utility factory functions to build the widgets
+	TSharedRef<SDataflowGraphEditor> CreateGraphEditorWidget(UDataflow* ObjectToEdit, TSharedPtr<IStructureDetailsView> PropertiesEditor);
+    TSharedPtr<IDetailsView> CreateAssetDetailsEditorWidget(UObject* ObjectToEdit);
+    TSharedPtr<IStructureDetailsView> CreateNodeDetailsEditorWidget(UObject* ObjectToEdit);
+    TSharedPtr<ISkeletonTree> CreateSkeletalEditorWidget(UObject* ObjectToEdit);
 
-	TSharedPtr<Dataflow::FEngineContext> Context;
-	Dataflow::FTimestamp LastNodeTimestamp = Dataflow::FTimestamp::Invalid;
+	// List of editor commands used  for the dataflow asset
+	TSharedPtr<FUICommandList> GraphEditorCommands;
 
+	// List of selection view / collection spreadsheet widgets that are listening to any changed in the graph
+	TArray<IDataflowViewListener*> ViewListeners;
+
+	// Graph delegates used to update the UI
 	FDelegateHandle OnSelectionChangedMulticastDelegateHandle;
-	FDelegateHandle OnNodeDeletedMulticastDelegateHandle;
-	FDelegateHandle OnFinishedChangingPropertiesDelegateHandle;
-	FDelegateHandle OnFinishedChangingAssetPropertiesDelegateHandle;
+    FDelegateHandle OnNodeDeletedMulticastDelegateHandle;
+    FDelegateHandle OnFinishedChangingPropertiesDelegateHandle;
+    FDelegateHandle OnFinishedChangingAssetPropertiesDelegateHandle;
+
+	// Previous node selection to compare values in delegates
+	TSet<UObject*> PrevNodeSelection;
 };
diff --git a/Engine/Plugins/Experimental/Dataflow/Source/DataflowEditor/Public/Dataflow/DataflowEditorViewport.h b/Engine/Plugins/Experimental/Dataflow/Source/DataflowEditor/Public/Dataflow/DataflowEditorViewport.h
--- a/Engine/Plugins/Experimental/Dataflow/Source/DataflowEditor/Public/Dataflow/DataflowEditorViewport.h
+++ b/Engine/Plugins/Experimental/Dataflow/Source/DataflowEditor/Public/Dataflow/DataflowEditorViewport.h
@@ -25,7 +25,9 @@ public:
 
 	// SEditorViewport
 	virtual void BindCommands() override;
+	virtual TSharedPtr<SWidget> MakeViewportToolbar() override;
 	virtual bool IsVisible() const override;
+	virtual void OnFocusViewportToSelection() override;
 
 	// ICommonEditorViewportToolbarInfoProvider
 	virtual TSharedRef<class SEditorViewport> GetViewportWidget() override;
diff --git a/Engine/Plugins/Experimental/Dataflow/Source/DataflowEditor/Public/Dataflow/DataflowEditorViewportClient.h b/Engine/Plugins/Experimental/Dataflow/Source/DataflowEditor/Public/Dataflow/DataflowEditorViewportClient.h
--- a/Engine/Plugins/Experimental/Dataflow/Source/DataflowEditor/Public/Dataflow/DataflowEditorViewportClient.h
+++ b/Engine/Plugins/Experimental/Dataflow/Source/DataflowEditor/Public/Dataflow/DataflowEditorViewportClient.h
@@ -8,8 +8,10 @@
 
 class FDataflowEditorToolkit;
 class ADataflowActor;
-
-
+class UTransformProxy;
+class UCombinedTransformGizmo;
+class FTransformGizmoDataBinder;
+class FDataflowPreviewScene;
 
 class DATAFLOWEDITOR_API FDataflowEditorViewportClient : public FEditorViewportClient
 {
@@ -19,37 +21,38 @@ public:
 	FDataflowEditorViewportClient(FEditorModeTools* InModeTools, FPreviewScene* InPreviewScene,
 								  const TWeakPtr<SEditorViewport> InEditorViewportWidget = nullptr);
 
+	/** Set the data flow toolkit used to create the client*/
 	void SetDataflowEditorToolkit(TWeakPtr<FDataflowEditorToolkit> DataflowToolkit);
 
-
+	/** Return the latest timestamp */
 	::Dataflow::FTimestamp LatestTimestamp(const UDataflow* Dataflow, const ::Dataflow::FContext* Context);
 
+	// Selection utils
 	void SetSelectionMode(FDataflowSelectionState::EMode InState);
 	bool CanSetSelectionMode(FDataflowSelectionState::EMode InState);
 	bool IsSelectionModeActive(FDataflowSelectionState::EMode InState);
 	FDataflowSelectionState::EMode GetSelectionMode() const { return SelectionMode; }
-
-
+	
 	// FEditorViewportClient interface
 	virtual bool InputKey(const FInputKeyEventArgs& EventArgs) override;
 	virtual void ProcessClick(FSceneView& View, HHitProxy* HitProxy, FKey Key, EInputEvent Event, uint32 HitX, uint32 HitY) override;
 	virtual void Tick(float DeltaSeconds) override;
-	// End of FEditorViewportClient
 
-	//~ FGCObject Interface
+	// FGCObject Interface
 	virtual void AddReferencedObjects(FReferenceCollector& Collector) override;
 	virtual FString GetReferencerName() const override { return TEXT("FDataflowEditorViewportClient"); }
 
 private:
 
+	/** Toolkit used to create the viewport client */
 	TWeakPtr<FDataflowEditorToolkit> DataflowEditorToolkitPtr = nullptr;
+	
+	/** Last time stamp used for dataflow evaluation on the component */
 	::Dataflow::FTimestamp LastModifiedTimestamp = ::Dataflow::FTimestamp::Invalid;
-	TObjectPtr<ADataflowActor> DataflowActor = nullptr;
 
-	//
-	// Selection
-	//
+	/** Dataflow preview scene from the toolkit */
+	FDataflowPreviewScene* PreviewScene = nullptr;
 
+	/** Selection mode to be used for vertices/faces/objects */
 	FDataflowSelectionState::EMode SelectionMode = FDataflowSelectionState::EMode::DSS_Dataflow_None;
-
 };
diff --git a/Engine/Plugins/Experimental/Dataflow/Source/DataflowEditor/Public/Dataflow/DataflowEditorViewportToolbar.h b/Engine/Plugins/Experimental/Dataflow/Source/DataflowEditor/Public/Dataflow/DataflowEditorViewportToolbar.h
--- a/Engine/Plugins/Experimental/Dataflow/Source/DataflowEditor/Public/Dataflow/DataflowEditorViewportToolbar.h
+++ b/Engine/Plugins/Experimental/Dataflow/Source/DataflowEditor/Public/Dataflow/DataflowEditorViewportToolbar.h
@@ -3,21 +3,21 @@
 
 #include "CoreMinimal.h"
 #include "SEditorViewportToolBarMenu.h"
+#include "SCommonEditorViewportToolbarBase.h"
 
 // ----------------------------------------------------------------------------------
 class SDataflowEditorViewport;
 
-class SDataflowViewportSelectionToolBar : public SViewportToolBar
+class SDataflowViewportSelectionToolBar : public SCommonEditorViewportToolbarBase
 {
 public:
 	SLATE_BEGIN_ARGS(SDataflowViewportSelectionToolBar) {}
-		SLATE_ARGUMENT(TWeakPtr<SDataflowEditorViewport>, EditorViewport)
 	SLATE_END_ARGS()
 
-		/** Constructs this widget with the given parameters */
-		void Construct(const FArguments& InArgs);
+	/** Constructs this widget with the given parameters */
+	void Construct(const FArguments& InArgs, TSharedPtr<SDataflowEditorViewport> InDataflowViewport);
 
-		TSharedRef<SWidget> MakeSelectionModeToolBar();
+	TSharedRef<SWidget> MakeSelectionModeToolBar();
 
 
 private:
diff --git a/Engine/Plugins/Experimental/Dataflow/Source/DataflowEditor/Public/Dataflow/DataflowPreviewScene.h b/Engine/Plugins/Experimental/Dataflow/Source/DataflowEditor/Public/Dataflow/DataflowPreviewScene.h
new file mode 100644
--- /dev/null
+++ b/Engine/Plugins/Experimental/Dataflow/Source/DataflowEditor/Public/Dataflow/DataflowPreviewScene.h
@@ -0,0 +1,106 @@
+// Copyright Epic Games, Inc. All Rights Reserved.
+
+#pragma once
+
+#include "AdvancedPreviewScene.h"
+
+struct FDataflowEditorDatas;
+class USkeletalMeshComponent;
+class UDynamicMeshComponent;
+class UDataflowComponent;
+class UAnimationAsset;
+class UAnimSingleNodeInstance;
+class UMaterialInterface;
+class FAssetEditorModeManager;
+
+namespace UE
+{
+	namespace Geometry
+	{
+		class FDynamicMesh3;
+	}
+}
+
+/**
+ * Dataflow preview scene
+ * @brief the scene is holding all the objects that will be
+ * visible and potentially editable within the viewport
+ */
+class DATAFLOWEDITOR_API FDataflowPreviewScene : public FAdvancedPreviewScene
+{
+public:
+
+	FDataflowPreviewScene(FPreviewScene::ConstructionValues ConstructionValues, FDataflowEditorDatas& DataflowAssetDatas);
+	virtual ~FDataflowPreviewScene();
+
+	// FGCObject interface
+	virtual void AddReferencedObjects(FReferenceCollector& Collector) override;
+
+	// Modify/Get the preview anim instance
+	UAnimSingleNodeInstance* ModifyPreviewAnimInstance() { return PreviewAnimInstance; }
+	const UAnimSingleNodeInstance* const GetPreviewAnimInstance() const { return PreviewAnimInstance; }
+
+	// Modify/Get the data flow component
+	UDataflowComponent* ModifyDataflowComponent() { return DataflowComponent; }
+	const UDataflowComponent* GetDataflowComponent() const { return DataflowComponent; }
+
+	// Modify/Get the skeletal mesh component
+	USkeletalMeshComponent* ModifySkeletalMeshComponent() { return SkeletalMeshComponent; }
+	const USkeletalMeshComponent* GetSkeletalMeshComponent() const { return SkeletalMeshComponent; }
+	
+	// Modify/Get the data flow datas
+	FDataflowEditorDatas& ModifyDataflowDatas() { return DataflowDatas; }
+	const FDataflowEditorDatas& GetDataflowDatas() const { return DataflowDatas; }
+
+	// Modify/Get the data flow mode manager
+	TSharedPtr<FAssetEditorModeManager>& ModifyDataflowModeManager() { return DataflowModeManager; }
+	const TSharedPtr<FAssetEditorModeManager>& GetDataflowModeManager() const { return DataflowModeManager; }
+	
+	/** Reset all the dynamic mesh components */
+	void ResetDynamicMeshComponents();
+	
+	/** Add a dynamic mesh component to the scene */
+	TObjectPtr<UDynamicMeshComponent>& AddDynamicMeshComponent(UE::Geometry::FDynamicMesh3&& DynamicMesh, const TArray<UMaterialInterface*>& MaterialSet);
+
+	/** Create the PreviewAnimationInstance if the AnimationAsset and SkeletalMesh both exist, and set the animation to run on the SkeletalMeshComponent */
+	void UpdateSkeletalMeshComponent();
+
+	/** Update the data flow component with the dataflow asset and reset the render targets */
+	void UpdateDataflowComponent();
+
+	/** Build the scene bounding box */
+	FBox GetBoundingBox() const;
+	
+private:
+	
+	/** Check if a primitive component is selected */
+	bool IsComponentSelected(const UPrimitiveComponent* InComponent) const;
+
+	/** Data flows assets datas used to build the scene components */
+	FDataflowEditorDatas& DataflowDatas;
+	
+	/** Dataflow actor */
+	TObjectPtr<AActor> DataflowActor = nullptr;
+	
+	/** Skeletal mesh actor */
+	TObjectPtr<AActor> SkeletalMeshActor = nullptr;
+
+	/** Dynamic mesh actor */
+	TObjectPtr<AActor> DynamicMeshActor = nullptr;
+	
+	/** Dataflow component that will be modified */
+	TObjectPtr<UDataflowComponent> DataflowComponent = nullptr;
+
+	/** Skeletal mesh component used for potential binding */
+	TObjectPtr<USkeletalMeshComponent> SkeletalMeshComponent = nullptr;
+
+	/** Preview anim instance for the skeletal mesh */
+	TObjectPtr<UAnimSingleNodeInstance> PreviewAnimInstance = nullptr;
+
+	/** Dynamic mesh components generated by the tools for edit/manipulation*/
+	TArray<TObjectPtr<UDynamicMeshComponent>> DynamicMeshComponents;
+
+	/** Mode Manager for selection */
+	TSharedPtr<FAssetEditorModeManager> DataflowModeManager;
+};
+
diff --git a/Engine/Plugins/Experimental/Dataflow/Source/DataflowEditor/Public/Dataflow/DataflowToolTarget.h b/Engine/Plugins/Experimental/Dataflow/Source/DataflowEditor/Public/Dataflow/DataflowToolTarget.h
new file mode 100644
--- /dev/null
+++ b/Engine/Plugins/Experimental/Dataflow/Source/DataflowEditor/Public/Dataflow/DataflowToolTarget.h
@@ -0,0 +1,117 @@
+// Copyright Epic Games, Inc. All Rights Reserved.
+
+#pragma once
+
+#include "CoreMinimal.h"
+
+#include "TargetInterfaces/DynamicMeshCommitter.h"
+#include "TargetInterfaces/DynamicMeshProvider.h"
+#include "TargetInterfaces/MaterialProvider.h"
+#include "TargetInterfaces/MeshDescriptionCommitter.h"
+#include "TargetInterfaces/MeshDescriptionProvider.h"
+#include "ToolTargets/PrimitiveComponentToolTarget.h"
+
+#include "DataflowToolTarget.generated.h"
+
+class UDataflow;
+
+/**
+ * A tool target backed by a read-only dataflow asset that can provide and take a mesh
+ * description.
+ */
+UCLASS(Transient)
+class DATAFLOWEDITOR_API UDataflowReadOnlyToolTarget :
+	public UToolTarget,
+	public IMeshDescriptionProvider,
+	public IDynamicMeshProvider, 
+	public IMaterialProvider
+{
+	GENERATED_BODY()
+
+public:
+
+	// UToolTarget implementation
+	virtual bool IsValid() const override;
+
+	// IMeshDescriptionProvider implementation
+	virtual const FMeshDescription* GetMeshDescription(const FGetMeshParameters& GetMeshParams = FGetMeshParameters()) override;
+	virtual FMeshDescription GetEmptyMeshDescription() override;
+
+	// IMaterialProvider implementation
+	virtual int32 GetNumMaterials() const override;
+	virtual UMaterialInterface* GetMaterial(int32 MaterialIndex) const override;
+	virtual void GetMaterialSet(FComponentMaterialSet& MaterialSetOut, bool bPreferAssetMaterials) const override;
+	virtual bool CommitMaterialSetUpdate(const FComponentMaterialSet& MaterialSet, bool bApplyToAsset) override;
+
+	// IDynamicMeshProvider
+	virtual UE::Geometry::FDynamicMesh3 GetDynamicMesh() override;
+	virtual UE::Geometry::FDynamicMesh3 GetDynamicMesh(bool bRequestTangents) override;
+
+protected:
+	// So that the tool target factory can poke into Component.
+	friend class UDataflowReadOnlyToolTargetFactory;
+
+	// Until UDataflow stores its internal representation as FMeshDescription, we need to
+	// retain the storage here to cover the lifetime of the pointer returned by GetMeshDescription(). 
+	TUniquePtr<FMeshDescription> DataflowMeshDescription;
+
+	// Internal dataflow pointer
+	TObjectPtr<UDataflow> Dataflow = nullptr;
+
+	// Internal asset pointer
+	TObjectPtr<UObject> Asset = nullptr;
+};
+
+/**
+ * A tool target backed by a dataflow asset that can provide and take a mesh
+ * description.
+ */
+UCLASS(Transient)
+class DATAFLOWEDITOR_API UDataflowToolTarget :
+	public UDataflowReadOnlyToolTarget,
+	public IMeshDescriptionCommitter,
+	public IDynamicMeshCommitter
+{
+	GENERATED_BODY()
+
+public:
+	// IMeshDescriptionCommitter implementation
+	virtual void CommitMeshDescription(const FCommitter& Committer, const FCommitMeshParameters& CommitParams = FCommitMeshParameters()) override;
+	using IMeshDescriptionCommitter::CommitMeshDescription; // unhide the other overload
+
+	// IDynamicMeshCommitter
+	virtual void CommitDynamicMesh(const UE::Geometry::FDynamicMesh3& Mesh, const FDynamicMeshCommitInfo& CommitInfo) override;
+	using IDynamicMeshCommitter::CommitDynamicMesh; // unhide the other overload
+
+protected:
+	// So that the tool target factory can poke into Component.
+	friend class UDataflowToolTargetFactory;
+};
+
+
+/** Factory for UDataflowReadOnlyToolTarget to be used by the target manager. */
+UCLASS(Transient)
+class DATAFLOWEDITOR_API UDataflowReadOnlyToolTargetFactory : public UToolTargetFactory
+{
+	GENERATED_BODY()
+
+public:
+
+	// UToolTargetFactory implementation
+	virtual bool CanBuildTarget(UObject* SourceObject, const FToolTargetTypeRequirements& TargetTypeInfo) const override;
+	virtual UToolTarget* BuildTarget(UObject* SourceObject, const FToolTargetTypeRequirements& TargetTypeInfo) override;
+};
+
+
+/** Factory for UDataflowToolTarget to be used by the target manager. */
+UCLASS(Transient)
+class DATAFLOWEDITOR_API UDataflowToolTargetFactory : public UToolTargetFactory
+{
+	GENERATED_BODY()
+
+public:
+
+	// UToolTargetFactory implementation
+	virtual bool CanBuildTarget(UObject* SourceObject, const FToolTargetTypeRequirements& TargetTypeInfo) const override;
+	virtual UToolTarget* BuildTarget(UObject* SourceObject, const FToolTargetTypeRequirements& TargetTypeInfo) override;
+};
diff --git a/Engine/Plugins/Experimental/Dataflow/Source/DataflowEnginePlugin/Private/Dataflow/DataflowComponent.cpp b/Engine/Plugins/Experimental/Dataflow/Source/DataflowEnginePlugin/Private/Dataflow/DataflowComponent.cpp
--- a/Engine/Plugins/Experimental/Dataflow/Source/DataflowEnginePlugin/Private/Dataflow/DataflowComponent.cpp
+++ b/Engine/Plugins/Experimental/Dataflow/Source/DataflowEnginePlugin/Private/Dataflow/DataflowComponent.cpp
@@ -48,6 +48,25 @@ void UDataflowComponent::AddRenderTarget(const UDataflowEdNode* InTarget)
 	Invalidate();
 }
 
+void UDataflowComponent::BuildRenderCollection()
+{
+	RenderCollection = FManagedArrayCollection();
+	GeometryCollection::Facades::FRenderingFacade Facade(RenderCollection);
+	Facade.DefineSchema();
+
+	if (Context && Dataflow)
+	{
+		for (const UDataflowEdNode* Target : RenderTargets)
+		{
+			if (Target)
+			{
+				Target->Render(Facade, Context);
+			}
+		}
+	}
+	MarkRenderStateDirty();
+}
+
 void UDataflowComponent::TickComponent(float DeltaTime, enum ELevelTick TickType, FActorComponentTickFunction *ThisTickFunction)
 {
 	bool bNeedsSceneProxyUpdate = false;
@@ -61,14 +80,13 @@ void UDataflowComponent::TickComponent(float DeltaTime, enum ELevelTick TickType
 		GeometryCollection::Facades::FRenderingFacade Facade(RenderCollection);
 		Facade.DefineSchema();
 
-		bool bNeedsRefresh = false;
 		if (Context && Dataflow)
 		{
 			for (const UDataflowEdNode* Target : RenderTargets)
 			{
 				if (Target)
 				{
-					bNeedsRefresh |= Target->Render(Facade, Context);
+					Target->Render(Facade, Context);
 				}
 			}
 		}
@@ -104,6 +122,8 @@ void UDataflowComponent::UpdateLocalBounds()
 		BoundingBox = BoundsFacade.GetBoundingBoxInCollectionSpace();
 
 		bBoundsNeedsUpdate = false;
+
+		UpdateBounds();
 	}
 }
 
diff --git a/Engine/Plugins/Experimental/Dataflow/Source/DataflowEnginePlugin/Private/Dataflow/DataflowEngineRendering.cpp b/Engine/Plugins/Experimental/Dataflow/Source/DataflowEnginePlugin/Private/Dataflow/DataflowEngineRendering.cpp
--- a/Engine/Plugins/Experimental/Dataflow/Source/DataflowEnginePlugin/Private/Dataflow/DataflowEngineRendering.cpp
+++ b/Engine/Plugins/Experimental/Dataflow/Source/DataflowEnginePlugin/Private/Dataflow/DataflowEngineRendering.cpp
@@ -35,6 +35,12 @@ namespace Dataflow
 
 					FManagedArrayCollection Default;
 					const FManagedArrayCollection& Collection = State.GetValue<FManagedArrayCollection>(PrimaryOutput, Default);
+					const bool bFoundIndices = Collection.FindAttributeTyped<FIntVector>("Indices", FGeometryCollection::FacesGroup) != nullptr;
+					const bool bFoundVertices = Collection.FindAttributeTyped<FVector3f>("Vertex", FGeometryCollection::VerticesGroup) != nullptr;
+					const bool bFoundTransforms = Collection.FindAttributeTyped<FTransform3f>(FTransformCollection::TransformAttribute, FTransformCollection::TransformGroup) != nullptr;
+					const bool bFoundBoneMap = Collection.FindAttributeTyped<int32>("BoneMap", FGeometryCollection::VerticesGroup) != nullptr;
+					const bool bFoundParents = Collection.FindAttributeTyped<int32>(FTransformCollection::ParentAttribute, FTransformCollection::TransformGroup) != nullptr;
+					UE_LOG(LogTemp, Warning, TEXT("Render GC with found params = %d %d %d %d %d"), bFoundIndices, bFoundVertices, bFoundTransforms, bFoundBoneMap, bFoundParents);
 					
 					if (Collection.FindAttributeTyped<FIntVector>("Indices", FGeometryCollection::FacesGroup)
 						&& Collection.FindAttributeTyped<FVector3f>("Vertex", FGeometryCollection::VerticesGroup)
diff --git a/Engine/Plugins/Experimental/Dataflow/Source/DataflowEnginePlugin/Private/Dataflow/DataflowEngineSceneProxy.h b/Engine/Plugins/Experimental/Dataflow/Source/DataflowEnginePlugin/Private/Dataflow/DataflowEngineSceneProxy.h
--- a/Engine/Plugins/Experimental/Dataflow/Source/DataflowEnginePlugin/Private/Dataflow/DataflowEngineSceneProxy.h
+++ b/Engine/Plugins/Experimental/Dataflow/Source/DataflowEnginePlugin/Private/Dataflow/DataflowEngineSceneProxy.h
@@ -37,12 +37,6 @@ public:
 	FDataflowEngineSceneProxy(UDataflowComponent* Component);
 	virtual ~FDataflowEngineSceneProxy();
 
-	/** Create the rendering buffer resources */
-	void InitResources();
-
-	/** Return the rendering buffer resources */
-	void ReleaseResources();
-
 	//~ FPrimitiveSceneProxy
 	virtual FPrimitiveViewRelevance GetViewRelevance(const FSceneView* View) const override;
 	virtual bool CanBeOccluded() const override;
diff --git a/Engine/Plugins/Experimental/Dataflow/Source/DataflowEnginePlugin/Public/Dataflow/DataflowComponent.h b/Engine/Plugins/Experimental/Dataflow/Source/DataflowEnginePlugin/Public/Dataflow/DataflowComponent.h
--- a/Engine/Plugins/Experimental/Dataflow/Source/DataflowEnginePlugin/Public/Dataflow/DataflowComponent.h
+++ b/Engine/Plugins/Experimental/Dataflow/Source/DataflowEnginePlugin/Public/Dataflow/DataflowComponent.h
@@ -51,6 +51,7 @@ public:
 	void SetRenderingCollection(FManagedArrayCollection&& InCollection);
 	const FManagedArrayCollection& GetRenderingCollection() const;
 	      FManagedArrayCollection& ModifyRenderingCollection();
+	void BuildRenderCollection();
 
 
 	/** Dataflow */
diff --git a/Engine/Source/Editor/UnrealEd/Private/Tools/BaseAssetToolkit.cpp b/Engine/Source/Editor/UnrealEd/Private/Tools/BaseAssetToolkit.cpp
--- a/Engine/Source/Editor/UnrealEd/Private/Tools/BaseAssetToolkit.cpp
+++ b/Engine/Source/Editor/UnrealEd/Private/Tools/BaseAssetToolkit.cpp
@@ -81,6 +81,14 @@ void FBaseAssetToolkit::RegisterTabSpawners(const TSharedRef<FTabManager>& InTab
 
 }
 
+void FBaseAssetToolkit::UnregisterTabSpawners(const TSharedRef<FTabManager>& InTabManager)
+{
+	FAssetEditorToolkit::UnregisterTabSpawners(InTabManager);
+
+	InTabManager->UnregisterTabSpawner(ViewportTabID);
+	InTabManager->UnregisterTabSpawner(DetailsTabID);
+}
+
 const TSharedRef<FTabManager::FLayout> FBaseAssetToolkit::GetDefaultLayout() const
 {
 	return StandaloneDefaultLayout.ToSharedRef();
diff --git a/Engine/Source/Editor/UnrealEd/Public/Tools/BaseAssetToolkit.h b/Engine/Source/Editor/UnrealEd/Public/Tools/BaseAssetToolkit.h
--- a/Engine/Source/Editor/UnrealEd/Public/Tools/BaseAssetToolkit.h
+++ b/Engine/Source/Editor/UnrealEd/Public/Tools/BaseAssetToolkit.h
@@ -42,7 +42,7 @@ public:
 		return FLinearColor::White;
 	}
 	UNREALED_API virtual void RegisterTabSpawners(const TSharedRef<FTabManager>& InTabManager) override;
-
+	UNREALED_API virtual void UnregisterTabSpawners(const TSharedRef<FTabManager>& InTabManager) override;
 
 	UNREALED_API virtual const TSharedRef<FTabManager::FLayout> GetDefaultLayout() const;
 	UNREALED_API virtual void CreateWidgets();
diff --git a/Engine/Source/Runtime/Experimental/Chaos/Public/GeometryCollection/Facades/CollectionRenderingFacade.h b/Engine/Source/Runtime/Experimental/Chaos/Public/GeometryCollection/Facades/CollectionRenderingFacade.h
--- a/Engine/Source/Runtime/Experimental/Chaos/Public/GeometryCollection/Facades/CollectionRenderingFacade.h
+++ b/Engine/Source/Runtime/Experimental/Chaos/Public/GeometryCollection/Facades/CollectionRenderingFacade.h
@@ -74,13 +74,19 @@ namespace GeometryCollection::Facades
 
 		/** GetVertices */
 		const TManagedArray< FVector3f >& GetVertices() const { return VertexAttribute.Get(); }
+		TManagedArray< FVector3f >& ModifyVertices() { check(!IsConst()); return VertexAttribute.Modify(); }
+		
+		/** GetNormals */
+		const TManagedArray< FVector3f >& GetNormals() const { return VertexNormalAttribute.Get(); }
+		TManagedArray< FVector3f >& ModifyNormals() { check(!IsConst()); return VertexNormalAttribute.Modify(); }
 
 		/** GetVertexSelectionAttribute */
 		const TManagedArray< int32 >& GetVertexSelection() const { return VertexSelectionAttribute.Get(); }
-		      TManagedArray< int32 >& ModifyVertexSelection() { check(!IsConst()); return VertexSelectionAttribute.Modify(); }
+		TManagedArray< int32 >& ModifyVertexSelection() { check(!IsConst()); return VertexSelectionAttribute.Modify(); }
 
 		/** GetVertexToGeometryIndexAttribute */
 		const TManagedArray< int32 >& GetVertexToGeometryIndex() const { return VertexToGeometryIndexAttribute.Get(); }
+		TManagedArray< int32 >& ModifyVertexToGeometryIndex() { check(!IsConst()); return VertexToGeometryIndexAttribute.Modify(); }
 
 		/** HitProxyIDAttribute */
 		const TManagedArray< int32 >& GetVertexHitProxyIndex() const { return VertexHitProxyIndexAttribute.Get(); }
