Dataflow : Update EditorDatas to be a UObject

diff --git a/Engine/Plugins/Experimental/Dataflow/Source/DataflowEditor/Private/Dataflow/DataflowEditor.cpp b/Engine/Plugins/Experimental/Dataflow/Source/DataflowEditor/Private/Dataflow/DataflowEditor.cpp
--- a/Engine/Plugins/Experimental/Dataflow/Source/DataflowEditor/Private/Dataflow/DataflowEditor.cpp
+++ b/Engine/Plugins/Experimental/Dataflow/Source/DataflowEditor/Private/Dataflow/DataflowEditor.cpp
@@ -61,6 +61,11 @@ namespace Private
 	}
 };
 
+UDataflowEditor::UDataflowEditor() : Super()
+{
+	Content = NewObject<UDataflowEditorContent>(this, MakeUniqueObjectName(this, UDataflowEditor::StaticClass(), "EditorData"));
+}
+
 TSharedPtr<FBaseAssetToolkit> UDataflowEditor::CreateToolkit()
 {
 	TSharedPtr<FDataflowEditorToolkit> DataflowToolkit = MakeShared<FDataflowEditorToolkit>(this);
@@ -69,42 +74,43 @@ TSharedPtr<FBaseAssetToolkit> UDataflowEditor::CreateToolkit()
 
 void UDataflowEditor::Initialize(const TArray<TObjectPtr<UObject>>& InObjects)
 {
-	// We extract the dataflow and the associated datas from the list of objects
+	check(Content);
+
 	TArray<TObjectPtr<UObject>> ObjectsToEdit;
-	DataflowDatas.bHasValidSkeletalMesh = false;
+	Content->bHasValidSkeletalMesh = false;
 	if (ensure(InObjects.Num() == 1))
 	{
 		TObjectPtr<UObject> RootObject = InObjects[0];
-		DataflowDatas.DataflowAsset = Cast<UDataflow>(RootObject);
-		if(!DataflowDatas.DataflowAsset)
+		Content->DataflowAsset = Cast<UDataflow>(RootObject);
+		if(!Content->DataflowAsset)
 		{
-			DataflowDatas.DataflowAsset = Private::GetDataflowAssetFrom(RootObject);
-			DataflowDatas.DataflowTerminal = Private::GetDataflowTerminalFrom(RootObject);
-			DataflowDatas.SkeletalMesh = Private::GetSkeletalMeshFrom(RootObject);
-			DataflowDatas.AnimationAsset = Private::GetAnimationAssetFrom(RootObject);
+			Content->DataflowAsset = Private::GetDataflowAssetFrom(RootObject);
+			Content->DataflowTerminal = Private::GetDataflowTerminalFrom(RootObject);
+			Content->SkeletalMesh = Private::GetSkeletalMeshFrom(RootObject);
+			Content->AnimationAsset = Private::GetAnimationAssetFrom(RootObject);
 		}
-		if(DataflowDatas.DataflowAsset)
+		if(Content->DataflowAsset)
 		{
-			DataflowDatas.DataflowOwner = RootObject;
-			DataflowDatas.DataflowAsset->Schema = UDataflowSchema::StaticClass();
-			DataflowDatas.DataflowContext = MakeShared<Dataflow::FEngineContext>(RootObject, DataflowDatas.DataflowAsset, FPlatformTime::Cycles64());
-			DataflowDatas.LastNodeTimestamp = DataflowDatas.DataflowContext->GetTimestamp();
+			Content->DataflowOwner = RootObject;
+			Content->DataflowAsset->Schema = UDataflowSchema::StaticClass();
+			Content->DataflowContext = MakeShared<Dataflow::FEngineContext>(RootObject, Content->DataflowAsset, FPlatformTime::Cycles64());
+			Content->LastNodeTimestamp = Content->DataflowContext->GetTimestamp();
 
-			ObjectsToEdit.Add(DataflowDatas.DataflowOwner);
+			ObjectsToEdit.Add(Content->DataflowOwner);
 
-			if(!DataflowDatas.SkeletalMesh)
+			if(!Content->SkeletalMesh)
 			{
-				const FName SkeletonName = MakeUniqueObjectName(DataflowDatas.DataflowAsset, UDataflow::StaticClass(), FName("USkeleton"));
-				const FName SkeletalMeshName = MakeUniqueObjectName(DataflowDatas.DataflowAsset, UDataflow::StaticClass(), FName("USkeletalMesh"));
+				const FName SkeletonName = MakeUniqueObjectName(Content->DataflowAsset, UDataflow::StaticClass(), FName("USkeleton"));
+				const FName SkeletalMeshName = MakeUniqueObjectName(Content->DataflowAsset, UDataflow::StaticClass(), FName("USkeletalMesh"));
 
-				DataflowDatas.Skeleton = NewObject<USkeleton>(DataflowDatas.DataflowAsset, SkeletonName);
-				DataflowDatas.SkeletalMesh = NewObject<USkeletalMesh>(DataflowDatas.DataflowAsset, SkeletalMeshName);
-				DataflowDatas.SkeletalMesh->SetSkeleton(DataflowDatas.Skeleton);
+				Content->Skeleton = NewObject<USkeleton>(Content->DataflowAsset, SkeletonName);
+				Content->SkeletalMesh = NewObject<USkeletalMesh>(Content->DataflowAsset, SkeletalMeshName);
+				Content->SkeletalMesh->SetSkeleton(Content->Skeleton);
 			}
 			else
 			{
-				DataflowDatas.Skeleton = DataflowDatas.SkeletalMesh->GetSkeleton();
-				DataflowDatas.bHasValidSkeletalMesh = true;
+				Content->Skeleton = Content->SkeletalMesh->GetSkeleton();
+				Content->bHasValidSkeletalMesh = true;
 			}
 		}
 	}
diff --git a/Engine/Plugins/Experimental/Dataflow/Source/DataflowEditor/Private/Dataflow/DataflowEditorToolkit.cpp b/Engine/Plugins/Experimental/Dataflow/Source/DataflowEditor/Private/Dataflow/DataflowEditorToolkit.cpp
--- a/Engine/Plugins/Experimental/Dataflow/Source/DataflowEditor/Private/Dataflow/DataflowEditorToolkit.cpp
+++ b/Engine/Plugins/Experimental/Dataflow/Source/DataflowEditor/Private/Dataflow/DataflowEditorToolkit.cpp
@@ -92,7 +92,7 @@ FDataflowEditorToolkit::FDataflowEditorToolkit(UAssetEditor* InOwningAssetEditor
 	FAdvancedPreviewScene::ConstructionValues PreviewSceneArgs;
 	PreviewSceneArgs.bShouldSimulatePhysics = 1;
 	PreviewSceneArgs.bCreatePhysicsScene = 1;
-	ObjectScene = MakeUnique<FDataflowPreviewScene>(PreviewSceneArgs, ModifyDataflowEditorDatas());
+	ObjectScene = MakeUnique<FDataflowPreviewScene>(PreviewSceneArgs, GetDataflowEditorContent());
 }
 
 
@@ -124,7 +124,7 @@ FDataflowEditorToolkit::~FDataflowEditorToolkit()
 void FDataflowEditorToolkit::CreateEditorModeManager()
 {
 	FBaseCharacterFXEditorToolkit::CreateEditorModeManager();
-	static_cast<FDataflowPreviewScene*>(ObjectScene.Get())->ModifyDataflowModeManager()
+	static_cast<FDataflowPreviewScene*>(ObjectScene.Get())->GetDataflowModeManager()
 		= StaticCastSharedPtr<FAssetEditorModeManager>(EditorModeManager);
 }
 
@@ -190,16 +190,17 @@ FEditorModeID FDataflowEditorToolkit::GetEditorModeId() const
 	return UDataflowEditorMode::EM_DataflowEditorModeId;
 }
 
-const FDataflowEditorDatas& FDataflowEditorToolkit::GetDataflowEditorDatas() const
+TObjectPtr<UDataflowEditorContent> FDataflowEditorToolkit::GetDataflowEditorContent()
 {
-	return Cast<UDataflowEditor>(OwningAssetEditor)->DataflowDatas;
+	return Cast<UDataflowEditor>(OwningAssetEditor)->Content;
 }
 
-FDataflowEditorDatas& FDataflowEditorToolkit::ModifyDataflowEditorDatas()
+TObjectPtr<const UDataflowEditorContent> FDataflowEditorToolkit::GetDataflowEditorContent() const
 {
-	return Cast<UDataflowEditor>(OwningAssetEditor)->DataflowDatas;
+	return Cast<UDataflowEditor>(OwningAssetEditor)->Content;
 }
 
+
 bool FDataflowEditorToolkit::OnRequestClose(EAssetEditorCloseReason InCloseReason)
 {
 	// Note: This needs a bit of adjusting, because currently OnRequestClose seems to be 
@@ -272,7 +273,7 @@ void FDataflowEditorToolkit::PostInitAssetEditor()
 	TSharedPtr<FDataflowEditorViewportClient> DataflowViewportClient = StaticCastSharedPtr<FDataflowEditorViewportClient>(ViewportClient);
 	DataflowViewportClient->SetDataflowEditorToolkit(StaticCastSharedRef<FDataflowEditorToolkit>(this->AsShared()));
 
-	//SetEditingObject(GetDataflowEditorDatas().DataflowOwner);
+	//SetEditingObject(GetDataflowEditorContent().DataflowOwner);
 
 	// @todo(DataflowMode) : Do this in UDataflowEditorMode::RefocusViewportClient and use the DataflowComponents bounds
 	//		... see ChaosClothAssetEditorMode::RefocusRestSpaceViewportClient() for reference
@@ -303,10 +304,13 @@ void FDataflowEditorToolkit::GetSaveableObjects(TArray<UObject*>& OutObjects) co
 {
 	FBaseCharacterFXEditorToolkit::GetSaveableObjects(OutObjects);
 
-	if (UDataflow* DataflowAsset = GetDataflowEditorDatas().DataflowAsset)
+	if (ensure(GetDataflowEditorContent()))
 	{
-		check(DataflowAsset->IsAsset());
-		OutObjects.Add(DataflowAsset);
+		if (UDataflow* DataflowAsset = GetDataflowEditorContent()->DataflowAsset)
+		{
+			check(DataflowAsset->IsAsset());
+			OutObjects.Add(DataflowAsset);
+		}
 	}
 }
 
@@ -319,13 +323,15 @@ void FDataflowEditorToolkit::CreateWidgets()
 {
 	FBaseCharacterFXEditorToolkit::CreateWidgets();
 
-	const FDataflowEditorDatas& DataflowEditorDatas = GetDataflowEditorDatas();
-	if (UDataflow* DataflowAsset = DataflowEditorDatas.DataflowAsset)
+	if (TObjectPtr<UDataflowEditorContent> Content = GetDataflowEditorContent())
 	{
-		NodeDetailsEditor = CreateNodeDetailsEditorWidget(DataflowEditorDatas.DataflowOwner);
-		AssetDetailsEditor = CreateAssetDetailsEditorWidget(DataflowEditorDatas.DataflowOwner);
-		GraphEditor = CreateGraphEditorWidget(DataflowAsset, NodeDetailsEditor);
-		SkeletalEditor = CreateSkeletalEditorWidget(DataflowEditorDatas.DataflowOwner);
+		if (UDataflow* DataflowAsset = Content->DataflowAsset)
+		{
+			NodeDetailsEditor = CreateNodeDetailsEditorWidget(Content->DataflowOwner);
+			AssetDetailsEditor = CreateAssetDetailsEditorWidget(Content->DataflowOwner);
+			GraphEditor = CreateGraphEditorWidget(DataflowAsset, NodeDetailsEditor);
+			SkeletalEditor = CreateSkeletalEditorWidget(Content->DataflowOwner);
+		}
 	}
 }
 
@@ -359,14 +365,19 @@ TSharedPtr<FEditorViewportClient> FDataflowEditorToolkit::CreateEditorViewportCl
 
 void FDataflowEditorToolkit::OnPropertyValueChanged(const FPropertyChangedEvent& PropertyChangedEvent)
 {
-	FDataflowEditorDatas& DataflowEditorDatas = ModifyDataflowEditorDatas();
-	FDataflowEditorCommands::OnPropertyValueChanged(DataflowEditorDatas.DataflowAsset, DataflowEditorDatas.DataflowContext, DataflowEditorDatas.LastNodeTimestamp, PropertyChangedEvent, PrevNodeSelection);
+	if (TObjectPtr<UDataflowEditorContent> EditorContent = GetDataflowEditorContent())
+	{
+		ensure(EditorContent);
+		FDataflowEditorCommands::OnPropertyValueChanged(EditorContent->DataflowAsset, EditorContent->DataflowContext, EditorContent->LastNodeTimestamp, PropertyChangedEvent, PrevNodeSelection);
+	}
 }
 
 void FDataflowEditorToolkit::OnAssetPropertyValueChanged(const FPropertyChangedEvent& PropertyChangedEvent)
 {
-	FDataflowEditorDatas& DataflowEditorDatas = ModifyDataflowEditorDatas();
-	FDataflowEditorCommands::OnAssetPropertyValueChanged(DataflowEditorDatas.DataflowAsset, DataflowEditorDatas.DataflowContext, DataflowEditorDatas.LastNodeTimestamp, PropertyChangedEvent);
+	if (TObjectPtr<UDataflowEditorContent> EditorContent = GetDataflowEditorContent())
+	{
+		ensure(EditorContent);	FDataflowEditorCommands::OnAssetPropertyValueChanged(EditorContent->DataflowAsset, EditorContent->DataflowContext, EditorContent->LastNodeTimestamp, PropertyChangedEvent);
+	}
 }
 
 bool FDataflowEditorToolkit::OnNodeVerifyTitleCommit(const FText& NewText, UEdGraphNode* GraphNode, FText& OutErrorMessage)
@@ -381,7 +392,7 @@ void FDataflowEditorToolkit::OnNodeTitleCommitted(const FText& InNewText, ETextC
 
 void FDataflowEditorToolkit::OnNodeSelectionChanged(const TSet<UObject*>& NewSelection)
 {
-	if(GetDataflowEditorDatas().DataflowAsset)
+	if (TObjectPtr<UDataflowEditorContent> EditorContent = GetDataflowEditorContent(); EditorContent->DataflowAsset)
 	{
 		// Only keep UDataflowEdNode from NewSelection
 		TSet<UObject*> ValidatedSelection;
@@ -431,7 +442,7 @@ void FDataflowEditorToolkit::OnNodeSelectionChanged(const TSet<UObject*>& NewSel
 				Listener->OnSelectedNodeChanged(nullptr);
 			}
 		}
-		
+
 		PrevNodeSelection = ValidatedSelection;
 	}
 }
@@ -449,16 +460,18 @@ void FDataflowEditorToolkit::OnNodeDeleted(const TSet<UObject*>& NewSelection)
 
 void FDataflowEditorToolkit::Tick(float DeltaTime)
 {
-	FDataflowEditorDatas& DataflowEditorDatas = ModifyDataflowEditorDatas();
-	if(DataflowEditorDatas.IsValid())
+	if (TObjectPtr<UDataflowEditorContent> EditorContent = GetDataflowEditorContent())
 	{
-		if (!DataflowEditorDatas.DataflowContext)
+		if (EditorContent->IsValid())
 		{
-			DataflowEditorDatas.DataflowContext = MakeShared<Dataflow::FEngineContext>(DataflowEditorDatas.DataflowOwner, DataflowEditorDatas.DataflowAsset, Dataflow::FTimestamp::Invalid);
-			DataflowEditorDatas.LastNodeTimestamp = Dataflow::FTimestamp::Invalid;
+			if (!EditorContent->DataflowContext)
+			{
+				EditorContent->DataflowContext = MakeShared<Dataflow::FEngineContext>(EditorContent->DataflowOwner, EditorContent->DataflowAsset, Dataflow::FTimestamp::Invalid);
+				EditorContent->LastNodeTimestamp = Dataflow::FTimestamp::Invalid;
+			}
+			FDataflowEditorCommands::EvaluateTerminalNode(*EditorContent->DataflowContext.Get(), EditorContent->LastNodeTimestamp,
+														  EditorContent->DataflowAsset, nullptr, nullptr, EditorContent->DataflowOwner, EditorContent->DataflowTerminal);
 		}
-		FDataflowEditorCommands::EvaluateTerminalNode(*DataflowEditorDatas.DataflowContext.Get(), DataflowEditorDatas.LastNodeTimestamp,
-			DataflowEditorDatas.DataflowAsset, nullptr, nullptr, DataflowEditorDatas.DataflowOwner, DataflowEditorDatas.DataflowTerminal);
 	}
 }
 
@@ -474,18 +487,20 @@ TSharedRef<SDataflowGraphEditor> FDataflowEditorToolkit::CreateGraphEditorWidget
 
 	FDataflowEditorCommands::FGraphEvaluationCallback Evaluate = [&](FDataflowNode* Node, FDataflowOutput* Out)
 	{
-		FDataflowEditorDatas& DataflowEditorDatas = ModifyDataflowEditorDatas();
-		if(DataflowEditorDatas.IsValid())
+		if (TObjectPtr<UDataflowEditorContent> EditorContent = GetDataflowEditorContent())
 		{
-			if (!DataflowEditorDatas.DataflowContext)
+			if (EditorContent->IsValid())
 			{
-				DataflowEditorDatas.DataflowContext = MakeShared<Dataflow::FEngineContext>(DataflowEditorDatas.DataflowOwner, DataflowEditorDatas.DataflowAsset, Dataflow::FTimestamp::Invalid);
+				if (!EditorContent->DataflowContext)
+				{
+					EditorContent->DataflowContext = MakeShared<Dataflow::FEngineContext>(EditorContent->DataflowOwner, EditorContent->DataflowAsset, Dataflow::FTimestamp::Invalid);
+				}
+				Node->Invalidate();
+				EditorContent->LastNodeTimestamp = Dataflow::FTimestamp::Invalid;
+
+				FDataflowEditorCommands::EvaluateTerminalNode(*EditorContent->DataflowContext.Get(), EditorContent->LastNodeTimestamp,
+															  EditorContent->DataflowAsset, Node, Out, EditorContent->DataflowOwner, EditorContent->DataflowTerminal);
 			}
-			Node->Invalidate();
-			DataflowEditorDatas.LastNodeTimestamp = Dataflow::FTimestamp::Invalid;
-			
-			FDataflowEditorCommands::EvaluateTerminalNode(*DataflowEditorDatas.DataflowContext.Get(), DataflowEditorDatas.LastNodeTimestamp,
-				DataflowEditorDatas.DataflowAsset, Node, Out, DataflowEditorDatas.DataflowOwner, DataflowEditorDatas.DataflowTerminal);
 		}
 	};
 
@@ -562,13 +577,15 @@ TSharedPtr<IDetailsView> FDataflowEditorToolkit::CreateAssetDetailsEditorWidget(
 
 TSharedPtr<ISkeletonTree> FDataflowEditorToolkit::CreateSkeletalEditorWidget(UObject* ObjectToEdit)
 {
-	const FDataflowEditorDatas& DataflowEditorDatas = GetDataflowEditorDatas();
-	if(DataflowEditorDatas.DataflowAsset)
+	if (TObjectPtr<UDataflowEditorContent> EditorContent = GetDataflowEditorContent())
 	{
-		FSkeletonTreeArgs SkeletonTreeArgs;
-		ISkeletonEditorModule& SkeletonEditorModule = FModuleManager::LoadModuleChecked<ISkeletonEditorModule>("SkeletonEditor");
-		TSharedPtr<ISkeletonTree> SkeletonTree = SkeletonEditorModule.CreateSkeletonTree(DataflowEditorDatas.Skeleton, SkeletonTreeArgs);
-		return SkeletonTree;
+		if (EditorContent->DataflowAsset)
+		{
+			FSkeletonTreeArgs SkeletonTreeArgs;
+			ISkeletonEditorModule& SkeletonEditorModule = FModuleManager::LoadModuleChecked<ISkeletonEditorModule>("SkeletonEditor");
+			TSharedPtr<ISkeletonTree> SkeletonTree = SkeletonEditorModule.CreateSkeletonTree(EditorContent->Skeleton, SkeletonTreeArgs);
+			return SkeletonTree;
+		}
 	}
 	return TSharedPtr<ISkeletonTree>(nullptr);
 }
@@ -612,8 +629,10 @@ TSharedRef<SDockTab> FDataflowEditorToolkit::SpawnTab_NodeDetails(const FSpawnTa
 TSharedRef<SDockTab> FDataflowEditorToolkit::SpawnTab_Skeletal(const FSpawnTabArgs& Args)
 {
 	check(Args.GetTabId() == SkeletalTabId);
-	const FDataflowEditorDatas& DataflowEditorDatas = GetDataflowEditorDatas();
-	SkeletalEditor->SetSkeletalMesh(DataflowEditorDatas.SkeletalMesh);
+	if (const TObjectPtr<UDataflowEditorContent> EditorContent = GetDataflowEditorContent())
+	{
+		SkeletalEditor->SetSkeletalMesh(EditorContent->SkeletalMesh);
+	}
 
 	return SNew(SDockTab)
 		.Label(LOCTEXT("DataflowEditor_Skeletal_TabTitle", "Skeletal Hierarchy"))
@@ -624,7 +643,7 @@ TSharedRef<SDockTab> FDataflowEditorToolkit::SpawnTab_Skeletal(const FSpawnTabAr
 
 TSharedRef<SDockTab> FDataflowEditorToolkit::SpawnTab_SelectionView(const FSpawnTabArgs& Args)
 {
-//	check(Args.GetTabId().TabType == SelectionViewTabId_1);
+	//	check(Args.GetTabId().TabType == SelectionViewTabId_1);
 
 	if (Args.GetTabId() == SelectionViewTabId_1)
 	{
@@ -662,51 +681,53 @@ TSharedRef<SDockTab> FDataflowEditorToolkit::SpawnTab_SelectionView(const FSpawn
 	TSharedPtr<SSelectionViewWidget> SelectionViewWidget;
 
 	TSharedRef<SDockTab> DockableTab = SNew(SDockTab)
-	[
-		SAssignNew(SelectionViewWidget, SSelectionViewWidget)
-	];
+		[
+			SAssignNew(SelectionViewWidget, SSelectionViewWidget)
+		];
 
 	if (SelectionViewWidget)
 	{
-		const FDataflowEditorDatas& DataflowEditorDatas = GetDataflowEditorDatas();
-		if (Args.GetTabId() == SelectionViewTabId_1)
+		if (const TObjectPtr<UDataflowEditorContent> EditorContent = GetDataflowEditorContent())
 		{
-			DataflowSelectionView_1->SetSelectionView(SelectionViewWidget);
-
-			// Set the Context on the interface
-			if (TSharedPtr<Dataflow::FContext> CurrentContext = DataflowEditorDatas.DataflowContext)
+			if (Args.GetTabId() == SelectionViewTabId_1)
 			{
-				DataflowSelectionView_1->SetContext(CurrentContext);
-			}
-		}
-		else if (Args.GetTabId() == SelectionViewTabId_2)
-		{
-			DataflowSelectionView_2->SetSelectionView(SelectionViewWidget);
+				DataflowSelectionView_1->SetSelectionView(SelectionViewWidget);
 
-			// Set the Context on the interface
-			if (TSharedPtr<Dataflow::FContext> CurrentContext = DataflowEditorDatas.DataflowContext)
-			{
-				DataflowSelectionView_2->SetContext(CurrentContext);
+				// Set the Context on the interface
+				if (TSharedPtr<Dataflow::FContext> CurrentContext = EditorContent->DataflowContext)
+				{
+					DataflowSelectionView_1->SetContext(CurrentContext);
+				}
 			}
-		}
-		else if (Args.GetTabId() == SelectionViewTabId_3)
-		{
-			DataflowSelectionView_3->SetSelectionView(SelectionViewWidget);
-
-			// Set the Context on the interface
-			if (TSharedPtr<Dataflow::FContext> CurrentContext = DataflowEditorDatas.DataflowContext)
+			else if (Args.GetTabId() == SelectionViewTabId_2)
 			{
-				DataflowSelectionView_3->SetContext(CurrentContext);
+				DataflowSelectionView_2->SetSelectionView(SelectionViewWidget);
+
+				// Set the Context on the interface
+				if (TSharedPtr<Dataflow::FContext> CurrentContext = EditorContent->DataflowContext)
+				{
+					DataflowSelectionView_2->SetContext(CurrentContext);
+				}
 			}
-		}
-		else if (Args.GetTabId() == SelectionViewTabId_4)
-		{
-			DataflowSelectionView_4->SetSelectionView(SelectionViewWidget);
+			else if (Args.GetTabId() == SelectionViewTabId_3)
+			{
+				DataflowSelectionView_3->SetSelectionView(SelectionViewWidget);
 
-			// Set the Context on the interface
-			if (TSharedPtr<Dataflow::FContext> CurrentContext = DataflowEditorDatas.DataflowContext)
+				// Set the Context on the interface
+				if (TSharedPtr<Dataflow::FContext> CurrentContext = EditorContent->DataflowContext)
+				{
+					DataflowSelectionView_3->SetContext(CurrentContext);
+				}
+			}
+			else if (Args.GetTabId() == SelectionViewTabId_4)
 			{
-				DataflowSelectionView_4->SetContext(CurrentContext);
+				DataflowSelectionView_4->SetSelectionView(SelectionViewWidget);
+
+				// Set the Context on the interface
+				if (TSharedPtr<Dataflow::FContext> CurrentContext = EditorContent->DataflowContext)
+				{
+					DataflowSelectionView_4->SetContext(CurrentContext);
+				}
 			}
 		}
 	}
@@ -760,45 +781,47 @@ TSharedRef<SDockTab> FDataflowEditorToolkit::SpawnTab_CollectionSpreadSheet(cons
 
 	if (CollectionSpreadSheetWidget)
 	{
-		const FDataflowEditorDatas& DataflowEditorDatas = GetDataflowEditorDatas();
-		if (Args.GetTabId() == CollectionSpreadSheetTabId_1)
+		if (const TObjectPtr<UDataflowEditorContent> EditorContent = GetDataflowEditorContent())
 		{
-			DataflowCollectionSpreadSheet_1->SetCollectionSpreadSheet(CollectionSpreadSheetWidget);
-
-			// Set the Context on the interface
-			if (TSharedPtr<Dataflow::FContext> CurrentContext = DataflowEditorDatas.DataflowContext)
+			if (Args.GetTabId() == CollectionSpreadSheetTabId_1)
 			{
-				DataflowCollectionSpreadSheet_1->SetContext(CurrentContext);
-			}
-		}
-		else if (Args.GetTabId() == CollectionSpreadSheetTabId_2)
-		{
-			DataflowCollectionSpreadSheet_2->SetCollectionSpreadSheet(CollectionSpreadSheetWidget);
+				DataflowCollectionSpreadSheet_1->SetCollectionSpreadSheet(CollectionSpreadSheetWidget);
 
-			// Set the Context on the interface
-			if (TSharedPtr<Dataflow::FContext> CurrentContext = DataflowEditorDatas.DataflowContext)
-			{
-				DataflowCollectionSpreadSheet_2->SetContext(CurrentContext);
+				// Set the Context on the interface
+				if (TSharedPtr<Dataflow::FContext> CurrentContext = EditorContent->DataflowContext)
+				{
+					DataflowCollectionSpreadSheet_1->SetContext(CurrentContext);
+				}
 			}
-		}
-		else if (Args.GetTabId() == CollectionSpreadSheetTabId_3)
-		{
-			DataflowCollectionSpreadSheet_3->SetCollectionSpreadSheet(CollectionSpreadSheetWidget);
-
-			// Set the Context on the interface
-			if (TSharedPtr<Dataflow::FContext> CurrentContext = DataflowEditorDatas.DataflowContext)
+			else if (Args.GetTabId() == CollectionSpreadSheetTabId_2)
 			{
-				DataflowCollectionSpreadSheet_3->SetContext(CurrentContext);
+				DataflowCollectionSpreadSheet_2->SetCollectionSpreadSheet(CollectionSpreadSheetWidget);
+
+				// Set the Context on the interface
+				if (TSharedPtr<Dataflow::FContext> CurrentContext = EditorContent->DataflowContext)
+				{
+					DataflowCollectionSpreadSheet_2->SetContext(CurrentContext);
+				}
 			}
-		}
-		else if (Args.GetTabId() == CollectionSpreadSheetTabId_4)
-		{
-			DataflowCollectionSpreadSheet_4->SetCollectionSpreadSheet(CollectionSpreadSheetWidget);
+			else if (Args.GetTabId() == CollectionSpreadSheetTabId_3)
+			{
+				DataflowCollectionSpreadSheet_3->SetCollectionSpreadSheet(CollectionSpreadSheetWidget);
 
-			// Set the Context on the interface
-			if (TSharedPtr<Dataflow::FContext> CurrentContext = DataflowEditorDatas.DataflowContext)
+				// Set the Context on the interface
+				if (TSharedPtr<Dataflow::FContext> CurrentContext = EditorContent->DataflowContext)
+				{
+					DataflowCollectionSpreadSheet_3->SetContext(CurrentContext);
+				}
+			}
+			else if (Args.GetTabId() == CollectionSpreadSheetTabId_4)
 			{
-				DataflowCollectionSpreadSheet_4->SetContext(CurrentContext);
+				DataflowCollectionSpreadSheet_4->SetCollectionSpreadSheet(CollectionSpreadSheetWidget);
+
+				// Set the Context on the interface
+				if (TSharedPtr<Dataflow::FContext> CurrentContext = EditorContent->DataflowContext)
+				{
+					DataflowCollectionSpreadSheet_4->SetContext(CurrentContext);
+				}
 			}
 		}
 	}
@@ -942,14 +965,16 @@ FName FDataflowEditorToolkit::GetToolkitFName() const
 
 FText FDataflowEditorToolkit::GetToolkitName() const
 {
-	const FDataflowEditorDatas& DataflowEditorDatas = GetDataflowEditorDatas();
-	if(DataflowEditorDatas.DataflowOwner)
-	{
-		return  GetLabelForObject(DataflowEditorDatas.DataflowOwner);
-	}
-	else if(DataflowEditorDatas.DataflowAsset)
+	if (TObjectPtr<const UDataflowEditorContent> EditorContent = GetDataflowEditorContent())
 	{
-		return  GetLabelForObject(DataflowEditorDatas.DataflowAsset);
+		if (EditorContent->DataflowOwner)
+		{
+			return  GetLabelForObject(EditorContent->DataflowOwner);
+		}
+		else if (EditorContent->DataflowAsset)
+		{
+			return  GetLabelForObject(EditorContent->DataflowAsset);
+		}
 	}
 	return  LOCTEXT("ToolkitName", "Empty Dataflow Editor");
 }
diff --git a/Engine/Plugins/Experimental/Dataflow/Source/DataflowEditor/Private/Dataflow/DataflowEditorViewportClient.cpp b/Engine/Plugins/Experimental/Dataflow/Source/DataflowEditor/Private/Dataflow/DataflowEditorViewportClient.cpp
--- a/Engine/Plugins/Experimental/Dataflow/Source/DataflowEditor/Private/Dataflow/DataflowEditorViewportClient.cpp
+++ b/Engine/Plugins/Experimental/Dataflow/Source/DataflowEditor/Private/Dataflow/DataflowEditorViewportClient.cpp
@@ -45,13 +45,13 @@ void FDataflowEditorViewportClient::SetSelectionMode(FDataflowSelectionState::EM
 	}
 
 	State.Mode = SelectionMode;
-	PreviewScene->ModifyDataflowComponent()->SetSelectionState(State);
+	PreviewScene->GetDataflowComponent()->SetSelectionState(State);
 
 	if (SelectionMode == FDataflowSelectionState::EMode::DSS_Dataflow_None)
 	{
 		if (!PreviewScene->GetDataflowComponent()->GetSelectionState().IsEmpty())
 		{
-			PreviewScene->ModifyDataflowComponent()->SetSelectionState(FDataflowSelectionState(SelectionMode));
+			PreviewScene->GetDataflowComponent()->SetSelectionState(FDataflowSelectionState(SelectionMode));
 		}
 	}
 }
@@ -60,20 +60,22 @@ bool FDataflowEditorViewportClient::CanSetSelectionMode(FDataflowSelectionState:
 	TSharedPtr<FDataflowEditorToolkit> DataflowEditorToolkit = DataflowEditorToolkitPtr.Pin();
 	if (DataflowEditorToolkitPtr.IsValid())
 	{
-		const FDataflowEditorDatas& DataflowEditorDatas = DataflowEditorToolkit->GetDataflowEditorDatas();
-		if (const UDataflow* Dataflow = DataflowEditorDatas.DataflowAsset)
+		if (TObjectPtr<UDataflowEditorContent> EditorContent = PreviewScene->GetDataflowEditorContent())
 		{
-			if (Dataflow->GetRenderTargets().Num())
+			if (const UDataflow* Dataflow = EditorContent->DataflowAsset)
 			{
-				if (InState == FDataflowSelectionState::EMode::DSS_Dataflow_Object)
+				if (Dataflow->GetRenderTargets().Num())
 				{
-					return true;
-				}
+					if (InState == FDataflowSelectionState::EMode::DSS_Dataflow_Object)
+					{
+						return true;
+					}
 
-				if (InState == FDataflowSelectionState::EMode::DSS_Dataflow_Vertex
-					&& !PreviewScene->GetDataflowComponent()->GetSelectionState().Nodes.IsEmpty())
-				{
-					return true;
+					if (InState == FDataflowSelectionState::EMode::DSS_Dataflow_Vertex
+						&& !PreviewScene->GetDataflowComponent()->GetSelectionState().Nodes.IsEmpty())
+					{
+						return true;
+					}
 				}
 			}
 		}
@@ -81,6 +83,7 @@ bool FDataflowEditorViewportClient::CanSetSelectionMode(FDataflowSelectionState:
 
 	return false;
 }
+
 bool FDataflowEditorViewportClient::IsSelectionModeActive(FDataflowSelectionState::EMode InState)
 {
 	return SelectionMode == InState;
@@ -232,7 +235,7 @@ void FDataflowEditorViewportClient::ProcessClick(FSceneView& View, HHitProxy* Hi
 
 		if (PreState != SelectionState)
 		{
-			PreviewScene->ModifyDataflowComponent()->SetSelectionState(SelectionState);
+			PreviewScene->GetDataflowComponent()->SetSelectionState(SelectionState);
 		}
 	}
 }
@@ -245,16 +248,18 @@ void FDataflowEditorViewportClient::Tick(float DeltaSeconds)
 
 	if (PreviewScene->GetDataflowComponent())
 	{
-		const FDataflowEditorDatas& DataflowEditorDatas = PreviewScene->GetDataflowDatas();
-		if (TSharedPtr<Dataflow::FContext> Context = DataflowEditorDatas.DataflowContext)
+		if (TObjectPtr<UDataflowEditorContent> EditorContent = PreviewScene->GetDataflowEditorContent())
 		{
-			if (const UDataflow* Dataflow = DataflowEditorDatas.DataflowAsset)
+			if (TSharedPtr<Dataflow::FContext> Context = EditorContent->DataflowContext)
 			{
-				const Dataflow::FTimestamp SystemTimestamp = LatestTimestamp(Dataflow, Context.Get());
-				if (SystemTimestamp >= LastModifiedTimestamp)
+				if (const UDataflow* Dataflow = EditorContent->DataflowAsset)
 				{
-					PreviewScene->UpdateDataflowComponent();
-					LastModifiedTimestamp = LatestTimestamp(Dataflow, Context.Get()).Value + 1;
+					const Dataflow::FTimestamp SystemTimestamp = LatestTimestamp(Dataflow, Context.Get());
+					if (SystemTimestamp >= LastModifiedTimestamp)
+					{
+						PreviewScene->UpdateDataflowComponent();
+						LastModifiedTimestamp = LatestTimestamp(Dataflow, Context.Get()).Value + 1;
+					}
 				}
 			}
 		}
diff --git a/Engine/Plugins/Experimental/Dataflow/Source/DataflowEditor/Private/Dataflow/DataflowPreviewScene.cpp b/Engine/Plugins/Experimental/Dataflow/Source/DataflowEditor/Private/Dataflow/DataflowPreviewScene.cpp
--- a/Engine/Plugins/Experimental/Dataflow/Source/DataflowEditor/Private/Dataflow/DataflowPreviewScene.cpp
+++ b/Engine/Plugins/Experimental/Dataflow/Source/DataflowEditor/Private/Dataflow/DataflowPreviewScene.cpp
@@ -6,18 +6,20 @@
 #include "AssetEditorModeManager.h"
 #include "Components/SkeletalMeshComponent.h"
 #include "Components/DynamicMeshComponent.h"
+#include "Dataflow/DataflowActor.h"
 #include "Dataflow/DataflowComponent.h"
 #include "Dataflow/DataflowEditor.h"
 #include "Elements/Framework/EngineElementsLibrary.h"
 
 #define LOCTEXT_NAMESPACE "FDataflowPreviewScene"
 
-FDataflowPreviewScene::FDataflowPreviewScene(FPreviewScene::ConstructionValues ConstructionValues, FDataflowEditorDatas& DataflowAssetDatas) :
-	FAdvancedPreviewScene(ConstructionValues), DataflowDatas(DataflowAssetDatas)
+FDataflowPreviewScene::FDataflowPreviewScene(FPreviewScene::ConstructionValues ConstructionValues,TObjectPtr<UDataflowEditorContent> InEditorContent) 
+	: FAdvancedPreviewScene(ConstructionValues), EditorContent(InEditorContent)
 {
+	check(EditorContent);
 	SkeletalMeshActor = GetWorld()->SpawnActor<AActor>(AActor::StaticClass());
 
-	if(DataflowDatas.SkeletalMesh && DataflowDatas.bHasValidSkeletalMesh)
+	if(EditorContent->SkeletalMesh && EditorContent->bHasValidSkeletalMesh)
 	{
 		SkeletalMeshComponent = NewObject<USkeletalMeshComponent>(SkeletalMeshActor);
 		SkeletalMeshComponent->SelectionOverrideDelegate = UPrimitiveComponent::FSelectionOverride::CreateRaw(this, &FDataflowPreviewScene::IsComponentSelected);
@@ -25,21 +27,15 @@ FDataflowPreviewScene::FDataflowPreviewScene(FPreviewScene::ConstructionValues C
 		UpdateSkeletalMeshComponent();
 	}
 	
-	DataflowActor = GetWorld()->SpawnActor<AActor>(AActor::StaticClass());
-
-	if(DataflowDatas.DataflowAsset)
-	{
-		DataflowComponent = NewObject<UDataflowComponent>(DataflowActor);
-		//DataflowComponent->SelectionOverrideDelegate = UPrimitiveComponent::FSelectionOverride::CreateRaw(this, &FDataflowPreviewScene::IsComponentSelected);
-		
-		// @todo(DynamicMeshRendering) : Enable Dynamic Mesh Rendering for dataflow terminals. Hide the dataflow
-		//DataflowComponent->SetVisibility(false);
-		UpdateDataflowComponent();
-	}
+	// @todo(DynamicMeshRendering) : Enable Dynamic Mesh Rendering for dataflow terminals. Hide the dataflow 
+	DataflowActor = Cast<ADataflowActor>(GetWorld()->SpawnActor<ADataflowActor>(ADataflowActor::StaticClass()));
+	DataflowComponent = DataflowActor->GetDataflowComponent();
+	//DataflowComponent->SelectionOverrideDelegate = UPrimitiveComponent::FSelectionOverride::CreateRaw(this, &FDataflowPreviewScene::IsComponentSelected);
+	//DataflowComponent->SetVisibility(false);
+	UpdateDataflowComponent();
 	
-	DynamicMeshActor = GetWorld()->SpawnActor<AActor>(AActor::StaticClass());
 
-	
+	DynamicMeshActor = GetWorld()->SpawnActor<AActor>(AActor::StaticClass());	
 	SkeletalMeshActor->RegisterAllComponents();
 	DataflowActor->RegisterAllComponents();
 	DynamicMeshActor->RegisterAllComponents();
@@ -70,6 +66,7 @@ FDataflowPreviewScene::~FDataflowPreviewScene()
 void FDataflowPreviewScene::AddReferencedObjects(FReferenceCollector& Collector)
 {
 	FAdvancedPreviewScene::AddReferencedObjects(Collector);
+	Collector.AddReferencedObject(EditorContent);
 
 	Collector.AddReferencedObject(DataflowComponent);
 	Collector.AddReferencedObject(SkeletalMeshComponent);
@@ -82,26 +79,29 @@ void FDataflowPreviewScene::AddReferencedObjects(FReferenceCollector& Collector)
 
 void FDataflowPreviewScene::UpdateDataflowComponent()
 {
-	DataflowComponent->ResetRenderTargets();
-	
-	DataflowComponent->SetDataflow(DataflowDatas.DataflowAsset);
-	DataflowComponent->SetContext(DataflowDatas.DataflowContext);
-	
-	for (const UDataflowEdNode* const Node : DataflowDatas.DataflowAsset->GetRenderTargets())
+	if (EditorContent->DataflowAsset)
 	{
-		DataflowComponent->AddRenderTarget(Node);
+		DataflowComponent->ResetRenderTargets();
+
+		DataflowComponent->SetDataflow(EditorContent->DataflowAsset);
+		DataflowComponent->SetContext(EditorContent->DataflowContext);
+
+		for (const UDataflowEdNode* const Node : EditorContent->DataflowAsset->GetRenderTargets())
+		{
+			DataflowComponent->AddRenderTarget(Node);
+		}
+		DataflowComponent->UpdateBounds();
 	}
-	DataflowComponent->UpdateBounds();
 }
 
 void FDataflowPreviewScene::UpdateSkeletalMeshComponent()
 {
-	SkeletalMeshComponent->SetSkeletalMeshAsset(DataflowDatas.SkeletalMesh);
+	SkeletalMeshComponent->SetSkeletalMeshAsset(EditorContent->SkeletalMesh);
 
-	if (DataflowDatas.AnimationAsset)
+	if (EditorContent->AnimationAsset)
 	{
 		PreviewAnimInstance = NewObject<UAnimSingleNodeInstance>(SkeletalMeshComponent);
-		PreviewAnimInstance->SetAnimationAsset(DataflowDatas.AnimationAsset);
+		PreviewAnimInstance->SetAnimationAsset(EditorContent->AnimationAsset);
 
 		SkeletalMeshComponent->SetAnimationMode(EAnimationMode::AnimationSingleNode);
 		SkeletalMeshComponent->InitAnim(true);
diff --git a/Engine/Plugins/Experimental/Dataflow/Source/DataflowEditor/Public/Dataflow/DataflowEditor.h b/Engine/Plugins/Experimental/Dataflow/Source/DataflowEditor/Public/Dataflow/DataflowEditor.h
--- a/Engine/Plugins/Experimental/Dataflow/Source/DataflowEditor/Public/Dataflow/DataflowEditor.h
+++ b/Engine/Plugins/Experimental/Dataflow/Source/DataflowEditor/Public/Dataflow/DataflowEditor.h
@@ -24,8 +24,12 @@ class UAnimationAsset;
 /** 
  * Dataflow datas that will be used within the editor classes to evaluate the graph
  */
-struct DATAFLOWEDITOR_API FDataflowEditorDatas
+UCLASS()
+class UDataflowEditorContent : public UObject
 {
+	GENERATED_BODY()
+
+public:
 	/** Check if the datas flow datas are valid */
 	bool IsValid() const { return DataflowOwner && DataflowAsset;}
 	
@@ -76,6 +80,7 @@ class DATAFLOWEDITOR_API UDataflowEditor : public UBaseCharacterFXEditor
 	GENERATED_BODY()
 
 public:
+	UDataflowEditor();
 
 	// UBaseCharacterFXEditor interface
 	virtual TSharedPtr<FBaseAssetToolkit> CreateToolkit() override;
@@ -89,7 +94,7 @@ private :
 	// is the one holding the dynamic mesh components to be rendered in the viewport
 	// It is why the data flow asset/owner/skelmesh have been added here. Could be added
 	// in the subsystem if necessary
-	FDataflowEditorDatas DataflowDatas;
+	TObjectPtr<UDataflowEditorContent> Content;
 };
 
 DECLARE_LOG_CATEGORY_EXTERN(LogDataflowEditor, Log, All);
\ No newline at end of file
diff --git a/Engine/Plugins/Experimental/Dataflow/Source/DataflowEditor/Public/Dataflow/DataflowEditorToolkit.h b/Engine/Plugins/Experimental/Dataflow/Source/DataflowEditor/Public/Dataflow/DataflowEditorToolkit.h
--- a/Engine/Plugins/Experimental/Dataflow/Source/DataflowEditor/Public/Dataflow/DataflowEditorToolkit.h
+++ b/Engine/Plugins/Experimental/Dataflow/Source/DataflowEditor/Public/Dataflow/DataflowEditorToolkit.h
@@ -22,7 +22,7 @@ class IToolkitHost;
 class UDataflow;
 class USkeletalMesh;
 class SDataflowGraphEditor;
-struct FDataflowEditorDatas;
+class UDataflowEditorContent;
 
 namespace Dataflow
 {
@@ -51,11 +51,10 @@ public:
 	static UDataflow* GetDataflowAsset(UObject* ObjectToEdit);
 	static const UDataflow* GetDataflowAsset(const UObject* ObjectToEdit);
 
-	/** Get the const dataflow editor datas */
-	const FDataflowEditorDatas& GetDataflowEditorDatas() const;
+	/** Datafloe Editor Content Access */
+	TObjectPtr<const UDataflowEditorContent> GetDataflowEditorContent() const;
+	TObjectPtr<UDataflowEditorContent> GetDataflowEditorContent();
 
-	/** Return the dataflow editor datas that could be modified */
-	FDataflowEditorDatas& ModifyDataflowEditorDatas();
 
 	// IToolkit interface
 	virtual FName GetToolkitFName() const override;
diff --git a/Engine/Plugins/Experimental/Dataflow/Source/DataflowEditor/Public/Dataflow/DataflowPreviewScene.h b/Engine/Plugins/Experimental/Dataflow/Source/DataflowEditor/Public/Dataflow/DataflowPreviewScene.h
--- a/Engine/Plugins/Experimental/Dataflow/Source/DataflowEditor/Public/Dataflow/DataflowPreviewScene.h
+++ b/Engine/Plugins/Experimental/Dataflow/Source/DataflowEditor/Public/Dataflow/DataflowPreviewScene.h
@@ -4,7 +4,7 @@
 
 #include "AdvancedPreviewScene.h"
 
-struct FDataflowEditorDatas;
+class UDataflowEditorContent;
 class USkeletalMeshComponent;
 class UDynamicMeshComponent;
 class UDataflowComponent;
@@ -12,6 +12,7 @@ class UAnimationAsset;
 class UAnimSingleNodeInstance;
 class UMaterialInterface;
 class FAssetEditorModeManager;
+class ADataflowActor;
 
 namespace UE
 {
@@ -30,30 +31,30 @@ class DATAFLOWEDITOR_API FDataflowPreviewScene : public FAdvancedPreviewScene
 {
 public:
 
-	FDataflowPreviewScene(FPreviewScene::ConstructionValues ConstructionValues, FDataflowEditorDatas& DataflowAssetDatas);
+	FDataflowPreviewScene(FPreviewScene::ConstructionValues ConstructionValues, TObjectPtr<UDataflowEditorContent> InEditorContent);
 	virtual ~FDataflowPreviewScene();
 
 	// FGCObject interface
 	virtual void AddReferencedObjects(FReferenceCollector& Collector) override;
 
 	// Modify/Get the preview anim instance
-	UAnimSingleNodeInstance* ModifyPreviewAnimInstance() { return PreviewAnimInstance; }
+	UAnimSingleNodeInstance* GetPreviewAnimInstance() { return PreviewAnimInstance; }
 	const UAnimSingleNodeInstance* const GetPreviewAnimInstance() const { return PreviewAnimInstance; }
 
 	// Modify/Get the data flow component
-	UDataflowComponent* ModifyDataflowComponent() { return DataflowComponent; }
+	UDataflowComponent* GetDataflowComponent() { return DataflowComponent; }
 	const UDataflowComponent* GetDataflowComponent() const { return DataflowComponent; }
 
 	// Modify/Get the skeletal mesh component
-	USkeletalMeshComponent* ModifySkeletalMeshComponent() { return SkeletalMeshComponent; }
+	USkeletalMeshComponent* GetSkeletalMeshComponent() { return SkeletalMeshComponent; }
 	const USkeletalMeshComponent* GetSkeletalMeshComponent() const { return SkeletalMeshComponent; }
 	
 	// Modify/Get the data flow datas
-	FDataflowEditorDatas& ModifyDataflowDatas() { return DataflowDatas; }
-	const FDataflowEditorDatas& GetDataflowDatas() const { return DataflowDatas; }
+	TObjectPtr<UDataflowEditorContent> GetDataflowEditorContent() { return EditorContent; }
+	const TObjectPtr<UDataflowEditorContent> GetDataflowEditorContent() const { return EditorContent; }
 
 	// Modify/Get the data flow mode manager
-	TSharedPtr<FAssetEditorModeManager>& ModifyDataflowModeManager() { return DataflowModeManager; }
+	TSharedPtr<FAssetEditorModeManager>& GetDataflowModeManager() { return DataflowModeManager; }
 	const TSharedPtr<FAssetEditorModeManager>& GetDataflowModeManager() const { return DataflowModeManager; }
 	
 	/** Reset all the dynamic mesh components */
@@ -77,10 +78,10 @@ private:
 	bool IsComponentSelected(const UPrimitiveComponent* InComponent) const;
 
 	/** Data flows assets datas used to build the scene components */
-	FDataflowEditorDatas& DataflowDatas;
+	TObjectPtr<UDataflowEditorContent> EditorContent = nullptr;
 	
 	/** Dataflow actor */
-	TObjectPtr<AActor> DataflowActor = nullptr;
+	TObjectPtr<ADataflowActor> DataflowActor = nullptr;
 	
 	/** Skeletal mesh actor */
 	TObjectPtr<AActor> SkeletalMeshActor = nullptr;
