Dataflow : Tools Context
- Wire the UDataflowEdtiorContent through to the tools so they may share the cached Context.

diff --git a/Engine/Plugins/Experimental/Dataflow/Source/DataflowEditor/Private/Dataflow/DataflowEditorMode.cpp b/Engine/Plugins/Experimental/Dataflow/Source/DataflowEditor/Private/Dataflow/DataflowEditorMode.cpp
--- a/Engine/Plugins/Experimental/Dataflow/Source/DataflowEditor/Private/Dataflow/DataflowEditorMode.cpp
+++ b/Engine/Plugins/Experimental/Dataflow/Source/DataflowEditor/Private/Dataflow/DataflowEditorMode.cpp
@@ -139,9 +139,9 @@ void UDataflowEditorMode::BindCommands()
 void UDataflowEditorMode::CreateToolTargets(const TArray<TObjectPtr<UObject>>& AssetsIn)
 {
 	ToolTargets.Reset();
-	for (UObject* Object : AssetsIn)
+	if (TObjectPtr<UDataflowEditorContent> EditorContent = PreviewScene->GetDataflowEditorContent())
 	{
-		if (UToolTarget* Target = GetInteractiveToolsContext()->TargetManager->BuildTarget(Object, GetToolTargetRequirements()))
+		if (UToolTarget* Target = GetInteractiveToolsContext()->TargetManager->BuildTarget(EditorContent.Get(), GetToolTargetRequirements()))
 		{
 			ToolTargets.Add(Target);
 		}
@@ -160,8 +160,6 @@ void UDataflowEditorMode::InitializeTargets(const TArray<TObjectPtr<UObject>>& A
 	PreviewScene->ResetDynamicMeshComponents();
 	for (UToolTarget* Target : ToolTargets)
 	{
-		UE::Geometry::FDynamicMesh3 DynamicMesh = UE::ToolTarget::GetDynamicMeshCopy(Target);
-
 		// @todo(DynamicMeshRendering) : Enable Dynamic Mesh Rendering for dataflow terminals. 
 		//  TObjectPtr<UDynamicMeshComponent>& DynamicMeshComponent =
 		//	PreviewScene->AddDynamicMeshComponent(MoveTemp(DynamicMesh), UE::ToolTarget::GetMaterialSet(Target).Materials);
diff --git a/Engine/Plugins/Experimental/Dataflow/Source/DataflowEditor/Private/Dataflow/DataflowToolTarget.cpp b/Engine/Plugins/Experimental/Dataflow/Source/DataflowEditor/Private/Dataflow/DataflowToolTarget.cpp
--- a/Engine/Plugins/Experimental/Dataflow/Source/DataflowEditor/Private/Dataflow/DataflowToolTarget.cpp
+++ b/Engine/Plugins/Experimental/Dataflow/Source/DataflowEditor/Private/Dataflow/DataflowToolTarget.cpp
@@ -3,6 +3,7 @@
 #include "Dataflow/DataflowToolTarget.h"
 
 #include "ConversionUtils/DynamicMeshViaMeshDescriptionUtil.h"
+#include "Dataflow/DataflowEditor.h"
 #include "Dataflow/DataflowEditorToolkit.h"
 #include "Dataflow/DataflowEdNode.h"
 #include "Dataflow/DataflowObject.h"
@@ -51,14 +52,12 @@ namespace UE
 		}
 		
 		// Convert a dataflow component to a dynamic mesh
-		void DataflowToDynamicMesh(UObject* Asset, UDataflow* Dataflow, FDynamicMesh3& DynamicMesh)
+		void DataflowToDynamicMesh(TSharedPtr<::Dataflow::FEngineContext> DataflowContext, UObject* Asset, UDataflow* Dataflow, FDynamicMesh3& DynamicMesh)
 		{
 			FManagedArrayCollection RenderCollection;
 			GeometryCollection::Facades::FRenderingFacade Facade(RenderCollection);
 			Facade.DefineSchema();
-			
-			const TSharedPtr<::Dataflow::FEngineContext> DataflowContext = 
-				MakeShared<::Dataflow::FEngineContext>(	Asset, Dataflow, ::Dataflow::FTimestamp::Invalid);
+		
 			
 			for (const UDataflowEdNode* Target : Dataflow->GetRenderTargets())
 			{
@@ -103,6 +102,24 @@ namespace UE
 	}
 }
 
+namespace Dataflow
+{
+	TSharedPtr<FEngineContext> GetContext(UDataflowEditorContent* Content)
+	{
+		if (Content)
+		{
+			if (!Content->DataflowContext)
+			{
+				Content->DataflowContext = MakeShared<FEngineContext>(Content->DataflowOwner, Content->DataflowAsset, FTimestamp::Invalid);
+			}
+			return Content->DataflowContext;
+		}
+
+		ensure(false);
+		return MakeShared<FEngineContext>(nullptr, nullptr, FTimestamp::Invalid);
+	}
+}
+
 //
 // UDataflowReadOnlyToolTarget
 //
@@ -184,7 +201,7 @@ FDynamicMesh3 UDataflowReadOnlyToolTarget::GetDynamicMesh()
 	FDynamicMesh3 DynamicMesh;
 	if(IsValid())
 	{
-		UE::Conversion::DataflowToDynamicMesh(Asset,Dataflow, DynamicMesh);
+		UE::Conversion::DataflowToDynamicMesh(Context, Asset, Dataflow, DynamicMesh);
 	}
 	return DynamicMesh;
 }
@@ -240,7 +257,10 @@ bool UDataflowReadOnlyToolTargetFactory::CanBuildTarget(UObject* SourceObject, c
 	// If you want to make the tool target work with some subclass of UDataflow,
 	// just add another factory that allows that class specifically(but make sure that
 	// GetMeshDescription and such work properly)
-	const UDataflow* Dataflow = FDataflowEditorToolkit::GetDataflowAsset(SourceObject);
+	UDataflowEditorContent* Content = CastChecked<UDataflowEditorContent>(SourceObject);
+	const UDataflow* Dataflow = Content->DataflowAsset;
+
+
 
 	return Dataflow &&
 		ExactCast<UDataflow>(Cast<UDataflow>(Dataflow)) &&
@@ -250,9 +270,12 @@ bool UDataflowReadOnlyToolTargetFactory::CanBuildTarget(UObject* SourceObject, c
 
 UToolTarget* UDataflowReadOnlyToolTargetFactory::BuildTarget(UObject* SourceObject, const FToolTargetTypeRequirements& Requirements)
 {
+	UDataflowEditorContent* Content = CastChecked<UDataflowEditorContent>(SourceObject);
+
 	UDataflowReadOnlyToolTarget* Target = NewObject<UDataflowReadOnlyToolTarget>();
-	Target->Asset = SourceObject;
-	Target->Dataflow = FDataflowEditorToolkit::GetDataflowAsset(SourceObject);
+	Target->Asset = Content->DataflowOwner;
+	Target->Dataflow = Content->DataflowAsset;
+	Target->Context = Dataflow::GetContext(Content);
 
 	// @todo(brice) : I needed to comment this out?
 	//checkSlow(Target->Component.IsValid() && Requirements.AreSatisfiedBy(Target));
@@ -268,7 +291,8 @@ bool UDataflowToolTargetFactory::CanBuildTarget(UObject* SourceObject, const FTo
 	// If you want to make the tool target work with some subclass of UDataflow,
 	// just add another factory that allows that class specifically(but make sure that
 	// GetMeshDescription and such work properly)
-	const UDataflow* Dataflow = FDataflowEditorToolkit::GetDataflowAsset(SourceObject);
+	UDataflowEditorContent* Content = CastChecked<UDataflowEditorContent>(SourceObject);
+	const UDataflow* Dataflow = Content->DataflowAsset;
 
 	return Dataflow &&
 		ExactCast<UDataflow>(Cast<UDataflow>(Dataflow)) &&
@@ -278,9 +302,11 @@ bool UDataflowToolTargetFactory::CanBuildTarget(UObject* SourceObject, const FTo
 
 UToolTarget* UDataflowToolTargetFactory::BuildTarget(UObject* SourceObject, const FToolTargetTypeRequirements& Requirements)
 {
+	UDataflowEditorContent* Content = CastChecked<UDataflowEditorContent>(SourceObject);
 	UDataflowToolTarget* Target = NewObject<UDataflowToolTarget>();
-	Target->Asset = SourceObject;
-	Target->Dataflow = FDataflowEditorToolkit::GetDataflowAsset(SourceObject);
+	Target->Asset = Content->DataflowOwner;
+	Target->Dataflow = Content->DataflowAsset;
+	Target->Context = Dataflow::GetContext(Content);
 	//checkSlow(Target->Component.IsValid() && Requirements.AreSatisfiedBy(Target));
 
 	return Target;
diff --git a/Engine/Plugins/Experimental/Dataflow/Source/DataflowEditor/Public/Dataflow/DataflowToolTarget.h b/Engine/Plugins/Experimental/Dataflow/Source/DataflowEditor/Public/Dataflow/DataflowToolTarget.h
--- a/Engine/Plugins/Experimental/Dataflow/Source/DataflowEditor/Public/Dataflow/DataflowToolTarget.h
+++ b/Engine/Plugins/Experimental/Dataflow/Source/DataflowEditor/Public/Dataflow/DataflowToolTarget.h
@@ -3,7 +3,7 @@
 #pragma once
 
 #include "CoreMinimal.h"
-
+#include "Dataflow/DataflowObjectInterface.h"
 #include "TargetInterfaces/DynamicMeshCommitter.h"
 #include "TargetInterfaces/DynamicMeshProvider.h"
 #include "TargetInterfaces/MaterialProvider.h"
@@ -60,6 +60,9 @@ protected:
 
 	// Internal asset pointer
 	TObjectPtr<UObject> Asset = nullptr;
+
+	/**  Engine context to be used for dataflow evaluation */
+	TSharedPtr<Dataflow::FEngineContext> Context = nullptr;
 };
 
 /**
