Dataflow : DynamicMeshRendering
- Updated to render using the DynamicMesh3 structures.

diff --git a/Engine/Plugins/Experimental/Dataflow/Source/DataflowEditor/Private/Dataflow/DataflowEditor.cpp b/Engine/Plugins/Experimental/Dataflow/Source/DataflowEditor/Private/Dataflow/DataflowEditor.cpp
--- a/Engine/Plugins/Experimental/Dataflow/Source/DataflowEditor/Private/Dataflow/DataflowEditor.cpp
+++ b/Engine/Plugins/Experimental/Dataflow/Source/DataflowEditor/Private/Dataflow/DataflowEditor.cpp
@@ -3,63 +3,16 @@
 #include "Dataflow/DataflowEditor.h"
 
 #include "Animation/Skeleton.h"
+#include "Dataflow/DataflowEditorContent.h"
 #include "Dataflow/DataflowEditorToolkit.h"
+#include "Dataflow/DataflowEditorUtil.h"
 #include "Engine/SkeletalMesh.h"
 
 #include UE_INLINE_GENERATED_CPP_BY_NAME(DataflowEditor)
 
 DEFINE_LOG_CATEGORY(LogDataflowEditor);
 
-namespace Private
-{
-	UDataflow* GetDataflowAssetFrom(UObject* InObject)
-	{
-		if (UClass* Class = InObject->GetClass())
-		{
-			if (FProperty* Property = Class->FindPropertyByName(FName("DataflowAsset")))
-			{
-				return *Property->ContainerPtrToValuePtr<UDataflow*>(InObject);
-			}
-		}
-		return nullptr;
-	}
-
-	USkeletalMesh* GetSkeletalMeshFrom(UObject* InObject)
-	{
-		if (UClass* Class = InObject->GetClass())
-		{
-			if (FProperty* Property = Class->FindPropertyByName(FName("SkeletalMesh")))
-			{
-				return *Property->ContainerPtrToValuePtr<USkeletalMesh*>(InObject);
-			}
-		}
-		return nullptr;
-	}
-	
-	UAnimationAsset* GetAnimationAssetFrom(UObject* InObject)
-	{
-		if (UClass* Class = InObject->GetClass())
-		{
-			if (FProperty* Property = Class->FindPropertyByName(FName("AnimationAsset")))
-			{
-				return *Property->ContainerPtrToValuePtr<UAnimationAsset*>(InObject);
-			}
-		}
-		return nullptr;
-	}
 
-	FString GetDataflowTerminalFrom(UObject* InObject)
-	{
-		if (UClass* Class = InObject->GetClass())
-		{
-			if (FProperty* Property = Class->FindPropertyByName(FName("DataflowTerminal")))
-			{
-				return *Property->ContainerPtrToValuePtr<FString>(InObject);
-			}
-		}
-		return FString();
-	}
-};
 
 UDataflowEditor::UDataflowEditor() : Super()
 {
@@ -94,7 +47,7 @@ void UDataflowEditor::Initialize(const TArray<TObjectPtr<UObject>>& InObjects)
 			Content->DataflowOwner = RootObject;
 			Content->DataflowAsset->Schema = UDataflowSchema::StaticClass();
 			Content->DataflowContext = MakeShared<Dataflow::FEngineContext>(RootObject, Content->DataflowAsset, FPlatformTime::Cycles64());
-			Content->LastNodeTimestamp = Content->DataflowContext->GetTimestamp();
+			Content->LastModifiedTimestamp = Content->DataflowContext->GetTimestamp();
 
 			ObjectsToEdit.Add(Content->DataflowOwner);
 
diff --git a/Engine/Plugins/Experimental/Dataflow/Source/DataflowEditor/Private/Dataflow/DataflowEditorContent.cpp b/Engine/Plugins/Experimental/Dataflow/Source/DataflowEditor/Private/Dataflow/DataflowEditorContent.cpp
new file mode 100644
--- /dev/null
+++ b/Engine/Plugins/Experimental/Dataflow/Source/DataflowEditor/Private/Dataflow/DataflowEditorContent.cpp
@@ -0,0 +1,12 @@
+// Copyright Epic Games, Inc. All Rights Reserved.
+
+#include "Dataflow/DataflowEditorContent.h"
+
+#include "Interfaces/IPluginManager.h"
+
+#include UE_INLINE_GENERATED_CPP_BY_NAME(DataflowEditorContent)
+
+
+UDataflowEditorContent::UDataflowEditorContent()
+{
+}
\ No newline at end of file
diff --git a/Engine/Plugins/Experimental/Dataflow/Source/DataflowEditor/Private/Dataflow/DataflowEditorContent.h b/Engine/Plugins/Experimental/Dataflow/Source/DataflowEditor/Private/Dataflow/DataflowEditorContent.h
new file mode 100644
--- /dev/null
+++ b/Engine/Plugins/Experimental/Dataflow/Source/DataflowEditor/Private/Dataflow/DataflowEditorContent.h
@@ -0,0 +1,59 @@
+// Copyright Epic Games, Inc. All Rights Reserved.
+
+#pragma once
+
+#include "Dataflow/DataflowNodeParameters.h"
+#include "Dataflow/DataflowObjectInterface.h"
+#include "DynamicMesh/DynamicMesh3.h"
+#include "GeometryCollection/ManagedArrayCollection.h"
+#include "Templates/SharedPointer.h"
+
+#include "DataflowEditorContent.generated.h"
+
+class FDataflowEditorToolkit;
+class UDataflow;
+class USkeletalMesh;
+class USkeleton;
+class UAnimationAsset;
+
+/** 
+ * Dataflow datas that will be used within the editor classes to evaluate the graph
+ */
+UCLASS()
+class UDataflowEditorContent : public UObject
+{
+	GENERATED_BODY()
+
+public:
+	UDataflowEditorContent();
+
+	/** Check if the datas flow datas are valid */
+	bool IsValid() const { return DataflowOwner && DataflowAsset;}
+
+	/** Data flow object owner */
+	TObjectPtr<UObject> DataflowOwner = nullptr;
+
+	/** Data flow asset that we will edit */
+	TObjectPtr<UDataflow> DataflowAsset = nullptr;
+
+	/** Data flow terminal path for evaluation */
+	FString DataflowTerminal = "";
+
+	/** Data flow skeletal mesh*/
+	TObjectPtr<USkeletalMesh> SkeletalMesh = nullptr;
+
+	/** Data flow skeleton*/
+	TObjectPtr<USkeleton> Skeleton = nullptr;
+
+	/** Animation asset to be used to preview simulation */
+	TObjectPtr<UAnimationAsset> AnimationAsset;
+
+	/**  Engine context to be used for dataflow evaluation */
+	TSharedPtr<Dataflow::FEngineContext> DataflowContext = nullptr;
+
+	/** Last data flow evaluated node time stamp */
+	Dataflow::FTimestamp LastModifiedTimestamp = Dataflow::FTimestamp::Invalid;
+
+	/** Boolean to check if the skelmesh is valid*/
+	bool bHasValidSkeletalMesh = false;
+};
diff --git a/Engine/Plugins/Experimental/Dataflow/Source/DataflowEditor/Private/Dataflow/DataflowEditorMode.cpp b/Engine/Plugins/Experimental/Dataflow/Source/DataflowEditor/Private/Dataflow/DataflowEditorMode.cpp
--- a/Engine/Plugins/Experimental/Dataflow/Source/DataflowEditor/Private/Dataflow/DataflowEditorMode.cpp
+++ b/Engine/Plugins/Experimental/Dataflow/Source/DataflowEditor/Private/Dataflow/DataflowEditorMode.cpp
@@ -1,18 +1,27 @@
 // Copyright Epic Games, Inc. All Rights Reserved.
 
 #include "Dataflow/DataflowEditorMode.h"
-#include "Dataflow/DataflowEditorModeToolkit.h"
+
+#include "AttributeEditorTool.h"
+#include "Components/DynamicMeshComponent.h"
 #include "Dataflow/DataflowComponentToolTarget.h"
-#include "Dataflow/DataflowToolTarget.h"
-#include "Dataflow/DataflowGraphEditor.h"
-#include "Dataflow/DataflowSNode.h"
+#include "Dataflow/DataflowEditor.h"
+#include "Dataflow/DataflowEditorContent.h"
 #include "Dataflow/DataflowEditorCommands.h"
+#include "Dataflow/DataflowEditorModeToolkit.h"
+#include "Dataflow/DataflowGraphEditor.h"
 #include "Dataflow/DataflowPreviewScene.h"
-#include "Dataflow/DataflowEditor.h"
+#include "Dataflow/DataflowSNode.h"
+#include "Dataflow/DataflowToolTarget.h"
+#include "EditorModeManager.h"
+#include "EdModeInteractiveToolsContext.h"
+#include "MeshSelectionTool.h"
+#include "MeshVertexPaintTool.h"
+#include "MeshAttributePaintTool.h"
+#include "ModelingToolTargetUtil.h"
 #include "TargetInterfaces/MaterialProvider.h"
 #include "TargetInterfaces/DynamicMeshCommitter.h"
 #include "TargetInterfaces/DynamicMeshProvider.h"
-#include "EdModeInteractiveToolsContext.h"
 #include "ToolTargetManager.h"
 #include "ToolTargets/DynamicMeshComponentToolTarget.h"
 #include "ToolTargets/StaticMeshComponentToolTarget.h"
@@ -20,13 +29,6 @@
 #include "ToolTargets/SkeletalMeshComponentToolTarget.h"
 #include "ToolTargets/SkeletalMeshToolTarget.h"
 #include "Tools/UEdMode.h"
-#include "AttributeEditorTool.h"
-#include "MeshSelectionTool.h"
-#include "MeshVertexPaintTool.h"
-#include "MeshAttributePaintTool.h"
-#include "ModelingToolTargetUtil.h"
-#include "Components/DynamicMeshComponent.h"
-#include "EditorModeManager.h"
 #include "Selection.h"
 
 #include UE_INLINE_GENERATED_CPP_BY_NAME(DataflowEditorMode)
@@ -157,14 +159,19 @@ void UDataflowEditorMode::InitializeTargets(const TArray<TObjectPtr<UObject>>& A
 	SelectedComponents->BeginBatchSelectOperation();
 	SelectedComponents->DeselectAll();
 
-	PreviewScene->ResetDynamicMeshComponents();
-	for (UToolTarget* Target : ToolTargets)
+	// @todo(brice) : What are the ToolTargets storing?
+	// ... for(ToolTarget& : ToolTargets){
+	// ... UE::ToolTarget::GetDynamicMeshCopy(Target)
+	// ... UE::ToolTarget::GetMaterialSet(Target).Materials for PreviewScene->AddDynamicMeshComponent
+	// ... }
+	PreviewScene->Update();
+
+	for (TObjectPtr<UDynamicMeshComponent>& Component : PreviewScene->GetDynamicMeshComponents())
 	{
-		// @todo(DynamicMeshRendering) : Enable Dynamic Mesh Rendering for dataflow terminals. 
-		//  TObjectPtr<UDynamicMeshComponent>& DynamicMeshComponent =
-		//	PreviewScene->AddDynamicMeshComponent(MoveTemp(DynamicMesh), UE::ToolTarget::GetMaterialSet(Target).Materials);
-		//
-		//SelectedComponents->Select(DynamicMeshComponent);
+		if (Component)
+		{
+			SelectedComponents->Select(Component);
+		}
 	}
 	SelectedComponents->EndBatchSelectOperation();
 }
diff --git a/Engine/Plugins/Experimental/Dataflow/Source/DataflowEditor/Private/Dataflow/DataflowEditorModeToolkit.cpp b/Engine/Plugins/Experimental/Dataflow/Source/DataflowEditor/Private/Dataflow/DataflowEditorModeToolkit.cpp
--- a/Engine/Plugins/Experimental/Dataflow/Source/DataflowEditor/Private/Dataflow/DataflowEditorModeToolkit.cpp
+++ b/Engine/Plugins/Experimental/Dataflow/Source/DataflowEditor/Private/Dataflow/DataflowEditorModeToolkit.cpp
@@ -20,9 +20,8 @@ void FDataflowEditorModeToolkit::BuildToolPalette(FName PaletteIndex, class FToo
 	const FDataflowEditorCommandsImpl& Commands = FDataflowEditorCommands::Get();
 	if (PaletteIndex == ToolsTabName)
 	{
-		// @todo(DynamicMeshRendering) : Enable Dynamic Mesh Rendering for dataflow terminals. 
-		//ToolbarBuilder.AddToolBarButton(Commands.BeginAttributeEditorTool);
-		//ToolbarBuilder.AddToolBarButton(Commands.BeginMeshSelectionTool);
+		ToolbarBuilder.AddToolBarButton(Commands.BeginAttributeEditorTool);
+		ToolbarBuilder.AddToolBarButton(Commands.BeginMeshSelectionTool);
 	}
 }
 
diff --git a/Engine/Plugins/Experimental/Dataflow/Source/DataflowEditor/Private/Dataflow/DataflowEditorStyle.h b/Engine/Plugins/Experimental/Dataflow/Source/DataflowEditor/Private/Dataflow/DataflowEditorStyle.h
--- a/Engine/Plugins/Experimental/Dataflow/Source/DataflowEditor/Private/Dataflow/DataflowEditorStyle.h
+++ b/Engine/Plugins/Experimental/Dataflow/Source/DataflowEditor/Private/Dataflow/DataflowEditorStyle.h
@@ -5,8 +5,10 @@
 #include "Styling/SlateStyleMacros.h"
 #include "Styling/SlateStyleRegistry.h"
 #include "Interfaces/IPluginManager.h"
+#include "Materials/Material.h"
 #include "Misc/Paths.h"
 #include "Dataflow/DataflowEditorCommands.h"
+#include "Dataflow/DataflowEditorUtil.h"
 
 class FDataflowEditorStyle final : public FSlateStyleSet
 {
@@ -40,6 +42,8 @@ public:
 		const FString MeshSelectionPropertyName = "DataflowEditor." + FDataflowEditorCommandsImpl::BeginMeshSelectionToolIdentifier;
 		Set(*MeshSelectionPropertyName, new FSlateImageBrush(RootToContentDir(TEXT("Slate/Dataflow_SelectVertex20x.png")), Icon20x20));
 
+		DefaultMaterial = Cast<UMaterial>(StaticLoadObject(UMaterial::StaticClass(), NULL, TEXT("/Engine/BasicShapes/BasicShapeMaterial")));
+
 		FSlateStyleRegistry::RegisterSlateStyle(*this);
 	}
 
@@ -48,6 +52,10 @@ public:
 		FSlateStyleRegistry::UnRegisterSlateStyle(*this);
 	}
 
+
+	/** Default Rendering Material for Mesh surfaces */
+	UMaterial* DefaultMaterial = nullptr;
+
 public:
 
 	static FDataflowEditorStyle& Get()
diff --git a/Engine/Plugins/Experimental/Dataflow/Source/DataflowEditor/Private/Dataflow/DataflowEditorToolkit.cpp b/Engine/Plugins/Experimental/Dataflow/Source/DataflowEditor/Private/Dataflow/DataflowEditorToolkit.cpp
--- a/Engine/Plugins/Experimental/Dataflow/Source/DataflowEditor/Private/Dataflow/DataflowEditorToolkit.cpp
+++ b/Engine/Plugins/Experimental/Dataflow/Source/DataflowEditor/Private/Dataflow/DataflowEditorToolkit.cpp
@@ -6,6 +6,7 @@
 #include "Animation/Skeleton.h"
 #include "Dataflow/DataflowCore.h"
 #include "Dataflow/DataflowEditor.h"
+#include "Dataflow/DataflowEditorContent.h"
 #include "Dataflow/DataflowEditorCommands.h"
 #include "Dataflow/DataflowEditorMode.h"
 #include "Dataflow/DataflowEditorViewport.h"
@@ -368,7 +369,7 @@ void FDataflowEditorToolkit::OnPropertyValueChanged(const FPropertyChangedEvent&
 	if (TObjectPtr<UDataflowEditorContent> EditorContent = GetDataflowEditorContent())
 	{
 		ensure(EditorContent);
-		FDataflowEditorCommands::OnPropertyValueChanged(EditorContent->DataflowAsset, EditorContent->DataflowContext, EditorContent->LastNodeTimestamp, PropertyChangedEvent, PrevNodeSelection);
+		FDataflowEditorCommands::OnPropertyValueChanged(EditorContent->DataflowAsset, EditorContent->DataflowContext, EditorContent->LastModifiedTimestamp, PropertyChangedEvent, PrevNodeSelection);
 	}
 }
 
@@ -376,7 +377,7 @@ void FDataflowEditorToolkit::OnAssetPropertyValueChanged(const FPropertyChangedE
 {
 	if (TObjectPtr<UDataflowEditorContent> EditorContent = GetDataflowEditorContent())
 	{
-		ensure(EditorContent);	FDataflowEditorCommands::OnAssetPropertyValueChanged(EditorContent->DataflowAsset, EditorContent->DataflowContext, EditorContent->LastNodeTimestamp, PropertyChangedEvent);
+		ensure(EditorContent);	FDataflowEditorCommands::OnAssetPropertyValueChanged(EditorContent->DataflowAsset, EditorContent->DataflowContext, EditorContent->LastModifiedTimestamp, PropertyChangedEvent);
 	}
 }
 
@@ -467,9 +468,9 @@ void FDataflowEditorToolkit::Tick(float DeltaTime)
 			if (!EditorContent->DataflowContext)
 			{
 				EditorContent->DataflowContext = MakeShared<Dataflow::FEngineContext>(EditorContent->DataflowOwner, EditorContent->DataflowAsset, Dataflow::FTimestamp::Invalid);
-				EditorContent->LastNodeTimestamp = Dataflow::FTimestamp::Invalid;
+				EditorContent->LastModifiedTimestamp = Dataflow::FTimestamp::Invalid;
 			}
-			FDataflowEditorCommands::EvaluateTerminalNode(*EditorContent->DataflowContext.Get(), EditorContent->LastNodeTimestamp,
+			FDataflowEditorCommands::EvaluateTerminalNode(*EditorContent->DataflowContext.Get(), EditorContent->LastModifiedTimestamp,
 														  EditorContent->DataflowAsset, nullptr, nullptr, EditorContent->DataflowOwner, EditorContent->DataflowTerminal);
 		}
 	}
@@ -496,9 +497,9 @@ TSharedRef<SDataflowGraphEditor> FDataflowEditorToolkit::CreateGraphEditorWidget
 					EditorContent->DataflowContext = MakeShared<Dataflow::FEngineContext>(EditorContent->DataflowOwner, EditorContent->DataflowAsset, Dataflow::FTimestamp::Invalid);
 				}
 				Node->Invalidate();
-				EditorContent->LastNodeTimestamp = Dataflow::FTimestamp::Invalid;
+				EditorContent->LastModifiedTimestamp = Dataflow::FTimestamp::Invalid;
 
-				FDataflowEditorCommands::EvaluateTerminalNode(*EditorContent->DataflowContext.Get(), EditorContent->LastNodeTimestamp,
+				FDataflowEditorCommands::EvaluateTerminalNode(*EditorContent->DataflowContext.Get(), EditorContent->LastModifiedTimestamp,
 															  EditorContent->DataflowAsset, Node, Out, EditorContent->DataflowOwner, EditorContent->DataflowTerminal);
 			}
 		}
diff --git a/Engine/Plugins/Experimental/Dataflow/Source/DataflowEditor/Private/Dataflow/DataflowEditorUtil.cpp b/Engine/Plugins/Experimental/Dataflow/Source/DataflowEditor/Private/Dataflow/DataflowEditorUtil.cpp
new file mode 100644
--- /dev/null
+++ b/Engine/Plugins/Experimental/Dataflow/Source/DataflowEditor/Private/Dataflow/DataflowEditorUtil.cpp
@@ -0,0 +1,182 @@
+// Copyright Epic Games, Inc. All Rights Reserved.
+
+#include "Dataflow/DataflowEditorUtil.h"
+
+#include "Animation/Skeleton.h"
+#include "Dataflow/DataflowEditor.h"
+#include "Dataflow/DataflowEditorContent.h"
+#include "Dataflow/DataflowEditorToolkit.h"
+#include "Dataflow/DataflowEdNode.h"
+#include "Dataflow/DataflowObject.h"
+#include "DynamicMesh/MeshNormals.h"
+#include "Engine/SkeletalMesh.h"
+#include "Materials/Material.h"
+
+
+using namespace UE::Geometry;
+
+namespace Private
+{
+	UDataflow* GetDataflowAssetFrom(UObject* InObject)
+	{
+		if (UClass* Class = InObject->GetClass())
+		{
+			if (FProperty* Property = Class->FindPropertyByName(FName("DataflowAsset")))
+			{
+				return *Property->ContainerPtrToValuePtr<UDataflow*>(InObject);
+			}
+		}
+		return nullptr;
+	}
+
+	USkeletalMesh* GetSkeletalMeshFrom(UObject* InObject)
+	{
+		if (UClass* Class = InObject->GetClass())
+		{
+			if (FProperty* Property = Class->FindPropertyByName(FName("SkeletalMesh")))
+			{
+				return *Property->ContainerPtrToValuePtr<USkeletalMesh*>(InObject);
+			}
+		}
+		return nullptr;
+	}
+	
+	UAnimationAsset* GetAnimationAssetFrom(UObject* InObject)
+	{
+		if (UClass* Class = InObject->GetClass())
+		{
+			if (FProperty* Property = Class->FindPropertyByName(FName("AnimationAsset")))
+			{
+				return *Property->ContainerPtrToValuePtr<UAnimationAsset*>(InObject);
+			}
+		}
+		return nullptr;
+	}
+
+	FString GetDataflowTerminalFrom(UObject* InObject)
+	{
+		if (UClass* Class = InObject->GetClass())
+		{
+			if (FProperty* Property = Class->FindPropertyByName(FName("DataflowTerminal")))
+			{
+				return *Property->ContainerPtrToValuePtr<FString>(InObject);
+			}
+		}
+		return FString();
+	}
+};
+
+
+namespace UE
+{
+	namespace Material
+	{
+		UMaterial* LoadMaterialFromPath(const FName& InPath, UObject* Outer)
+		{
+			if (InPath == NAME_None) return nullptr;
+
+			return Cast<UMaterial>(StaticLoadObject(UMaterial::StaticClass(), Outer, *InPath.ToString()));
+		}
+	}
+
+
+	namespace Conversion
+	{
+		// Convert a rendering facade to a dynamic mesh
+		void RenderingFacadeToDynamicMesh(const GeometryCollection::Facades::FRenderingFacade& Facade, FDynamicMesh3& DynamicMesh)
+		{
+			if (Facade.CanRenderSurface())
+			{
+				const int32 NumTriangles = Facade.NumTriangles();
+				const int32 NumVertices = Facade.NumVertices();
+
+				const TManagedArray<FIntVector>& Indices = Facade.GetIndices();
+				const TManagedArray<FVector3f>& Positions = Facade.GetVertices();
+				const TManagedArray<FVector3f>& Normals = Facade.GetNormals();
+				const TManagedArray<FLinearColor>& Colors = Facade.GetVertexColor();
+
+				DynamicMesh.Clear();
+				for (int32 VertexIndex = 0; VertexIndex < NumVertices; ++VertexIndex)
+				{
+					DynamicMesh.AppendVertex(FVertexInfo(FVector3d(Positions[VertexIndex]), Normals[VertexIndex],
+											 FVector3f(Colors[VertexIndex].R, Colors[VertexIndex].G, Colors[VertexIndex].B)));
+				}
+				for (int32 TriangleIndex = 0; TriangleIndex < NumTriangles; ++TriangleIndex)
+				{
+					DynamicMesh.AppendTriangle(FIndex3i(Indices[TriangleIndex].X, Indices[TriangleIndex].Y, Indices[TriangleIndex].Z));
+				}
+
+				FMeshNormals::QuickComputeVertexNormals(DynamicMesh);
+			}
+		}
+
+		// Convert a dataflow component to a dynamic mesh
+		void DataflowToDynamicMesh(TSharedPtr<::Dataflow::FEngineContext> DataflowContext, UObject* Asset, UDataflow* Dataflow, FDynamicMesh3& DynamicMesh)
+		{
+			// just call update on the preview scene.
+
+			FManagedArrayCollection RenderCollection;
+			GeometryCollection::Facades::FRenderingFacade Facade(RenderCollection);
+			Facade.DefineSchema();
+
+
+			for (const UDataflowEdNode* Target : Dataflow->GetRenderTargets())
+			{
+				if (Target)
+				{
+					Target->Render(Facade, DataflowContext);
+				}
+			}
+			RenderingFacadeToDynamicMesh(Facade, DynamicMesh);
+		}
+
+		// Convert a dynamic mesh to a rendering facade
+		void DynamicMeshToRenderingFacade(const FDynamicMesh3& DynamicMesh, GeometryCollection::Facades::FRenderingFacade& Facade)
+		{
+			if (Facade.CanRenderSurface())
+			{
+				const int32 NumTriangles = Facade.NumTriangles();
+				const int32 NumVertices = Facade.NumVertices();
+
+				// We can only override vertices attributes (position, normals, colors)
+				if ((NumTriangles == DynamicMesh.TriangleCount()) && (NumVertices == DynamicMesh.VertexCount()))
+				{
+					TManagedArray<FVector3f>& Positions = Facade.ModifyVertices();
+					TManagedArray<FVector3f>& Normals = Facade.ModifyNormals();
+					TManagedArray<FLinearColor>& Colors = Facade.ModifyVertexColor();
+
+					for (int32 VertexIndex = 0; VertexIndex < NumVertices; ++VertexIndex)
+					{
+						Positions[VertexIndex] = FVector3f(DynamicMesh.GetVertex(VertexIndex));
+						Normals[VertexIndex] = DynamicMesh.GetVertexNormal(VertexIndex);
+						Colors[VertexIndex] = DynamicMesh.GetVertexColor(VertexIndex);
+					}
+				}
+			}
+		}
+
+		// Convert a dynamic mesh to a dataflow component
+		void DynamicMeshToDataflow(const FDynamicMesh3& DynamicMesh, UDataflow* Dataflow)
+		{
+			//todo
+		}
+	}
+}
+
+namespace Dataflow
+{
+	TSharedPtr<FEngineContext> GetContext(UDataflowEditorContent* Content)
+	{
+		if (Content)
+		{
+			if (!Content->DataflowContext)
+			{
+				Content->DataflowContext = MakeShared<FEngineContext>(Content->DataflowOwner, Content->DataflowAsset, FTimestamp::Invalid);
+			}
+			return Content->DataflowContext;
+		}
+
+		ensure(false);
+		return MakeShared<FEngineContext>(nullptr, nullptr, FTimestamp::Invalid);
+	}
+}
\ No newline at end of file
diff --git a/Engine/Plugins/Experimental/Dataflow/Source/DataflowEditor/Private/Dataflow/DataflowEditorUtil.h b/Engine/Plugins/Experimental/Dataflow/Source/DataflowEditor/Private/Dataflow/DataflowEditorUtil.h
new file mode 100644
--- /dev/null
+++ b/Engine/Plugins/Experimental/Dataflow/Source/DataflowEditor/Private/Dataflow/DataflowEditorUtil.h
@@ -0,0 +1,56 @@
+// Copyright Epic Games, Inc. All Rights Reserved.
+
+#pragma once
+
+#include "CoreTypes.h"
+#include "Containers/UnrealString.h"
+#include "DynamicMesh/DynamicMesh3.h"
+#include "Dataflow/DataflowObjectInterface.h"
+#include "GeometryCollection/Facades/CollectionRenderingFacade.h"
+#include "Templates/SharedPointer.h"
+
+class UDataflow;
+class UObject;
+class UDataflowEditorContent;
+class USkeletalMesh;
+class UAnimationAsset;
+class UMaterial;
+
+namespace Private
+{
+	UDataflow* GetDataflowAssetFrom(UObject* InObject);
+
+	USkeletalMesh* GetSkeletalMeshFrom(UObject* InObject);
+
+	UAnimationAsset* GetAnimationAssetFrom(UObject* InObject);
+
+	FString GetDataflowTerminalFrom(UObject* InObject);
+};
+
+namespace UE
+{
+	namespace Material
+	{
+		UMaterial* LoadMaterialFromPath( const FName& Path, UObject* Outer);
+	}
+
+	namespace Conversion
+	{
+		// Convert a rendering facade to a dynamic mesh
+		void RenderingFacadeToDynamicMesh(const GeometryCollection::Facades::FRenderingFacade& Facade, UE::Geometry::FDynamicMesh3& DynamicMesh);
+
+		// Convert a dataflow component to a dynamic mesh
+		void DataflowToDynamicMesh(TSharedPtr<::Dataflow::FEngineContext> DataflowContext, UObject* Asset, UDataflow* Dataflow, UE::Geometry::FDynamicMesh3& DynamicMesh);
+
+		// Convert a dynamic mesh to a rendering facade
+		void DynamicMeshToRenderingFacade(const UE::Geometry::FDynamicMesh3& DynamicMesh, GeometryCollection::Facades::FRenderingFacade& Facade);
+
+		// Convert a dynamic mesh to a dataflow component
+		void DynamicMeshToDataflow(const UE::Geometry::FDynamicMesh3& DynamicMesh, UDataflow* Dataflow);
+	}
+}
+
+namespace Dataflow
+{
+	TSharedPtr<::Dataflow::FEngineContext> GetContext(UDataflowEditorContent* Content);
+}
\ No newline at end of file
diff --git a/Engine/Plugins/Experimental/Dataflow/Source/DataflowEditor/Private/Dataflow/DataflowEditorViewportClient.cpp b/Engine/Plugins/Experimental/Dataflow/Source/DataflowEditor/Private/Dataflow/DataflowEditorViewportClient.cpp
--- a/Engine/Plugins/Experimental/Dataflow/Source/DataflowEditor/Private/Dataflow/DataflowEditorViewportClient.cpp
+++ b/Engine/Plugins/Experimental/Dataflow/Source/DataflowEditor/Private/Dataflow/DataflowEditorViewportClient.cpp
@@ -1,16 +1,18 @@
 // Copyright Epic Games, Inc. All Rights Reserved.
 #include "Dataflow/DataflowEditorViewportClient.h"
+
 #include "Dataflow/DataflowActor.h"
+#include "Dataflow/DataflowEditor.h"
+#include "Dataflow/DataflowEditorContent.h"
+#include "Dataflow/DataflowEditorMode.h"
 #include "Dataflow/DataflowEditorToolkit.h"
 #include "Dataflow/DataflowEngineSceneHitProxies.h"
-#include "Dataflow/DataflowXml.h"
-#include "Dataflow/DataflowEditor.h"
 #include "Dataflow/DataflowPreviewScene.h"
-#include "Dataflow/DataflowEditorMode.h"
-#include "Selection.h"
+#include "Dataflow/DataflowXml.h"
+#include "EditorModeManager.h"
 #include "HAL/PlatformApplicationMisc.h"
 #include "PreviewScene.h"
-#include "EditorModeManager.h"
+#include "Selection.h"
 
 FDataflowEditorViewportClient::FDataflowEditorViewportClient(FEditorModeTools* InModeTools,
                                                              FPreviewScene* InPreviewScene,
@@ -31,65 +33,8 @@ void FDataflowEditorViewportClient::SetDataflowEditorToolkit(TWeakPtr<FDataflowE
 	DataflowEditorToolkitPtr = InDataflowEditorToolkitPtr;
 }
 
-void FDataflowEditorViewportClient::SetSelectionMode(FDataflowSelectionState::EMode InState)
-{
-	FDataflowSelectionState State = PreviewScene->GetDataflowComponent()->GetSelectionState();
-
-	if (SelectionMode == InState)
-	{
-		SelectionMode = FDataflowSelectionState::EMode::DSS_Dataflow_None;
-	}
-	else
-	{
-		SelectionMode = InState;
-	}
-
-	State.Mode = SelectionMode;
-	PreviewScene->GetDataflowComponent()->SetSelectionState(State);
-
-	if (SelectionMode == FDataflowSelectionState::EMode::DSS_Dataflow_None)
-	{
-		if (!PreviewScene->GetDataflowComponent()->GetSelectionState().IsEmpty())
-		{
-			PreviewScene->GetDataflowComponent()->SetSelectionState(FDataflowSelectionState(SelectionMode));
-		}
-	}
-}
-bool FDataflowEditorViewportClient::CanSetSelectionMode(FDataflowSelectionState::EMode InState)
-{
-	TSharedPtr<FDataflowEditorToolkit> DataflowEditorToolkit = DataflowEditorToolkitPtr.Pin();
-	if (DataflowEditorToolkitPtr.IsValid())
-	{
-		if (TObjectPtr<UDataflowEditorContent> EditorContent = PreviewScene->GetDataflowEditorContent())
-		{
-			if (const UDataflow* Dataflow = EditorContent->DataflowAsset)
-			{
-				if (Dataflow->GetRenderTargets().Num())
-				{
-					if (InState == FDataflowSelectionState::EMode::DSS_Dataflow_Object)
-					{
-						return true;
-					}
-
-					if (InState == FDataflowSelectionState::EMode::DSS_Dataflow_Vertex
-						&& !PreviewScene->GetDataflowComponent()->GetSelectionState().Nodes.IsEmpty())
-					{
-						return true;
-					}
-				}
-			}
-		}
-	}
-
-	return false;
-}
-
-bool FDataflowEditorViewportClient::IsSelectionModeActive(FDataflowSelectionState::EMode InState)
-{
-	return SelectionMode == InState;
-}
 
-::Dataflow::FTimestamp FDataflowEditorViewportClient::LatestTimestamp(const UDataflow* Dataflow, const ::Dataflow::FContext* Context)
+Dataflow::FTimestamp LatestTimestamp(const UDataflow* Dataflow, const ::Dataflow::FContext* Context)
 {
 	if (Dataflow && Context)
 	{
@@ -98,155 +43,13 @@ bool FDataflowEditorViewportClient::IsSelectionModeActive(FDataflowSelectionStat
 	return ::Dataflow::FTimestamp::Invalid;
 }
 
-bool FDataflowEditorViewportClient::InputKey(const FInputKeyEventArgs& EventArgs)
-{
-	bool bHandled = false;
-
-	FInputEventState InputState(EventArgs.Viewport, EventArgs.Key, EventArgs.Event);
-	if (InputState.IsCtrlButtonPressed() && EventArgs.Key == EKeys::C)
-	{
-		if (PreviewScene->GetDataflowComponent()->GetSelectionState().Vertices.Num())
-		{
-			FString XmlBuffer = FDataflowXmlWrite()
-				.Begin()
-				.MakeVertexSelectionBlock(PreviewScene->GetDataflowComponent()->GetSelectionState().Vertices)
-				.End()
-				.ToString();
-			FPlatformApplicationMisc::ClipboardCopy(*XmlBuffer);
-		}
-	}
-
-
-	if (!bHandled)
-	{
-		bHandled = FEditorViewportClient::InputKey(EventArgs);
-	}
-
-	return bHandled;
-}
-
-void FDataflowEditorViewportClient::ProcessClick(FSceneView& View, HHitProxy* HitProxy, FKey Key, EInputEvent Event, uint32 HitX, uint32 HitY)
-{
-	Super::ProcessClick(View, HitProxy, Key, Event, HitX, HitY);
-	if (PreviewScene->GetDataflowComponent())
-	{
-		const bool bIsShiftKeyDown = Viewport->KeyState(EKeys::LeftShift) || Viewport->KeyState(EKeys::RightShift);
-		const bool bIsCtrltKeyDown = Viewport->KeyState(EKeys::LeftControl) || Viewport->KeyState(EKeys::RightControl);
-
-		USelection* SelectedComponents = ModeTools->GetSelectedComponents();
-
-		TArray<UPrimitiveComponent*> PreviouslySelectedComponents;
-		SelectedComponents->GetSelectedObjects<UPrimitiveComponent>(PreviouslySelectedComponents);
-
-		SelectedComponents->Modify();
-		SelectedComponents->BeginBatchSelectOperation();
-
-		SelectedComponents->DeselectAll();
-
-		if (HitProxy && HitProxy->IsA(HActor::StaticGetType()))
-		{
-			const HActor* const ActorProxy = static_cast<HActor*>(HitProxy);
-			if (ActorProxy && ActorProxy->Actor)
-			{
-				const AActor* const Actor = ActorProxy->Actor;
-			
-				Actor->ForEachComponent<UPrimitiveComponent>(true, [&](UPrimitiveComponent* Component)
-				{
-					SelectedComponents->Select(Component);
-					Component->PushSelectionToProxy();
-				});
-			}
-		}
-
-		SelectedComponents->EndBatchSelectOperation();
-
-		for (UPrimitiveComponent* const Component : PreviouslySelectedComponents)
-		{
-			Component->PushSelectionToProxy();
-		}
-
-		FDataflowSelectionState SelectionState = PreviewScene->GetDataflowComponent()->GetSelectionState();
-		FDataflowSelectionState PreState = SelectionState;
-
-		if (SelectionMode == FDataflowSelectionState::EMode::DSS_Dataflow_Object)
-		{
-			if (HitProxy && HitProxy->IsA(HDataflowNode::StaticGetType()))
-			{
-				HDataflowNode* DataflowNode = (HDataflowNode*)(HitProxy);
-				FDataflowSelectionState::ObjectID ID(DataflowNode->NodeName, DataflowNode->GeometryIndex);
-
-				if (bIsShiftKeyDown)
-				{
-					if (!SelectionState.Nodes.Contains(ID))
-					{
-						SelectionState.Nodes.AddUnique(ID);
-					}
-				}
-				else if (bIsCtrltKeyDown)
-				{
-					if (SelectionState.Nodes.Contains(ID))
-					{
-						SelectionState.Nodes.Remove(ID);
-					}
-				}
-				else
-				{
-					SelectionState.Nodes.Empty();
-					SelectionState.Nodes.AddUnique(ID);
-				}
-			}
-			else if (!bIsShiftKeyDown && !bIsCtrltKeyDown)
-			{
-				SelectionState.Nodes.Empty();
-			}
-		}
-
-		if (SelectionMode == FDataflowSelectionState::EMode::DSS_Dataflow_Vertex)
-		{
-			if (HitProxy && HitProxy->IsA(HDataflowVertex::StaticGetType()))
-			{
-				HDataflowVertex* DataflowVertex = (HDataflowVertex*)(HitProxy);
-				int32 ID = DataflowVertex->SectionIndex;
-				if (bIsShiftKeyDown)
-				{
-					if (!SelectionState.Vertices.Contains(ID))
-					{
-						SelectionState.Vertices.AddUnique(ID);
-					}
-				}
-				else if (bIsCtrltKeyDown)
-				{
-					if (SelectionState.Vertices.Contains(ID))
-					{
-						SelectionState.Vertices.Remove(ID);
-					}
-				}
-				else
-				{
-					SelectionState.Vertices.Empty();
-					SelectionState.Vertices.AddUnique(ID);
-				}
-			}
-			else if (!bIsShiftKeyDown && !bIsCtrltKeyDown)
-			{
-				SelectionState.Vertices.Empty();
-			}
-		}
-
-		if (PreState != SelectionState)
-		{
-			PreviewScene->GetDataflowComponent()->SetSelectionState(SelectionState);
-		}
-	}
-}
-
 void FDataflowEditorViewportClient::Tick(float DeltaSeconds)
 {
 	Super::Tick(DeltaSeconds);
 
 	const TSharedPtr<FDataflowEditorToolkit> DataflowEditorToolkit = DataflowEditorToolkitPtr.Pin();
 
-	if (PreviewScene->GetDataflowComponent())
+	if (PreviewScene)
 	{
 		if (TObjectPtr<UDataflowEditorContent> EditorContent = PreviewScene->GetDataflowEditorContent())
 		{
@@ -255,10 +58,10 @@ void FDataflowEditorViewportClient::Tick(float DeltaSeconds)
 				if (const UDataflow* Dataflow = EditorContent->DataflowAsset)
 				{
 					const Dataflow::FTimestamp SystemTimestamp = LatestTimestamp(Dataflow, Context.Get());
-					if (SystemTimestamp >= LastModifiedTimestamp)
+					if (SystemTimestamp >= EditorContent->LastModifiedTimestamp)
 					{
-						PreviewScene->UpdateDataflowComponent();
-						LastModifiedTimestamp = LatestTimestamp(Dataflow, Context.Get()).Value + 1;
+						PreviewScene->Update();
+						EditorContent->LastModifiedTimestamp = LatestTimestamp(EditorContent->DataflowAsset, EditorContent->DataflowContext.Get()).Value + 1;
 					}
 				}
 			}
diff --git a/Engine/Plugins/Experimental/Dataflow/Source/DataflowEditor/Private/Dataflow/DataflowPreviewScene.cpp b/Engine/Plugins/Experimental/Dataflow/Source/DataflowEditor/Private/Dataflow/DataflowPreviewScene.cpp
--- a/Engine/Plugins/Experimental/Dataflow/Source/DataflowEditor/Private/Dataflow/DataflowPreviewScene.cpp
+++ b/Engine/Plugins/Experimental/Dataflow/Source/DataflowEditor/Private/Dataflow/DataflowPreviewScene.cpp
@@ -9,57 +9,26 @@
 #include "Dataflow/DataflowActor.h"
 #include "Dataflow/DataflowComponent.h"
 #include "Dataflow/DataflowEditor.h"
+#include "Dataflow/DataflowEditorContent.h"
+#include "Dataflow/DataflowEditorStyle.h"
+#include "Dataflow/DataflowEditorUtil.h"
 #include "Elements/Framework/EngineElementsLibrary.h"
+#include "ModelingToolTargetUtil.h"
 
 #define LOCTEXT_NAMESPACE "FDataflowPreviewScene"
 
+
 FDataflowPreviewScene::FDataflowPreviewScene(FPreviewScene::ConstructionValues ConstructionValues,TObjectPtr<UDataflowEditorContent> InEditorContent) 
 	: FAdvancedPreviewScene(ConstructionValues), EditorContent(InEditorContent)
 {
 	check(EditorContent);
-	SkeletalMeshActor = GetWorld()->SpawnActor<AActor>(AActor::StaticClass());
-
-	if(EditorContent->SkeletalMesh && EditorContent->bHasValidSkeletalMesh)
-	{
-		SkeletalMeshComponent = NewObject<USkeletalMeshComponent>(SkeletalMeshActor);
-		SkeletalMeshComponent->SelectionOverrideDelegate = UPrimitiveComponent::FSelectionOverride::CreateRaw(this, &FDataflowPreviewScene::IsComponentSelected);
-		SkeletalMeshComponent->SetDisablePostProcessBlueprint(true);
-		UpdateSkeletalMeshComponent();
-	}
-	
-	// @todo(DynamicMeshRendering) : Enable Dynamic Mesh Rendering for dataflow terminals. Hide the dataflow 
-	DataflowActor = Cast<ADataflowActor>(GetWorld()->SpawnActor<ADataflowActor>(ADataflowActor::StaticClass()));
-	DataflowComponent = DataflowActor->GetDataflowComponent();
-	//DataflowComponent->SelectionOverrideDelegate = UPrimitiveComponent::FSelectionOverride::CreateRaw(this, &FDataflowPreviewScene::IsComponentSelected);
-	//DataflowComponent->SetVisibility(false);
-	UpdateDataflowComponent();
-	
-
 	DynamicMeshActor = GetWorld()->SpawnActor<AActor>(AActor::StaticClass());	
-	SkeletalMeshActor->RegisterAllComponents();
-	DataflowActor->RegisterAllComponents();
 	DynamicMeshActor->RegisterAllComponents();
-
 	SetFloorVisibility(false, true);
 }
 
 FDataflowPreviewScene::~FDataflowPreviewScene()
 {
-	if (SkeletalMeshComponent)
-	{
-		SkeletalMeshComponent->TransformUpdated.RemoveAll(this);
-		SkeletalMeshComponent->SelectionOverrideDelegate.Unbind();
-		SkeletalMeshComponent->UnregisterComponent();
-		SkeletalMeshComponent->DestroyComponent();
-	}
-
-	if (DataflowComponent)
-	{
-		DataflowComponent->SelectionOverrideDelegate.Unbind();
-		DataflowComponent->UnregisterComponent();
-		DataflowComponent->DestroyComponent();
-	}
-
 	ResetDynamicMeshComponents();
 }
 
@@ -67,58 +36,43 @@ void FDataflowPreviewScene::AddReferencedObjects(FReferenceCollector& Collector)
 {
 	FAdvancedPreviewScene::AddReferencedObjects(Collector);
 	Collector.AddReferencedObject(EditorContent);
-
-	Collector.AddReferencedObject(DataflowComponent);
-	Collector.AddReferencedObject(SkeletalMeshComponent);
-	Collector.AddReferencedObject(SkeletalMeshActor);
-	Collector.AddReferencedObject(DataflowActor);
 	Collector.AddReferencedObject(DynamicMeshActor);
-	Collector.AddReferencedObject(PreviewAnimInstance);
 	Collector.AddReferencedObjects(DynamicMeshComponents);
 }
 
-void FDataflowPreviewScene::UpdateDataflowComponent()
+void FDataflowPreviewScene::Update()
 {
-	if (EditorContent->DataflowAsset)
-	{
-		DataflowComponent->ResetRenderTargets();
+	using namespace UE::Geometry;//FDynamicMesh3
 
-		DataflowComponent->SetDataflow(EditorContent->DataflowAsset);
-		DataflowComponent->SetContext(EditorContent->DataflowContext);
+	// The preview scene for the construction view will be
+	// cleared and rebuilt from scratch. This will genrate a 
+	// list of UPrimitiveComponents for rendering.
+	ResetDynamicMeshComponents();
 
-		for (const UDataflowEdNode* const Node : EditorContent->DataflowAsset->GetRenderTargets())
+	if (EditorContent)
+	{
+		TObjectPtr<UDataflow> DataflowAsset = EditorContent->DataflowAsset;
+		TSharedPtr<Dataflow::FEngineContext> DataflowContext = EditorContent->DataflowContext;
+		if(DataflowAsset && DataflowContext)
 		{
-			DataflowComponent->AddRenderTarget(Node);
+			for (const UDataflowEdNode* Target : DataflowAsset->GetRenderTargets())
+			{
+				if (Target)
+				{
+					FDynamicMesh3 DynamicMesh;
+					FManagedArrayCollection RenderCollection;
+					GeometryCollection::Facades::FRenderingFacade Facade(RenderCollection);
+					Facade.DefineSchema();
+
+					Target->Render(Facade, DataflowContext);
+					UE::Conversion::RenderingFacadeToDynamicMesh(Facade, DynamicMesh);
+					AddDynamicMeshComponent(MoveTemp(DynamicMesh), {});
+				}
+			}
 		}
-		DataflowComponent->UpdateBounds();
 	}
 }
 
-void FDataflowPreviewScene::UpdateSkeletalMeshComponent()
-{
-	SkeletalMeshComponent->SetSkeletalMeshAsset(EditorContent->SkeletalMesh);
-
-	if (EditorContent->AnimationAsset)
-	{
-		PreviewAnimInstance = NewObject<UAnimSingleNodeInstance>(SkeletalMeshComponent);
-		PreviewAnimInstance->SetAnimationAsset(EditorContent->AnimationAsset);
-
-		SkeletalMeshComponent->SetAnimationMode(EAnimationMode::AnimationSingleNode);
-		SkeletalMeshComponent->InitAnim(true);
-		SkeletalMeshComponent->AnimationData.PopulateFrom(PreviewAnimInstance);
-		SkeletalMeshComponent->AnimScriptInstance = PreviewAnimInstance;
-		SkeletalMeshComponent->AnimScriptInstance->InitializeAnimation();
-		SkeletalMeshComponent->ValidateAnimation();
-	}
-	else
-	{
-		SkeletalMeshComponent->Stop();
-		SkeletalMeshComponent->AnimationData = FSingleAnimationPlayData();
-		SkeletalMeshComponent->AnimScriptInstance = nullptr;
-	}
-	SkeletalMeshComponent->UpdateBounds();
-}
-
 void FDataflowPreviewScene::ResetDynamicMeshComponents()
 {
 	for(const TObjectPtr<UDynamicMeshComponent>& DynamicMeshComponent : DynamicMeshComponents)
@@ -135,8 +89,16 @@ TObjectPtr<UDynamicMeshComponent>& FDataflowPreviewScene::AddDynamicMeshComponen
 	TObjectPtr<UDynamicMeshComponent> DynamicMeshComponent = NewObject<UDynamicMeshComponent>(DynamicMeshActor);
 		
 	DynamicMeshComponent->SetMesh(MoveTemp(DynamicMesh));
-	DynamicMeshComponent->ConfigureMaterialSet(MaterialSet);
-	check(DynamicMeshComponent->ValidateMaterialSlots(false, false));
+
+	// @todo(Material) This is just to have a material, we should transfer the materials from the assets if they have them. 
+	if (FDataflowEditorStyle::Get().DefaultMaterial)
+	{
+		DynamicMeshComponent->ConfigureMaterialSet({ FDataflowEditorStyle::Get().DefaultMaterial });
+	}
+	else
+	{
+		DynamicMeshComponent->ValidateMaterialSlots(true, false);
+	}
 	DynamicMeshComponent->SelectionOverrideDelegate = UPrimitiveComponent::FSelectionOverride::CreateRaw(this, &FDataflowPreviewScene::IsComponentSelected);
 	
 	DynamicMeshComponent->RegisterComponentWithWorld(GetWorld());
diff --git a/Engine/Plugins/Experimental/Dataflow/Source/DataflowEditor/Private/Dataflow/DataflowToolTarget.cpp b/Engine/Plugins/Experimental/Dataflow/Source/DataflowEditor/Private/Dataflow/DataflowToolTarget.cpp
--- a/Engine/Plugins/Experimental/Dataflow/Source/DataflowEditor/Private/Dataflow/DataflowToolTarget.cpp
+++ b/Engine/Plugins/Experimental/Dataflow/Source/DataflowEditor/Private/Dataflow/DataflowToolTarget.cpp
@@ -4,7 +4,9 @@
 
 #include "ConversionUtils/DynamicMeshViaMeshDescriptionUtil.h"
 #include "Dataflow/DataflowEditor.h"
+#include "Dataflow/DataflowEditorContent.h"
 #include "Dataflow/DataflowEditorToolkit.h"
+#include "Dataflow/DataflowEditorUtil.h"
 #include "Dataflow/DataflowEdNode.h"
 #include "Dataflow/DataflowObject.h"
 #include "Dataflow/DataflowObjectInterface.h"
@@ -19,106 +21,7 @@
 
 #include UE_INLINE_GENERATED_CPP_BY_NAME(DataflowToolTarget)
 
-using namespace UE::Geometry;
 
-namespace UE
-{
-	namespace Conversion
-	{
-		// Convert a rendering facade to a dynamic mesh
-		void RenderingFacadeToDynamicMesh(const GeometryCollection::Facades::FRenderingFacade& Facade, FDynamicMesh3& DynamicMesh)
-		{
-			if(Facade.CanRenderSurface())
-			{
-				const int32 NumTriangles = Facade.NumTriangles();
-				const int32 NumVertices = Facade.NumVertices();
-				
-				const TManagedArray<FIntVector>& Indices = Facade.GetIndices();
-				const TManagedArray<FVector3f>& Positions = Facade.GetVertices();
-				const TManagedArray<FVector3f>& Normals = Facade.GetNormals();
-				const TManagedArray<FLinearColor>& Colors = Facade.GetVertexColor();
-
-				DynamicMesh.Clear();
-				for(int32 VertexIndex = 0; VertexIndex < NumVertices; ++VertexIndex)
-				{
-					DynamicMesh.AppendVertex(FVertexInfo(FVector3d(Positions[VertexIndex]), Normals[VertexIndex],
-						FVector3f(Colors[VertexIndex].R, Colors[VertexIndex].G, Colors[VertexIndex].B)));
-				}
-				for(int32 TriangleIndex = 0; TriangleIndex < NumTriangles; ++TriangleIndex)
-				{
-					DynamicMesh.AppendTriangle(FIndex3i(Indices[TriangleIndex].X, Indices[TriangleIndex].Y, Indices[TriangleIndex].Z));
-				}
-			}
-		}
-		
-		// Convert a dataflow component to a dynamic mesh
-		void DataflowToDynamicMesh(TSharedPtr<::Dataflow::FEngineContext> DataflowContext, UObject* Asset, UDataflow* Dataflow, FDynamicMesh3& DynamicMesh)
-		{
-			FManagedArrayCollection RenderCollection;
-			GeometryCollection::Facades::FRenderingFacade Facade(RenderCollection);
-			Facade.DefineSchema();
-		
-			
-			for (const UDataflowEdNode* Target : Dataflow->GetRenderTargets())
-			{
-				if (Target)
-				{
-					Target->Render(Facade, DataflowContext);
-				}
-			}
-			RenderingFacadeToDynamicMesh(Facade, DynamicMesh);
-		}
-
-		// Convert a dynamic mesh to a rendering facade
-		void DynamicMeshToRenderingFacade(const FDynamicMesh3& DynamicMesh, GeometryCollection::Facades::FRenderingFacade& Facade)
-		{
-			if(Facade.CanRenderSurface())
-			{
-				const int32 NumTriangles = Facade.NumTriangles();
-				const int32 NumVertices = Facade.NumVertices();
-
-				// We can only override vertices attributes (position, normals, colors)
-				if((NumTriangles == DynamicMesh.TriangleCount()) && (NumVertices == DynamicMesh.VertexCount()))
-				{
-					TManagedArray<FVector3f>& Positions = Facade.ModifyVertices();
-					TManagedArray<FVector3f>& Normals = Facade.ModifyNormals();
-					TManagedArray<FLinearColor>& Colors = Facade.ModifyVertexColor();
-
-					for(int32 VertexIndex = 0; VertexIndex < NumVertices; ++VertexIndex)
-					{
-						Positions[VertexIndex] = FVector3f(DynamicMesh.GetVertex(VertexIndex));
-						Normals[VertexIndex] = DynamicMesh.GetVertexNormal(VertexIndex);
-						Colors[VertexIndex] = DynamicMesh.GetVertexColor(VertexIndex);
-					}
-				}
-			}
-		}
-		
-		// Convert a dynamic mesh to a dataflow component
-		void DynamicMeshToDataflow(const FDynamicMesh3& DynamicMesh, UDataflow* Dataflow)
-		{
-			//todo
-		}
-	}
-}
-
-namespace Dataflow
-{
-	TSharedPtr<FEngineContext> GetContext(UDataflowEditorContent* Content)
-	{
-		if (Content)
-		{
-			if (!Content->DataflowContext)
-			{
-				Content->DataflowContext = MakeShared<FEngineContext>(Content->DataflowOwner, Content->DataflowAsset, FTimestamp::Invalid);
-			}
-			return Content->DataflowContext;
-		}
-
-		ensure(false);
-		return MakeShared<FEngineContext>(nullptr, nullptr, FTimestamp::Invalid);
-	}
-}
 
 //
 // UDataflowReadOnlyToolTarget
diff --git a/Engine/Plugins/Experimental/Dataflow/Source/DataflowEditor/Public/Dataflow/DataflowEditor.h b/Engine/Plugins/Experimental/Dataflow/Source/DataflowEditor/Public/Dataflow/DataflowEditor.h
--- a/Engine/Plugins/Experimental/Dataflow/Source/DataflowEditor/Public/Dataflow/DataflowEditor.h
+++ b/Engine/Plugins/Experimental/Dataflow/Source/DataflowEditor/Public/Dataflow/DataflowEditor.h
@@ -17,49 +17,7 @@
 #include "DataflowEditor.generated.h"
 
 class FDataflowEditorToolkit;
-class USkeletalMesh;
-class USkeleton;
-class UAnimationAsset;
-
-/** 
- * Dataflow datas that will be used within the editor classes to evaluate the graph
- */
-UCLASS()
-class UDataflowEditorContent : public UObject
-{
-	GENERATED_BODY()
-
-public:
-	/** Check if the datas flow datas are valid */
-	bool IsValid() const { return DataflowOwner && DataflowAsset;}
-	
-	/** Data flow object owner */
-	TObjectPtr<UObject> DataflowOwner = nullptr;
-
-	/** Data flow asset that we will edit */
-	TObjectPtr<UDataflow> DataflowAsset = nullptr;
-
-	/** Data flow terminal path for evaluation */
-	FString DataflowTerminal = "";
-
-	/** Data flow skeletal mesh*/
-	TObjectPtr<USkeletalMesh> SkeletalMesh = nullptr;
-
-	/** Data flow skeleton*/
-	TObjectPtr<USkeleton> Skeleton = nullptr;
-
-	/** Animation asset to be used to preview simulation */
-	TObjectPtr<UAnimationAsset> AnimationAsset;
-
-	/**  Engine context to be used for dataflow evaluation */
-	TSharedPtr<Dataflow::FEngineContext> DataflowContext = nullptr;
-
-	/** Last data flow evaluated node time stamp */
-	Dataflow::FTimestamp LastNodeTimestamp = Dataflow::FTimestamp::Invalid;
-
-	/** Boolean to check if the skelmesh is valid*/
-	bool bHasValidSkeletalMesh = false;
-};
+class UDataflowEditorContent;
 
 /** 
  * The actual asset editor class doesn't have that much in it, intentionally. 
diff --git a/Engine/Plugins/Experimental/Dataflow/Source/DataflowEditor/Public/Dataflow/DataflowEditorViewportClient.h b/Engine/Plugins/Experimental/Dataflow/Source/DataflowEditor/Public/Dataflow/DataflowEditorViewportClient.h
--- a/Engine/Plugins/Experimental/Dataflow/Source/DataflowEditor/Public/Dataflow/DataflowEditorViewportClient.h
+++ b/Engine/Plugins/Experimental/Dataflow/Source/DataflowEditor/Public/Dataflow/DataflowEditorViewportClient.h
@@ -21,23 +21,12 @@ public:
 	FDataflowEditorViewportClient(FEditorModeTools* InModeTools, FPreviewScene* InPreviewScene,
 								  const TWeakPtr<SEditorViewport> InEditorViewportWidget = nullptr);
 
-	/** Set the data flow toolkit used to create the client*/
-	void SetDataflowEditorToolkit(TWeakPtr<FDataflowEditorToolkit> DataflowToolkit);
-
-	/** Return the latest timestamp */
-	::Dataflow::FTimestamp LatestTimestamp(const UDataflow* Dataflow, const ::Dataflow::FContext* Context);
-
-	// Selection utils
-	void SetSelectionMode(FDataflowSelectionState::EMode InState);
-	bool CanSetSelectionMode(FDataflowSelectionState::EMode InState);
-	bool IsSelectionModeActive(FDataflowSelectionState::EMode InState);
-	FDataflowSelectionState::EMode GetSelectionMode() const { return SelectionMode; }
-	
 	// FEditorViewportClient interface
-	virtual bool InputKey(const FInputKeyEventArgs& EventArgs) override;
-	virtual void ProcessClick(FSceneView& View, HHitProxy* HitProxy, FKey Key, EInputEvent Event, uint32 HitX, uint32 HitY) override;
 	virtual void Tick(float DeltaSeconds) override;
 
+	/** Set the data flow toolkit used to create the client*/
+	void SetDataflowEditorToolkit(TWeakPtr<FDataflowEditorToolkit> DataflowToolkit);
+	
 	// FGCObject Interface
 	virtual void AddReferencedObjects(FReferenceCollector& Collector) override;
 	virtual FString GetReferencerName() const override { return TEXT("FDataflowEditorViewportClient"); }
@@ -46,13 +35,8 @@ private:
 
 	/** Toolkit used to create the viewport client */
 	TWeakPtr<FDataflowEditorToolkit> DataflowEditorToolkitPtr = nullptr;
-	
-	/** Last time stamp used for dataflow evaluation on the component */
-	::Dataflow::FTimestamp LastModifiedTimestamp = ::Dataflow::FTimestamp::Invalid;
 
 	/** Dataflow preview scene from the toolkit */
 	FDataflowPreviewScene* PreviewScene = nullptr;
 
-	/** Selection mode to be used for vertices/faces/objects */
-	FDataflowSelectionState::EMode SelectionMode = FDataflowSelectionState::EMode::DSS_Dataflow_None;
 };
diff --git a/Engine/Plugins/Experimental/Dataflow/Source/DataflowEditor/Public/Dataflow/DataflowPreviewScene.h b/Engine/Plugins/Experimental/Dataflow/Source/DataflowEditor/Public/Dataflow/DataflowPreviewScene.h
--- a/Engine/Plugins/Experimental/Dataflow/Source/DataflowEditor/Public/Dataflow/DataflowPreviewScene.h
+++ b/Engine/Plugins/Experimental/Dataflow/Source/DataflowEditor/Public/Dataflow/DataflowPreviewScene.h
@@ -34,21 +34,13 @@ public:
 	FDataflowPreviewScene(FPreviewScene::ConstructionValues ConstructionValues, TObjectPtr<UDataflowEditorContent> InEditorContent);
 	virtual ~FDataflowPreviewScene();
 
-	// FGCObject interface
-	virtual void AddReferencedObjects(FReferenceCollector& Collector) override;
 
-	// Modify/Get the preview anim instance
-	UAnimSingleNodeInstance* GetPreviewAnimInstance() { return PreviewAnimInstance; }
-	const UAnimSingleNodeInstance* const GetPreviewAnimInstance() const { return PreviewAnimInstance; }
+	void Update();
 
-	// Modify/Get the data flow component
-	UDataflowComponent* GetDataflowComponent() { return DataflowComponent; }
-	const UDataflowComponent* GetDataflowComponent() const { return DataflowComponent; }
 
-	// Modify/Get the skeletal mesh component
-	USkeletalMeshComponent* GetSkeletalMeshComponent() { return SkeletalMeshComponent; }
-	const USkeletalMeshComponent* GetSkeletalMeshComponent() const { return SkeletalMeshComponent; }
-	
+	// FGCObject interface
+	virtual void AddReferencedObjects(FReferenceCollector& Collector) override;
+
 	// Modify/Get the data flow datas
 	TObjectPtr<UDataflowEditorContent> GetDataflowEditorContent() { return EditorContent; }
 	const TObjectPtr<UDataflowEditorContent> GetDataflowEditorContent() const { return EditorContent; }
@@ -63,11 +55,7 @@ public:
 	/** Add a dynamic mesh component to the scene */
 	TObjectPtr<UDynamicMeshComponent>& AddDynamicMeshComponent(UE::Geometry::FDynamicMesh3&& DynamicMesh, const TArray<UMaterialInterface*>& MaterialSet);
 
-	/** Create the PreviewAnimationInstance if the AnimationAsset and SkeletalMesh both exist, and set the animation to run on the SkeletalMeshComponent */
-	void UpdateSkeletalMeshComponent();
-
-	/** Update the data flow component with the dataflow asset and reset the render targets */
-	void UpdateDataflowComponent();
+	TArray<TObjectPtr<UDynamicMeshComponent>> GetDynamicMeshComponents() {return DynamicMeshComponents;}
 
 	/** Build the scene bounding box */
 	FBox GetBoundingBox() const;
@@ -80,24 +68,9 @@ private:
 	/** Data flows assets datas used to build the scene components */
 	TObjectPtr<UDataflowEditorContent> EditorContent = nullptr;
 	
-	/** Dataflow actor */
-	TObjectPtr<ADataflowActor> DataflowActor = nullptr;
-	
-	/** Skeletal mesh actor */
-	TObjectPtr<AActor> SkeletalMeshActor = nullptr;
-
 	/** Dynamic mesh actor */
 	TObjectPtr<AActor> DynamicMeshActor = nullptr;
 	
-	/** Dataflow component that will be modified */
-	TObjectPtr<UDataflowComponent> DataflowComponent = nullptr;
-
-	/** Skeletal mesh component used for potential binding */
-	TObjectPtr<USkeletalMeshComponent> SkeletalMeshComponent = nullptr;
-
-	/** Preview anim instance for the skeletal mesh */
-	TObjectPtr<UAnimSingleNodeInstance> PreviewAnimInstance = nullptr;
-
 	/** Dynamic mesh components generated by the tools for edit/manipulation*/
 	TArray<TObjectPtr<UDynamicMeshComponent>> DynamicMeshComponents;
 
