Dataflow - SkeletalMesh rendering
- Added skeletal mesh rendering of the owners skeletal mesh
- UI updates when the asset is changed.

diff --git a/Engine/Plugins/Experimental/Dataflow/Source/DataflowEditor/Private/Dataflow/DataflowEditor.cpp b/Engine/Plugins/Experimental/Dataflow/Source/DataflowEditor/Private/Dataflow/DataflowEditor.cpp
--- a/Engine/Plugins/Experimental/Dataflow/Source/DataflowEditor/Private/Dataflow/DataflowEditor.cpp
+++ b/Engine/Plugins/Experimental/Dataflow/Source/DataflowEditor/Private/Dataflow/DataflowEditor.cpp
@@ -30,40 +30,39 @@ void UDataflowEditor::Initialize(const TArray<TObjectPtr<UObject>>& InObjects)
 	check(Content);
 
 	TArray<TObjectPtr<UObject>> ObjectsToEdit;
-	Content->bHasValidSkeletalMesh = false;
 	if (ensure(InObjects.Num() == 1))
 	{
 		TObjectPtr<UObject> RootObject = InObjects[0];
-		Content->DataflowAsset = Cast<UDataflow>(RootObject);
-		if(!Content->DataflowAsset)
+		Content->SetDataflowAsset(Cast<UDataflow>(RootObject));
+		if(!Content->GetDataflowAsset())
 		{
-			Content->DataflowAsset = Private::GetDataflowAssetFrom(RootObject);
-			Content->DataflowTerminal = Private::GetDataflowTerminalFrom(RootObject);
-			Content->SkeletalMesh = Private::GetSkeletalMeshFrom(RootObject);
-			Content->AnimationAsset = Private::GetAnimationAssetFrom(RootObject);
+			Content->SetDataflowAsset(Private::GetDataflowAssetFrom(RootObject));
+			Content->SetDataflowTerminal(Private::GetDataflowTerminalFrom(RootObject));
+			Content->SetSkeletalMesh(Private::GetSkeletalMeshFrom(RootObject));
+			Content->SetSkeleton(Private::GetSkeletonFrom(RootObject));
+			Content->SetAnimationAsset(Private::GetAnimationAssetFrom(RootObject));
 		}
-		if(Content->DataflowAsset)
+		if(Content->GetDataflowAsset())
 		{
-			Content->DataflowOwner = RootObject;
-			Content->DataflowAsset->Schema = UDataflowSchema::StaticClass();
-			Content->DataflowContext = MakeShared<Dataflow::FEngineContext>(RootObject, Content->DataflowAsset, FPlatformTime::Cycles64());
-			Content->LastModifiedTimestamp = Content->DataflowContext->GetTimestamp();
+			Content->SetDataflowOwner(RootObject);
+			Content->GetDataflowAsset()->Schema = UDataflowSchema::StaticClass();
+			Content->SetDataflowContext(MakeShared<Dataflow::FEngineContext>(RootObject, Content->GetDataflowAsset(), FPlatformTime::Cycles64()));
+			Content->SetLastModifiedTimestamp(Content->GetDataflowContext()->GetTimestamp());
 
-			ObjectsToEdit.Add(Content->DataflowOwner);
+			ObjectsToEdit.Add(Content->GetDataflowOwner());
 
-			if(!Content->SkeletalMesh)
+			if(!Content->GetSkeletalMesh())
 			{
-				const FName SkeletonName = MakeUniqueObjectName(Content->DataflowAsset, UDataflow::StaticClass(), FName("USkeleton"));
-				const FName SkeletalMeshName = MakeUniqueObjectName(Content->DataflowAsset, UDataflow::StaticClass(), FName("USkeletalMesh"));
+				const FName SkeletonName = MakeUniqueObjectName(Content->GetDataflowAsset(), UDataflow::StaticClass(), FName("USkeleton"));
+				const FName SkeletalMeshName = MakeUniqueObjectName(Content->GetDataflowAsset(), UDataflow::StaticClass(), FName("USkeletalMesh"));
 
-				Content->Skeleton = NewObject<USkeleton>(Content->DataflowAsset, SkeletonName);
-				Content->SkeletalMesh = NewObject<USkeletalMesh>(Content->DataflowAsset, SkeletalMeshName);
-				Content->SkeletalMesh->SetSkeleton(Content->Skeleton);
+				Content->SetSkeleton(NewObject<USkeleton>(Content->GetDataflowAsset(), SkeletonName));
+				Content->SetSkeletalMesh(NewObject<USkeletalMesh>(Content->GetDataflowAsset(), SkeletalMeshName));
+				Content->GetSkeletalMesh()->SetSkeleton(Content->GetSkeleton());
 			}
-			else
+			else if(!Content->GetSkeleton())
 			{
-				Content->Skeleton = Content->SkeletalMesh->GetSkeleton();
-				Content->bHasValidSkeletalMesh = true;
+				Content->SetSkeleton(Content->GetSkeletalMesh()->GetSkeleton());
 			}
 		}
 	}
diff --git a/Engine/Plugins/Experimental/Dataflow/Source/DataflowEditor/Private/Dataflow/DataflowEditorCommands.cpp b/Engine/Plugins/Experimental/Dataflow/Source/DataflowEditor/Private/Dataflow/DataflowEditorCommands.cpp
--- a/Engine/Plugins/Experimental/Dataflow/Source/DataflowEditor/Private/Dataflow/DataflowEditorCommands.cpp
+++ b/Engine/Plugins/Experimental/Dataflow/Source/DataflowEditor/Private/Dataflow/DataflowEditorCommands.cpp
@@ -3,17 +3,19 @@
 #include "Dataflow/DataflowEditorCommands.h"
 
 #include "Dataflow/DataflowEdNode.h"
+#include "Dataflow/DataflowEditorContent.h"
+#include "Dataflow/DataflowEditorStyle.h"
+#include "Dataflow/DataflowEditorUtil.h"
+#include "Dataflow/DataflowGraphEditor.h"
 #include "Dataflow/DataflowNodeFactory.h"
 #include "Dataflow/DataflowObject.h"
 #include "Dataflow/DataflowOverrideNode.h"
-#include "Dataflow/DataflowEditorStyle.h"
-#include "EdGraph/EdGraphNode.h"
-#include "IStructureDetailsView.h"
-#include "EdGraphNode_Comment.h"
-#include "Editor.h"
-#include "Dataflow/DataflowGraphEditor.h"
 #include "Dataflow/DataflowSCommentNode.h"
 #include "Dataflow/DataflowSNode.h"
+#include "EdGraphNode_Comment.h"
+#include "EdGraph/EdGraphNode.h"
+#include "Editor.h"
+#include "IStructureDetailsView.h"
 
 #define LOCTEXT_NAMESPACE "DataflowEditorCommands"
 
@@ -270,27 +272,42 @@ void FDataflowEditorCommands::OnNodeTitleCommitted(const FText& InNewText, EText
 	}
 }
 
-void FDataflowEditorCommands::OnAssetPropertyValueChanged(UDataflow* Graph, TSharedPtr<Dataflow::FEngineContext>& Context, Dataflow::FTimestamp& OutLastNodeTimestamp, const FPropertyChangedEvent& InPropertyChangedEvent)
+void FDataflowEditorCommands::OnAssetPropertyValueChanged(TObjectPtr<UDataflowEditorContent> Content, const FPropertyChangedEvent& InPropertyChangedEvent)
 {
-	if (ensureMsgf(Graph != nullptr, TEXT("Warning : Failed to find valid graph.")))
+	if (Content)
 	{
-		if (InPropertyChangedEvent.ChangeType == EPropertyChangeType::ValueSet ||
-			InPropertyChangedEvent.ChangeType == EPropertyChangeType::ArrayRemove ||
-			InPropertyChangedEvent.ChangeType == EPropertyChangeType::ArrayClear)
+		TObjectPtr<UDataflow>& DataflowAsset = Content->GetDataflowAsset();
+		if (DataflowAsset)
 		{
-			if (InPropertyChangedEvent.GetPropertyName() == FName("Overrides_Key") ||
-				InPropertyChangedEvent.GetPropertyName() == FName("Overrides"))
+			if (InPropertyChangedEvent.ChangeType == EPropertyChangeType::ValueSet ||
+				InPropertyChangedEvent.ChangeType == EPropertyChangeType::ArrayRemove ||
+				InPropertyChangedEvent.ChangeType == EPropertyChangeType::ArrayClear)
 			{
-				for (const TSharedPtr<FDataflowNode>& DataflowNode : Graph->Dataflow->GetNodes())
+				if (InPropertyChangedEvent.GetPropertyName() == FName("Overrides_Key") ||
+					InPropertyChangedEvent.GetPropertyName() == FName("Overrides"))
 				{
-					if (DataflowNode->IsA(FDataflowOverrideNode::StaticType()))
+					if (ensureMsgf(DataflowAsset != nullptr, TEXT("Warning : Failed to find valid graph.")))
 					{
-						// TODO: For now we invalidate all the FDataflowOverrideNode nodes
-						// Once the Variable system will be in place only the neccessary nodes
-						// will be invalidated
-						DataflowNode->Invalidate();
+						for (const TSharedPtr<FDataflowNode>& DataflowNode : DataflowAsset->Dataflow->GetNodes())
+						{
+							if (DataflowNode->IsA(FDataflowOverrideNode::StaticType()))
+							{
+								// TODO: For now we invalidate all the FDataflowOverrideNode nodes
+								// Once the Variable system will be in place only the neccessary nodes
+								// will be invalidated
+								DataflowNode->Invalidate();
+							}
+						}
 					}
 				}
+				else if (InPropertyChangedEvent.GetPropertyName() == FName("SkeletalMesh"))
+				{
+					Content->SetSkeletalMesh(Private::GetSkeletalMeshFrom(Content->GetDataflowOwner()));
+				}
+				else if (InPropertyChangedEvent.GetPropertyName() == FName("Skeleton"))
+				{
+					Content->SetSkeleton(Private::GetSkeletonFrom(Content->GetDataflowOwner()));
+				}
 			}
 		}
 	}
@@ -303,7 +320,7 @@ void FDataflowEditorCommands::OnPropertyValueChanged(UDataflow* OutDataflow, TSh
 		TSharedPtr<const FDataflowNode> UpdatedNode = nullptr;
 		if (OutDataflow && InPropertyChangedEvent.Property && InPropertyChangedEvent.Property->GetOwnerUObject())
 		{
-//			OutDataflow->MarkPackageDirty();
+			//			OutDataflow->MarkPackageDirty();
 			OutDataflow->Modify();
 
 			for (UObject* SelectedNode : SelectedNodes)
diff --git a/Engine/Plugins/Experimental/Dataflow/Source/DataflowEditor/Private/Dataflow/DataflowEditorContent.h b/Engine/Plugins/Experimental/Dataflow/Source/DataflowEditor/Private/Dataflow/DataflowEditorContent.h
--- a/Engine/Plugins/Experimental/Dataflow/Source/DataflowEditor/Private/Dataflow/DataflowEditorContent.h
+++ b/Engine/Plugins/Experimental/Dataflow/Source/DataflowEditor/Private/Dataflow/DataflowEditorContent.h
@@ -30,14 +30,56 @@ public:
 	/** Check if the datas flow datas are valid */
 	bool IsValid() const { return DataflowOwner && DataflowAsset;}
 
+	/** Check if the datas flow datas are dirty and need to be re-renmdered */
+	bool IsDirty() const { return bIsDirty; }
+	void SetIsDirty(bool InDirty) { bIsDirty = InDirty; }
+
+	//bool HasCompleteSkeletalSetup() const { return SkeletalMesh != nullptr && Skeleton!=nullptr; }
+	void SetSkeletalMesh(TObjectPtr<USkeletalMesh> InMesh) { SkeletalMesh = InMesh; bIsDirty = true; }
+	TObjectPtr<USkeletalMesh>& GetSkeletalMesh() { return SkeletalMesh; }
+	const TObjectPtr<USkeletalMesh>& GetSkeletalMesh() const { return SkeletalMesh; }
+
 	/** Data flow object owner */
-	TObjectPtr<UObject> DataflowOwner = nullptr;
+	void SetDataflowOwner(TObjectPtr<UObject> InOwner) { DataflowOwner = InOwner; bIsDirty = true;}
+	TObjectPtr<UObject>& GetDataflowOwner() { return DataflowOwner; }
+	const TObjectPtr<UObject>& GetDataflowOwner() const { return DataflowOwner; }
 
 	/** Data flow asset that we will edit */
-	TObjectPtr<UDataflow> DataflowAsset = nullptr;
+	void SetDataflowAsset(TObjectPtr<UDataflow> InAsset) { DataflowAsset = InAsset;  bIsDirty = true;}
+	TObjectPtr<UDataflow>& GetDataflowAsset() { return DataflowAsset; }
+	const TObjectPtr<UDataflow>& GetDataflowAsset() const { return DataflowAsset; }
 
 	/** Data flow terminal path for evaluation */
-	FString DataflowTerminal = "";
+	void SetDataflowTerminal(FString InPath) { DataflowTerminal = InPath;  bIsDirty = true;}
+	FString GetDataflowTerminal() const { return DataflowTerminal; }
+
+	/** Data flow skeleton*/
+	void SetSkeleton(TObjectPtr<USkeleton> InSkeleton) { Skeleton = InSkeleton;  bIsDirty = true;}
+	TObjectPtr<USkeleton>& GetSkeleton() { return Skeleton; }
+	const TObjectPtr<USkeleton>& GetSkeleton() const { return Skeleton; }
+
+	/** Animation asset to be used to preview simulation */
+	void SetAnimationAsset(TObjectPtr<UAnimationAsset> InAnimation) { AnimationAsset = InAnimation; bIsDirty = true;}
+	TObjectPtr<UAnimationAsset>& GetAnimationAsset() { return AnimationAsset; }
+	const TObjectPtr<UAnimationAsset>& GetAnimationAsset() const { return AnimationAsset; }
+
+	/**  Engine context to be used for dataflow evaluation */
+	void SetDataflowContext(TSharedPtr<Dataflow::FEngineContext> InContext) { DataflowContext = InContext; bIsDirty = true;}
+	TSharedPtr<Dataflow::FEngineContext>& GetDataflowContext() { return DataflowContext; }
+	const TSharedPtr<Dataflow::FEngineContext>& GetDataflowContext() const { return DataflowContext; }
+
+	/** Last data flow evaluated node time stamp */
+	void SetLastModifiedTimestamp(Dataflow::FTimestamp InTimestamp) { LastModifiedTimestamp = InTimestamp; bIsDirty = true;}
+	Dataflow::FTimestamp& GetLastModifiedTimestamp() { return LastModifiedTimestamp; }
+	const Dataflow::FTimestamp& GetLastModifiedTimestamp() const { return LastModifiedTimestamp; }
+	
+
+private:
+	/** Data flow object owner */
+	TObjectPtr<UObject> DataflowOwner = nullptr;
+
+	/** Data flow asset that we will edit */
+	TObjectPtr<UDataflow> DataflowAsset = nullptr;
 
 	/** Data flow skeletal mesh*/
 	TObjectPtr<USkeletalMesh> SkeletalMesh = nullptr;
@@ -45,8 +87,11 @@ public:
 	/** Data flow skeleton*/
 	TObjectPtr<USkeleton> Skeleton = nullptr;
 
+	/** Data flow terminal path for evaluation */
+	FString DataflowTerminal = "";
+
 	/** Animation asset to be used to preview simulation */
-	TObjectPtr<UAnimationAsset> AnimationAsset;
+	TObjectPtr<UAnimationAsset> AnimationAsset = nullptr;
 
 	/**  Engine context to be used for dataflow evaluation */
 	TSharedPtr<Dataflow::FEngineContext> DataflowContext = nullptr;
@@ -54,6 +99,5 @@ public:
 	/** Last data flow evaluated node time stamp */
 	Dataflow::FTimestamp LastModifiedTimestamp = Dataflow::FTimestamp::Invalid;
 
-	/** Boolean to check if the skelmesh is valid*/
-	bool bHasValidSkeletalMesh = false;
+	bool bIsDirty = true;
 };
diff --git a/Engine/Plugins/Experimental/Dataflow/Source/DataflowEditor/Private/Dataflow/DataflowEditorToolkit.cpp b/Engine/Plugins/Experimental/Dataflow/Source/DataflowEditor/Private/Dataflow/DataflowEditorToolkit.cpp
--- a/Engine/Plugins/Experimental/Dataflow/Source/DataflowEditor/Private/Dataflow/DataflowEditorToolkit.cpp
+++ b/Engine/Plugins/Experimental/Dataflow/Source/DataflowEditor/Private/Dataflow/DataflowEditorToolkit.cpp
@@ -307,7 +307,7 @@ void FDataflowEditorToolkit::GetSaveableObjects(TArray<UObject*>& OutObjects) co
 
 	if (ensure(GetDataflowEditorContent()))
 	{
-		if (UDataflow* DataflowAsset = GetDataflowEditorContent()->DataflowAsset)
+		if (UDataflow* DataflowAsset = GetDataflowEditorContent()->GetDataflowAsset())
 		{
 			check(DataflowAsset->IsAsset());
 			OutObjects.Add(DataflowAsset);
@@ -326,12 +326,12 @@ void FDataflowEditorToolkit::CreateWidgets()
 
 	if (TObjectPtr<UDataflowEditorContent> Content = GetDataflowEditorContent())
 	{
-		if (UDataflow* DataflowAsset = Content->DataflowAsset)
+		if (UDataflow* DataflowAsset = Content->GetDataflowAsset())
 		{
-			NodeDetailsEditor = CreateNodeDetailsEditorWidget(Content->DataflowOwner);
-			AssetDetailsEditor = CreateAssetDetailsEditorWidget(Content->DataflowOwner);
+			NodeDetailsEditor = CreateNodeDetailsEditorWidget(Content->GetDataflowOwner());
+			AssetDetailsEditor = CreateAssetDetailsEditorWidget(Content->GetDataflowOwner());
 			GraphEditor = CreateGraphEditorWidget(DataflowAsset, NodeDetailsEditor);
-			SkeletalEditor = CreateSkeletalEditorWidget(Content->DataflowOwner);
+			SkeletalEditor = CreateSkeletalEditorWidget(Content->GetDataflowOwner());
 		}
 	}
 }
@@ -369,7 +369,12 @@ void FDataflowEditorToolkit::OnPropertyValueChanged(const FPropertyChangedEvent&
 	if (TObjectPtr<UDataflowEditorContent> EditorContent = GetDataflowEditorContent())
 	{
 		ensure(EditorContent);
-		FDataflowEditorCommands::OnPropertyValueChanged(EditorContent->DataflowAsset, EditorContent->DataflowContext, EditorContent->LastModifiedTimestamp, PropertyChangedEvent, PrevNodeSelection);
+		if (UDataflow* DataflowAsset = EditorContent->GetDataflowAsset())
+		{
+			TSharedPtr<Dataflow::FEngineContext>& DataflowContext = EditorContent->GetDataflowContext();
+			Dataflow::FTimestamp& LastNodeTimestamp = EditorContent->GetLastModifiedTimestamp();
+			FDataflowEditorCommands::OnPropertyValueChanged(DataflowAsset, DataflowContext, LastNodeTimestamp, PropertyChangedEvent, PrevNodeSelection);
+		}
 	}
 }
 
@@ -377,7 +382,7 @@ void FDataflowEditorToolkit::OnAssetPropertyValueChanged(const FPropertyChangedE
 {
 	if (TObjectPtr<UDataflowEditorContent> EditorContent = GetDataflowEditorContent())
 	{
-		ensure(EditorContent);	FDataflowEditorCommands::OnAssetPropertyValueChanged(EditorContent->DataflowAsset, EditorContent->DataflowContext, EditorContent->LastModifiedTimestamp, PropertyChangedEvent);
+		ensure(EditorContent);	FDataflowEditorCommands::OnAssetPropertyValueChanged(EditorContent, PropertyChangedEvent);
 	}
 }
 
@@ -393,7 +398,7 @@ void FDataflowEditorToolkit::OnNodeTitleCommitted(const FText& InNewText, ETextC
 
 void FDataflowEditorToolkit::OnNodeSelectionChanged(const TSet<UObject*>& NewSelection)
 {
-	if (TObjectPtr<UDataflowEditorContent> EditorContent = GetDataflowEditorContent(); EditorContent->DataflowAsset)
+	if (TObjectPtr<UDataflowEditorContent> EditorContent = GetDataflowEditorContent(); EditorContent->GetDataflowAsset())
 	{
 		// Only keep UDataflowEdNode from NewSelection
 		TSet<UObject*> ValidatedSelection;
@@ -465,13 +470,13 @@ void FDataflowEditorToolkit::Tick(float DeltaTime)
 	{
 		if (EditorContent->IsValid())
 		{
-			if (!EditorContent->DataflowContext)
+			if (!EditorContent->GetDataflowContext())
 			{
-				EditorContent->DataflowContext = MakeShared<Dataflow::FEngineContext>(EditorContent->DataflowOwner, EditorContent->DataflowAsset, Dataflow::FTimestamp::Invalid);
-				EditorContent->LastModifiedTimestamp = Dataflow::FTimestamp::Invalid;
+				EditorContent->GetDataflowContext() = MakeShared<Dataflow::FEngineContext>(EditorContent->GetDataflowOwner(), EditorContent->GetDataflowAsset(), Dataflow::FTimestamp::Invalid);
+				EditorContent->GetLastModifiedTimestamp() = Dataflow::FTimestamp::Invalid;
 			}
-			FDataflowEditorCommands::EvaluateTerminalNode(*EditorContent->DataflowContext.Get(), EditorContent->LastModifiedTimestamp,
-														  EditorContent->DataflowAsset, nullptr, nullptr, EditorContent->DataflowOwner, EditorContent->DataflowTerminal);
+			FDataflowEditorCommands::EvaluateTerminalNode(*EditorContent->GetDataflowContext().Get(), EditorContent->GetLastModifiedTimestamp(),
+														  EditorContent->GetDataflowAsset(), nullptr, nullptr, EditorContent->GetDataflowOwner(), EditorContent->GetDataflowTerminal());
 		}
 	}
 }
@@ -492,15 +497,15 @@ TSharedRef<SDataflowGraphEditor> FDataflowEditorToolkit::CreateGraphEditorWidget
 		{
 			if (EditorContent->IsValid())
 			{
-				if (!EditorContent->DataflowContext)
+				if (!EditorContent->GetDataflowContext())
 				{
-					EditorContent->DataflowContext = MakeShared<Dataflow::FEngineContext>(EditorContent->DataflowOwner, EditorContent->DataflowAsset, Dataflow::FTimestamp::Invalid);
+					EditorContent->GetDataflowContext() = MakeShared<Dataflow::FEngineContext>(EditorContent->GetDataflowOwner(), EditorContent->GetDataflowAsset(), Dataflow::FTimestamp::Invalid);
 				}
 				Node->Invalidate();
-				EditorContent->LastModifiedTimestamp = Dataflow::FTimestamp::Invalid;
+				EditorContent->GetLastModifiedTimestamp() = Dataflow::FTimestamp::Invalid;
 
-				FDataflowEditorCommands::EvaluateTerminalNode(*EditorContent->DataflowContext.Get(), EditorContent->LastModifiedTimestamp,
-															  EditorContent->DataflowAsset, Node, Out, EditorContent->DataflowOwner, EditorContent->DataflowTerminal);
+				FDataflowEditorCommands::EvaluateTerminalNode(*EditorContent->GetDataflowContext().Get(), EditorContent->GetLastModifiedTimestamp(),
+															  EditorContent->GetDataflowAsset(), Node, Out, EditorContent->GetDataflowOwner(), EditorContent->GetDataflowTerminal());
 			}
 		}
 	};
@@ -580,11 +585,11 @@ TSharedPtr<ISkeletonTree> FDataflowEditorToolkit::CreateSkeletalEditorWidget(UOb
 {
 	if (TObjectPtr<UDataflowEditorContent> EditorContent = GetDataflowEditorContent())
 	{
-		if (EditorContent->DataflowAsset)
+		if (EditorContent->GetDataflowAsset())
 		{
 			FSkeletonTreeArgs SkeletonTreeArgs;
 			ISkeletonEditorModule& SkeletonEditorModule = FModuleManager::LoadModuleChecked<ISkeletonEditorModule>("SkeletonEditor");
-			TSharedPtr<ISkeletonTree> SkeletonTree = SkeletonEditorModule.CreateSkeletonTree(EditorContent->Skeleton, SkeletonTreeArgs);
+			TSharedPtr<ISkeletonTree> SkeletonTree = SkeletonEditorModule.CreateSkeletonTree(EditorContent->GetSkeleton(), SkeletonTreeArgs);
 			return SkeletonTree;
 		}
 	}
@@ -632,7 +637,7 @@ TSharedRef<SDockTab> FDataflowEditorToolkit::SpawnTab_Skeletal(const FSpawnTabAr
 	check(Args.GetTabId() == SkeletalTabId);
 	if (const TObjectPtr<UDataflowEditorContent> EditorContent = GetDataflowEditorContent())
 	{
-		SkeletalEditor->SetSkeletalMesh(EditorContent->SkeletalMesh);
+		SkeletalEditor->SetSkeletalMesh(EditorContent->GetSkeletalMesh());
 	}
 
 	return SNew(SDockTab)
@@ -695,7 +700,7 @@ TSharedRef<SDockTab> FDataflowEditorToolkit::SpawnTab_SelectionView(const FSpawn
 				DataflowSelectionView_1->SetSelectionView(SelectionViewWidget);
 
 				// Set the Context on the interface
-				if (TSharedPtr<Dataflow::FContext> CurrentContext = EditorContent->DataflowContext)
+				if (TSharedPtr<Dataflow::FContext> CurrentContext = EditorContent->GetDataflowContext())
 				{
 					DataflowSelectionView_1->SetContext(CurrentContext);
 				}
@@ -705,7 +710,7 @@ TSharedRef<SDockTab> FDataflowEditorToolkit::SpawnTab_SelectionView(const FSpawn
 				DataflowSelectionView_2->SetSelectionView(SelectionViewWidget);
 
 				// Set the Context on the interface
-				if (TSharedPtr<Dataflow::FContext> CurrentContext = EditorContent->DataflowContext)
+				if (TSharedPtr<Dataflow::FContext> CurrentContext = EditorContent->GetDataflowContext())
 				{
 					DataflowSelectionView_2->SetContext(CurrentContext);
 				}
@@ -715,7 +720,7 @@ TSharedRef<SDockTab> FDataflowEditorToolkit::SpawnTab_SelectionView(const FSpawn
 				DataflowSelectionView_3->SetSelectionView(SelectionViewWidget);
 
 				// Set the Context on the interface
-				if (TSharedPtr<Dataflow::FContext> CurrentContext = EditorContent->DataflowContext)
+				if (TSharedPtr<Dataflow::FContext> CurrentContext = EditorContent->GetDataflowContext())
 				{
 					DataflowSelectionView_3->SetContext(CurrentContext);
 				}
@@ -725,7 +730,7 @@ TSharedRef<SDockTab> FDataflowEditorToolkit::SpawnTab_SelectionView(const FSpawn
 				DataflowSelectionView_4->SetSelectionView(SelectionViewWidget);
 
 				// Set the Context on the interface
-				if (TSharedPtr<Dataflow::FContext> CurrentContext = EditorContent->DataflowContext)
+				if (TSharedPtr<Dataflow::FContext> CurrentContext = EditorContent->GetDataflowContext())
 				{
 					DataflowSelectionView_4->SetContext(CurrentContext);
 				}
@@ -789,7 +794,7 @@ TSharedRef<SDockTab> FDataflowEditorToolkit::SpawnTab_CollectionSpreadSheet(cons
 				DataflowCollectionSpreadSheet_1->SetCollectionSpreadSheet(CollectionSpreadSheetWidget);
 
 				// Set the Context on the interface
-				if (TSharedPtr<Dataflow::FContext> CurrentContext = EditorContent->DataflowContext)
+				if (TSharedPtr<Dataflow::FContext> CurrentContext = EditorContent->GetDataflowContext())
 				{
 					DataflowCollectionSpreadSheet_1->SetContext(CurrentContext);
 				}
@@ -799,7 +804,7 @@ TSharedRef<SDockTab> FDataflowEditorToolkit::SpawnTab_CollectionSpreadSheet(cons
 				DataflowCollectionSpreadSheet_2->SetCollectionSpreadSheet(CollectionSpreadSheetWidget);
 
 				// Set the Context on the interface
-				if (TSharedPtr<Dataflow::FContext> CurrentContext = EditorContent->DataflowContext)
+				if (TSharedPtr<Dataflow::FContext> CurrentContext = EditorContent->GetDataflowContext())
 				{
 					DataflowCollectionSpreadSheet_2->SetContext(CurrentContext);
 				}
@@ -809,7 +814,7 @@ TSharedRef<SDockTab> FDataflowEditorToolkit::SpawnTab_CollectionSpreadSheet(cons
 				DataflowCollectionSpreadSheet_3->SetCollectionSpreadSheet(CollectionSpreadSheetWidget);
 
 				// Set the Context on the interface
-				if (TSharedPtr<Dataflow::FContext> CurrentContext = EditorContent->DataflowContext)
+				if (TSharedPtr<Dataflow::FContext> CurrentContext = EditorContent->GetDataflowContext())
 				{
 					DataflowCollectionSpreadSheet_3->SetContext(CurrentContext);
 				}
@@ -819,7 +824,7 @@ TSharedRef<SDockTab> FDataflowEditorToolkit::SpawnTab_CollectionSpreadSheet(cons
 				DataflowCollectionSpreadSheet_4->SetCollectionSpreadSheet(CollectionSpreadSheetWidget);
 
 				// Set the Context on the interface
-				if (TSharedPtr<Dataflow::FContext> CurrentContext = EditorContent->DataflowContext)
+				if (TSharedPtr<Dataflow::FContext> CurrentContext = EditorContent->GetDataflowContext())
 				{
 					DataflowCollectionSpreadSheet_4->SetContext(CurrentContext);
 				}
@@ -968,13 +973,13 @@ FText FDataflowEditorToolkit::GetToolkitName() const
 {
 	if (TObjectPtr<const UDataflowEditorContent> EditorContent = GetDataflowEditorContent())
 	{
-		if (EditorContent->DataflowOwner)
+		if (EditorContent->GetDataflowOwner())
 		{
-			return  GetLabelForObject(EditorContent->DataflowOwner);
+			return  GetLabelForObject(EditorContent->GetDataflowOwner());
 		}
-		else if (EditorContent->DataflowAsset)
+		else if (EditorContent->GetDataflowAsset())
 		{
-			return  GetLabelForObject(EditorContent->DataflowAsset);
+			return  GetLabelForObject(EditorContent->GetDataflowAsset());
 		}
 	}
 	return  LOCTEXT("ToolkitName", "Empty Dataflow Editor");
diff --git a/Engine/Plugins/Experimental/Dataflow/Source/DataflowEditor/Private/Dataflow/DataflowEditorUtil.cpp b/Engine/Plugins/Experimental/Dataflow/Source/DataflowEditor/Private/Dataflow/DataflowEditorUtil.cpp
--- a/Engine/Plugins/Experimental/Dataflow/Source/DataflowEditor/Private/Dataflow/DataflowEditorUtil.cpp
+++ b/Engine/Plugins/Experimental/Dataflow/Source/DataflowEditor/Private/Dataflow/DataflowEditorUtil.cpp
@@ -10,6 +10,7 @@
 #include "Dataflow/DataflowObject.h"
 #include "DynamicMesh/MeshNormals.h"
 #include "Engine/SkeletalMesh.h"
+#include "Animation/Skeleton.h"
 #include "Materials/Material.h"
 
 
@@ -40,6 +41,18 @@ namespace Private
 		}
 		return nullptr;
 	}
+
+	USkeleton* GetSkeletonFrom(UObject* InObject)
+	{
+		if (UClass* Class = InObject->GetClass())
+		{
+			if (FProperty* Property = Class->FindPropertyByName(FName("Skeleton")))
+			{
+				return *Property->ContainerPtrToValuePtr<USkeleton*>(InObject);
+			}
+		}
+		return nullptr;
+	}
 	
 	UAnimationAsset* GetAnimationAssetFrom(UObject* InObject)
 	{
@@ -169,11 +182,11 @@ namespace Dataflow
 	{
 		if (Content)
 		{
-			if (!Content->DataflowContext)
+			if (!Content->GetDataflowContext())
 			{
-				Content->DataflowContext = MakeShared<FEngineContext>(Content->DataflowOwner, Content->DataflowAsset, FTimestamp::Invalid);
+				Content->SetDataflowContext(MakeShared<FEngineContext>(Content->GetDataflowOwner(), Content->GetDataflowAsset(), FTimestamp::Invalid));
 			}
-			return Content->DataflowContext;
+			return Content->GetDataflowContext();
 		}
 
 		ensure(false);
diff --git a/Engine/Plugins/Experimental/Dataflow/Source/DataflowEditor/Private/Dataflow/DataflowEditorUtil.h b/Engine/Plugins/Experimental/Dataflow/Source/DataflowEditor/Private/Dataflow/DataflowEditorUtil.h
--- a/Engine/Plugins/Experimental/Dataflow/Source/DataflowEditor/Private/Dataflow/DataflowEditorUtil.h
+++ b/Engine/Plugins/Experimental/Dataflow/Source/DataflowEditor/Private/Dataflow/DataflowEditorUtil.h
@@ -13,6 +13,7 @@ class UDataflow;
 class UObject;
 class UDataflowEditorContent;
 class USkeletalMesh;
+class USkeleton;
 class UAnimationAsset;
 class UMaterial;
 
@@ -22,6 +23,8 @@ namespace Private
 
 	USkeletalMesh* GetSkeletalMeshFrom(UObject* InObject);
 
+	USkeleton* GetSkeletonFrom(UObject* InObject);
+
 	UAnimationAsset* GetAnimationAssetFrom(UObject* InObject);
 
 	FString GetDataflowTerminalFrom(UObject* InObject);
diff --git a/Engine/Plugins/Experimental/Dataflow/Source/DataflowEditor/Private/Dataflow/DataflowEditorViewportClient.cpp b/Engine/Plugins/Experimental/Dataflow/Source/DataflowEditor/Private/Dataflow/DataflowEditorViewportClient.cpp
--- a/Engine/Plugins/Experimental/Dataflow/Source/DataflowEditor/Private/Dataflow/DataflowEditorViewportClient.cpp
+++ b/Engine/Plugins/Experimental/Dataflow/Source/DataflowEditor/Private/Dataflow/DataflowEditorViewportClient.cpp
@@ -53,17 +53,21 @@ void FDataflowEditorViewportClient::Tick(float DeltaSeconds)
 	{
 		if (TObjectPtr<UDataflowEditorContent> EditorContent = PreviewScene->GetDataflowEditorContent())
 		{
-			if (TSharedPtr<Dataflow::FContext> Context = EditorContent->DataflowContext)
+			if (TSharedPtr<Dataflow::FContext> Context = EditorContent->GetDataflowContext())
 			{
-				if (const UDataflow* Dataflow = EditorContent->DataflowAsset)
+				if (const UDataflow* Dataflow = EditorContent->GetDataflowAsset())
 				{
 					const Dataflow::FTimestamp SystemTimestamp = LatestTimestamp(Dataflow, Context.Get());
-					if (SystemTimestamp >= EditorContent->LastModifiedTimestamp)
+					if (SystemTimestamp >= EditorContent->GetLastModifiedTimestamp())
 					{
-						PreviewScene->Update();
-						EditorContent->LastModifiedTimestamp = LatestTimestamp(EditorContent->DataflowAsset, EditorContent->DataflowContext.Get()).Value + 1;
+						EditorContent->SetLastModifiedTimestamp(LatestTimestamp(EditorContent->GetDataflowAsset(), EditorContent->GetDataflowContext().Get()).Value + 1);
 					}
 				}
+
+				if (EditorContent->IsDirty())
+				{
+					PreviewScene->Update();
+				}
 			}
 		}
 	}
diff --git a/Engine/Plugins/Experimental/Dataflow/Source/DataflowEditor/Private/Dataflow/DataflowPreviewScene.cpp b/Engine/Plugins/Experimental/Dataflow/Source/DataflowEditor/Private/Dataflow/DataflowPreviewScene.cpp
--- a/Engine/Plugins/Experimental/Dataflow/Source/DataflowEditor/Private/Dataflow/DataflowPreviewScene.cpp
+++ b/Engine/Plugins/Experimental/Dataflow/Source/DataflowEditor/Private/Dataflow/DataflowPreviewScene.cpp
@@ -22,19 +22,40 @@ FDataflowPreviewScene::FDataflowPreviewScene(FPreviewScene::ConstructionValues C
 	: FAdvancedPreviewScene(ConstructionValues), EditorContent(InEditorContent)
 {
 	check(EditorContent);
+	SkeletalMeshActor = GetWorld()->SpawnActor<AActor>(AActor::StaticClass());
+	if (EditorContent->GetSkeletalMesh())
+	{
+		SkeletalMeshComponent = NewObject<USkeletalMeshComponent>(SkeletalMeshActor);
+		SkeletalMeshComponent->SelectionOverrideDelegate = UPrimitiveComponent::FSelectionOverride::CreateRaw(this, &FDataflowPreviewScene::IsComponentSelected);
+		SkeletalMeshComponent->SetDisablePostProcessBlueprint(true);
+		UpdateSkeletalMeshComponent();
+	}
+	SkeletalMeshActor->RegisterAllComponents();
+
 	DynamicMeshActor = GetWorld()->SpawnActor<AActor>(AActor::StaticClass());	
 	DynamicMeshActor->RegisterAllComponents();
+
 	SetFloorVisibility(false, true);
 }
 
 FDataflowPreviewScene::~FDataflowPreviewScene()
 {
+	if (SkeletalMeshComponent)
+	{
+		SkeletalMeshComponent->TransformUpdated.RemoveAll(this);
+		SkeletalMeshComponent->SelectionOverrideDelegate.Unbind();
+		SkeletalMeshComponent->UnregisterComponent();
+		SkeletalMeshComponent->DestroyComponent();
+	}
+
 	ResetDynamicMeshComponents();
 }
 
 void FDataflowPreviewScene::AddReferencedObjects(FReferenceCollector& Collector)
 {
 	FAdvancedPreviewScene::AddReferencedObjects(Collector);
+	Collector.AddReferencedObject(SkeletalMeshComponent);
+	Collector.AddReferencedObject(SkeletalMeshActor);
 	Collector.AddReferencedObject(EditorContent);
 	Collector.AddReferencedObject(DynamicMeshActor);
 	Collector.AddReferencedObjects(DynamicMeshComponents);
@@ -44,6 +65,11 @@ void FDataflowPreviewScene::Update()
 {
 	using namespace UE::Geometry;//FDynamicMesh3
 
+	// Update the SkeletalMeshComponent for animation 
+	// changes.
+	UpdateSkeletalMeshComponent();
+
+
 	// The preview scene for the construction view will be
 	// cleared and rebuilt from scratch. This will genrate a 
 	// list of UPrimitiveComponents for rendering.
@@ -51,8 +77,8 @@ void FDataflowPreviewScene::Update()
 
 	if (EditorContent)
 	{
-		TObjectPtr<UDataflow> DataflowAsset = EditorContent->DataflowAsset;
-		TSharedPtr<Dataflow::FEngineContext> DataflowContext = EditorContent->DataflowContext;
+		TObjectPtr<UDataflow> DataflowAsset = EditorContent->GetDataflowAsset();
+		TSharedPtr<Dataflow::FEngineContext> DataflowContext = EditorContent->GetDataflowContext();
 		if(DataflowAsset && DataflowContext)
 		{
 			for (const UDataflowEdNode* Target : DataflowAsset->GetRenderTargets())
@@ -70,6 +96,8 @@ void FDataflowPreviewScene::Update()
 				}
 			}
 		}
+
+		EditorContent->SetIsDirty(false);
 	}
 }
 
@@ -108,9 +136,58 @@ TObjectPtr<UDynamicMeshComponent>& FDataflowPreviewScene::AddDynamicMeshComponen
 	return DynamicMeshComponents[ElementIndex];
 }
 
+void FDataflowPreviewScene::UpdateSkeletalMeshComponent()
+{
+	if (SkeletalMeshComponent)
+	{
+		if (EditorContent->GetSkeletalMesh())
+		{
+			if (EditorContent->GetSkeletalMesh() != SkeletalMeshComponent->GetSkeletalMeshAsset())
+			{
+				SkeletalMeshComponent->SetSkeletalMeshAsset(EditorContent->GetSkeletalMesh());
+			}
+
+			if (EditorContent->GetAnimationAsset())
+			{
+				PreviewAnimInstance = NewObject<UAnimSingleNodeInstance>(SkeletalMeshComponent);
+				PreviewAnimInstance->SetAnimationAsset(EditorContent->GetAnimationAsset());
+				SkeletalMeshComponent->SetAnimationMode(EAnimationMode::AnimationSingleNode);
+				SkeletalMeshComponent->InitAnim(true);
+				SkeletalMeshComponent->AnimationData.PopulateFrom(PreviewAnimInstance);
+				SkeletalMeshComponent->AnimScriptInstance = PreviewAnimInstance;
+				SkeletalMeshComponent->AnimScriptInstance->InitializeAnimation();
+				SkeletalMeshComponent->ValidateAnimation();
+			}
+			else
+			{
+				SkeletalMeshComponent->Stop();
+				SkeletalMeshComponent->AnimationData = FSingleAnimationPlayData();
+				SkeletalMeshComponent->AnimScriptInstance = nullptr;
+			}
+		}
+		else
+		{
+			SkeletalMeshComponent->SetSkeletalMeshAsset(nullptr);
+			SkeletalMeshComponent->Stop();
+			SkeletalMeshComponent->AnimationData = FSingleAnimationPlayData();
+			SkeletalMeshComponent->AnimScriptInstance = nullptr;
+		}
+
+		SkeletalMeshComponent->UpdateBounds();
+	}
+}
+
 FBox FDataflowPreviewScene::GetBoundingBox() const
 {
-	FBox SceneBounds(ForceInitToZero);
+	FBox SceneBounds(EForceInit::ForceInit);
+
+	if (SkeletalMeshComponent)
+	{
+		FTransform ComponentTransform = SkeletalMeshComponent->GetComponentTransform();
+		FBox LocalBox = SkeletalMeshComponent->GetLocalBounds().GetBox();
+		SceneBounds += LocalBox.TransformBy(ComponentTransform);
+	}
+
 	for (const TObjectPtr<UDynamicMeshComponent>& MeshComponent : DynamicMeshComponents)
 	{
 		if (MeshComponent)
@@ -118,6 +195,7 @@ FBox FDataflowPreviewScene::GetBoundingBox() const
 			SceneBounds += MeshComponent->Bounds.GetBox();
 		}
 	}
+
 	return SceneBounds;
 }
 
diff --git a/Engine/Plugins/Experimental/Dataflow/Source/DataflowEditor/Private/Dataflow/DataflowToolTarget.cpp b/Engine/Plugins/Experimental/Dataflow/Source/DataflowEditor/Private/Dataflow/DataflowToolTarget.cpp
--- a/Engine/Plugins/Experimental/Dataflow/Source/DataflowEditor/Private/Dataflow/DataflowToolTarget.cpp
+++ b/Engine/Plugins/Experimental/Dataflow/Source/DataflowEditor/Private/Dataflow/DataflowToolTarget.cpp
@@ -161,7 +161,7 @@ bool UDataflowReadOnlyToolTargetFactory::CanBuildTarget(UObject* SourceObject, c
 	// just add another factory that allows that class specifically(but make sure that
 	// GetMeshDescription and such work properly)
 	UDataflowEditorContent* Content = CastChecked<UDataflowEditorContent>(SourceObject);
-	const UDataflow* Dataflow = Content->DataflowAsset;
+	const UDataflow* Dataflow = Content->GetDataflowAsset();
 
 
 
@@ -176,8 +176,8 @@ UToolTarget* UDataflowReadOnlyToolTargetFactory::BuildTarget(UObject* SourceObje
 	UDataflowEditorContent* Content = CastChecked<UDataflowEditorContent>(SourceObject);
 
 	UDataflowReadOnlyToolTarget* Target = NewObject<UDataflowReadOnlyToolTarget>();
-	Target->Asset = Content->DataflowOwner;
-	Target->Dataflow = Content->DataflowAsset;
+	Target->Asset = Content->GetDataflowOwner();
+	Target->Dataflow = Content->GetDataflowAsset();
 	Target->Context = Dataflow::GetContext(Content);
 
 	// @todo(brice) : I needed to comment this out?
@@ -195,7 +195,7 @@ bool UDataflowToolTargetFactory::CanBuildTarget(UObject* SourceObject, const FTo
 	// just add another factory that allows that class specifically(but make sure that
 	// GetMeshDescription and such work properly)
 	UDataflowEditorContent* Content = CastChecked<UDataflowEditorContent>(SourceObject);
-	const UDataflow* Dataflow = Content->DataflowAsset;
+	const UDataflow* Dataflow = Content->GetDataflowAsset();
 
 	return Dataflow &&
 		ExactCast<UDataflow>(Cast<UDataflow>(Dataflow)) &&
@@ -207,8 +207,8 @@ UToolTarget* UDataflowToolTargetFactory::BuildTarget(UObject* SourceObject, cons
 {
 	UDataflowEditorContent* Content = CastChecked<UDataflowEditorContent>(SourceObject);
 	UDataflowToolTarget* Target = NewObject<UDataflowToolTarget>();
-	Target->Asset = Content->DataflowOwner;
-	Target->Dataflow = Content->DataflowAsset;
+	Target->Asset = Content->GetDataflowOwner();
+	Target->Dataflow = Content->GetDataflowAsset();
 	Target->Context = Dataflow::GetContext(Content);
 	//checkSlow(Target->Component.IsValid() && Requirements.AreSatisfiedBy(Target));
 
diff --git a/Engine/Plugins/Experimental/Dataflow/Source/DataflowEditor/Public/Dataflow/DataflowEditorCommands.h b/Engine/Plugins/Experimental/Dataflow/Source/DataflowEditor/Public/Dataflow/DataflowEditorCommands.h
--- a/Engine/Plugins/Experimental/Dataflow/Source/DataflowEditor/Public/Dataflow/DataflowEditorCommands.h
+++ b/Engine/Plugins/Experimental/Dataflow/Source/DataflowEditor/Public/Dataflow/DataflowEditorCommands.h
@@ -16,6 +16,7 @@ class UDataflowEdNode;
 struct FDataflowNode;
 class UEdGraphNode;
 class SDataflowGraphEditor;
+class UDataflowEditorContent;
 
 typedef TSet<class UObject*> FGraphPanelSelectionSet;
 
@@ -101,7 +102,7 @@ public:
 	*  OnPropertyValueChanged
 	*/
 	static void OnPropertyValueChanged(UDataflow* Graph, TSharedPtr<Dataflow::FEngineContext>& Context, Dataflow::FTimestamp& OutLastNodeTimestamp, const FPropertyChangedEvent& PropertyChangedEvent, const TSet<UObject*>& SelectedNodes = TSet<UObject*>());
-	static void OnAssetPropertyValueChanged(UDataflow* Graph, TSharedPtr<Dataflow::FEngineContext>& Context, Dataflow::FTimestamp& OutLastNodeTimestamp, const FPropertyChangedEvent& PropertyChangedEvent);
+	static void OnAssetPropertyValueChanged(TObjectPtr<UDataflowEditorContent> Content, const FPropertyChangedEvent& PropertyChangedEvent);
 
 	/*
 	*  OnSelectedNodesChanged
diff --git a/Engine/Plugins/Experimental/Dataflow/Source/DataflowEditor/Public/Dataflow/DataflowPreviewScene.h b/Engine/Plugins/Experimental/Dataflow/Source/DataflowEditor/Public/Dataflow/DataflowPreviewScene.h
--- a/Engine/Plugins/Experimental/Dataflow/Source/DataflowEditor/Public/Dataflow/DataflowPreviewScene.h
+++ b/Engine/Plugins/Experimental/Dataflow/Source/DataflowEditor/Public/Dataflow/DataflowPreviewScene.h
@@ -37,18 +37,28 @@ public:
 
 	void Update();
 
+	/** Update Skeletal Mesh Component */
+	void UpdateSkeletalMeshComponent();
 
-	// FGCObject interface
+	/** FGCObject interface */
 	virtual void AddReferencedObjects(FReferenceCollector& Collector) override;
 
-	// Modify/Get the data flow datas
+	/** Modify/Get the data flow datas */
 	TObjectPtr<UDataflowEditorContent> GetDataflowEditorContent() { return EditorContent; }
 	const TObjectPtr<UDataflowEditorContent> GetDataflowEditorContent() const { return EditorContent; }
 
-	// Modify/Get the data flow mode manager
+	/** Modify/Get the data flow mode manager */
 	TSharedPtr<FAssetEditorModeManager>& GetDataflowModeManager() { return DataflowModeManager; }
 	const TSharedPtr<FAssetEditorModeManager>& GetDataflowModeManager() const { return DataflowModeManager; }
 	
+	/** Modify/Get the preview anim instance */
+	UAnimSingleNodeInstance* GetPreviewAnimInstance() { return PreviewAnimInstance; }
+	const UAnimSingleNodeInstance* const GetPreviewAnimInstance() const { return PreviewAnimInstance; }
+
+	/** Modify/Get the skeletal mesh component */
+	USkeletalMeshComponent* GetSkeletalMeshComponent() { return SkeletalMeshComponent; }
+	const USkeletalMeshComponent* GetSkeletalMeshComponent() const { return SkeletalMeshComponent; }
+
 	/** Reset all the dynamic mesh components */
 	void ResetDynamicMeshComponents();
 	
@@ -70,10 +80,19 @@ private:
 	
 	/** Dynamic mesh actor */
 	TObjectPtr<AActor> DynamicMeshActor = nullptr;
-	
+
 	/** Dynamic mesh components generated by the tools for edit/manipulation*/
 	TArray<TObjectPtr<UDynamicMeshComponent>> DynamicMeshComponents;
 
+	/** Skeletal mesh actor */
+	TObjectPtr<AActor> SkeletalMeshActor = nullptr;
+
+	/** Skeletal mesh component used for potential binding */
+	TObjectPtr<USkeletalMeshComponent> SkeletalMeshComponent = nullptr;
+
+	/** Preview anim instance for the skeletal mesh */
+	TObjectPtr<UAnimSingleNodeInstance> PreviewAnimInstance = nullptr;
+
 	/** Mode Manager for selection */
 	TSharedPtr<FAssetEditorModeManager> DataflowModeManager;
 };
