Dataflow : Fix typo in file path.

diff --git a/Engine/Plugins/Experimental/Dataflow/Source/DataflowEditor/Private/Dataflow/DataflowEditorStyle.h b/Engine/Plugins/Experimental/Dataflow/Source/DataflowEditor/Private/Dataflow/DataflowEditorStyle.h
--- a/Engine/Plugins/Experimental/Dataflow/Source/DataflowEditor/Private/Dataflow/DataflowEditorStyle.h
+++ b/Engine/Plugins/Experimental/Dataflow/Source/DataflowEditor/Private/Dataflow/DataflowEditorStyle.h
@@ -37,10 +37,10 @@ public:
 		Set("Dataflow.SelectVertex", new FSlateImageBrush(RootToContentDir(TEXT("Slate/Dataflow_SelectVertex_40x.png")), Icon40x40));
 
 		const FString AttributeEditorPropertyName = "DataflowEditor." + FDataflowEditorCommandsImpl::BeginAttributeEditorToolIdentifier;
-		Set(*AttributeEditorPropertyName, new FSlateImageBrush(RootToContentDir(TEXT("Slate/Dataflow_SelectObject20x.png")), Icon20x20));
+		Set(*AttributeEditorPropertyName, new FSlateImageBrush(RootToContentDir(TEXT("Slate/Dataflow_SelectObject_20x.png")), Icon20x20));
 
 		const FString MeshSelectionPropertyName = "DataflowEditor." + FDataflowEditorCommandsImpl::BeginMeshSelectionToolIdentifier;
-		Set(*MeshSelectionPropertyName, new FSlateImageBrush(RootToContentDir(TEXT("Slate/Dataflow_SelectVertex20x.png")), Icon20x20));
+		Set(*MeshSelectionPropertyName, new FSlateImageBrush(RootToContentDir(TEXT("Slate/Dataflow_SelectVertex_20x.png")), Icon20x20));
 
 		DefaultMaterial = Cast<UMaterial>(StaticLoadObject(UMaterial::StaticClass(), NULL, TEXT("/Engine/BasicShapes/BasicShapeMaterial")));
 
