Dataflow : Support for painting weights

diff --git a/Engine/Plugins/Experimental/Dataflow/Dataflow.uplugin b/Engine/Plugins/Experimental/Dataflow/Dataflow.uplugin
--- a/Engine/Plugins/Experimental/Dataflow/Dataflow.uplugin
+++ b/Engine/Plugins/Experimental/Dataflow/Dataflow.uplugin
@@ -21,6 +21,11 @@
 			"Type": "Editor",
 			"LoadingPhase": "Default"
 		},
+		{
+			"Name": "DataflowAssetTools",
+			"Type": "Editor",
+			"LoadingPhase": "Default"
+		},
 		{
 			"Name": "DataflowEnginePlugin",
 			"Type": "Runtime",
diff --git a/Engine/Plugins/Experimental/Dataflow/Source/DataflowAssetTools/DataflowAssetTools.Build.cs b/Engine/Plugins/Experimental/Dataflow/Source/DataflowAssetTools/DataflowAssetTools.Build.cs
new file mode 100644
--- /dev/null
+++ b/Engine/Plugins/Experimental/Dataflow/Source/DataflowAssetTools/DataflowAssetTools.Build.cs
@@ -0,0 +1,35 @@
+// Copyright Epic Games, Inc. All Rights Reserved.
+
+using UnrealBuildTool;
+
+public class DataflowAssetTools : ModuleRules
+{
+	public DataflowAssetTools(ReadOnlyTargetRules Target) : base(Target)
+	{
+		PCHUsage = ModuleRules.PCHUsageMode.UseExplicitOrSharedPCHs;
+
+		PublicDependencyModuleNames.AddRange(
+			new string[] {
+			"MeshConversion",
+			}
+		);
+
+		PrivateDependencyModuleNames.AddRange(
+			new string[] {
+				"Core",
+				"Engine",
+				"CoreUObject",
+				"RenderCore",
+				"RHI",
+				"Chaos",
+				"DataflowCore",
+				"DataflowEnginePlugin",
+				"SlateCore",
+				"Projects",
+				"UnrealEd",
+				"GeometryCore",					// For DynamicMesh conversion
+				"SkeletalMeshDescription"		// For FSkeletalMeshAttributes::DefaultSkinWeightProfileName
+			}
+		);
+	}
+}
diff --git a/Engine/Plugins/Experimental/Dataflow/Source/DataflowAssetTools/Private/Dataflow/CollectionPatternToDynamicMesh.cpp b/Engine/Plugins/Experimental/Dataflow/Source/DataflowAssetTools/Private/Dataflow/CollectionPatternToDynamicMesh.cpp
new file mode 100644
--- /dev/null
+++ b/Engine/Plugins/Experimental/Dataflow/Source/DataflowAssetTools/Private/Dataflow/CollectionPatternToDynamicMesh.cpp
@@ -0,0 +1,48 @@
+// Copyright Epic Games, Inc. All Rights Reserved. 
+
+#if WITH_EDITOR
+#include "Dataflow/CollectionPatternToDynamicMesh.h"
+#include "DynamicMesh/MeshNormals.h"
+#include "DynamicMesh/NonManifoldMappingSupport.h"
+#include "Engine/SkeletalMesh.h"
+#include "SkeletalMeshAttributes.h"
+#include "ToDynamicMesh.h"
+
+namespace Dataflow
+{
+
+	void FCollectionPatternToDynamicMesh::Convert(const FManagedArrayCollection* FleshCollection, int32 PatternIndex, EDataflowPatternVertexType VertexDataType, UE::Geometry::FDynamicMesh3& MeshOut, bool bDisableAttributes)
+	{
+		// @todo(brice)(do the conversion)
+		check(0);
+	}
+
+	void FCollectionPatternToDynamicMesh::Convert(const UObject* FleshAssetIn, int32 LODIndex, int32 PatternIndex, EDataflowPatternVertexType VertexDataType, UE::Geometry::FDynamicMesh3& MeshOut)
+	{
+		// @todo(brice)(make this generic)
+		//Convert(FleshAssetIn->GetCollection(), PatternIndex, VertexDataType, MeshOut);
+	}
+
+}	// namespace Dataflow
+
+#else
+
+namespace Dataflow
+{
+
+	void FCollectionPatternToDynamicMesh::Convert(const FManagedArrayCollection* FleshCollection, int32 PatternIndex, EDataflowPatternVertexType VertexDataType, UE::Geometry::FDynamicMesh3& MeshOut, bool bDisableAttributes)
+	{
+		// Conversion only supported with editor.
+		check(0);
+	}
+
+	void FCollectionPatternToDynamicMesh::Convert(const UObject* AssetIn, int32 LODIndex, int32 PatternIndex, EDataflowPatternVertexType VertexDataType, UE::Geometry::FDynamicMesh3& MeshOut)
+	{
+		// Conversion only supported with editor.
+		check(0);
+	}
+
+}	// namespace Dataflow
+
+#endif  // end with editor
+
diff --git a/Engine/Plugins/Experimental/Dataflow/Source/DataflowAssetTools/Private/Dataflow/CollectionRenderingPatternUtility.cpp b/Engine/Plugins/Experimental/Dataflow/Source/DataflowAssetTools/Private/Dataflow/CollectionRenderingPatternUtility.cpp
new file mode 100644
--- /dev/null
+++ b/Engine/Plugins/Experimental/Dataflow/Source/DataflowAssetTools/Private/Dataflow/CollectionRenderingPatternUtility.cpp
@@ -0,0 +1,101 @@
+// Copyright Epic Games, Inc. All Rights Reserved. 
+
+#include "Dataflow/CollectionRenderingPatternUtility.h"
+
+#include "Dataflow/DataflowEdNode.h"
+#include "Dataflow/DataflowObject.h"
+#include "DynamicMesh/MeshNormals.h"
+#include "DynamicMesh/NonManifoldMappingSupport.h"
+#include "Engine/SkeletalMesh.h"
+#include "GeometryCollection/Facades/CollectionRenderingFacade.h"
+#include "SkeletalMeshAttributes.h"
+#include "ToDynamicMesh.h"
+
+using namespace UE::Geometry;
+
+namespace Dataflow
+{
+	namespace Conversion
+	{
+		// Convert a rendering facade to a dynamic mesh
+		void RenderingFacadeToDynamicMesh(const GeometryCollection::Facades::FRenderingFacade& Facade, FDynamicMesh3& DynamicMesh)
+		{
+			if (Facade.CanRenderSurface())
+			{
+				const int32 NumTriangles = Facade.NumTriangles();
+				const int32 NumVertices = Facade.NumVertices();
+
+				const TManagedArray<FIntVector>& Indices = Facade.GetIndices();
+				const TManagedArray<FVector3f>& Positions = Facade.GetVertices();
+				const TManagedArray<FVector3f>& Normals = Facade.GetNormals();
+				const TManagedArray<FLinearColor>& Colors = Facade.GetVertexColor();
+
+				DynamicMesh.Clear();
+				for (int32 VertexIndex = 0; VertexIndex < NumVertices; ++VertexIndex)
+				{
+					DynamicMesh.AppendVertex(FVertexInfo(FVector3d(Positions[VertexIndex]), Normals[VertexIndex],
+											 FVector3f(Colors[VertexIndex].R, Colors[VertexIndex].G, Colors[VertexIndex].B)));
+				}
+				for (int32 TriangleIndex = 0; TriangleIndex < NumTriangles; ++TriangleIndex)
+				{
+					DynamicMesh.AppendTriangle(FIndex3i(Indices[TriangleIndex].X, Indices[TriangleIndex].Y, Indices[TriangleIndex].Z));
+				}
+
+				FMeshNormals::QuickComputeVertexNormals(DynamicMesh);
+			}
+		}
+
+		// Convert a dataflow component to a dynamic mesh
+		void DataflowToDynamicMesh(TSharedPtr<::Dataflow::FEngineContext> DataflowContext, UObject* Asset, UDataflow* Dataflow, FDynamicMesh3& DynamicMesh)
+		{
+			// just call update on the preview scene.
+
+			FManagedArrayCollection RenderCollection;
+			GeometryCollection::Facades::FRenderingFacade Facade(RenderCollection);
+			Facade.DefineSchema();
+
+
+			for (const UDataflowEdNode* Target : Dataflow->GetRenderTargets())
+			{
+				if (Target)
+				{
+					// @todo(brice) Fix this cast
+					// Target->Render(Facade, DataflowContext);
+				}
+			}
+			RenderingFacadeToDynamicMesh(Facade, DynamicMesh);
+		}
+
+		// Convert a dynamic mesh to a rendering facade
+		void DynamicMeshToRenderingFacade(const FDynamicMesh3& DynamicMesh, GeometryCollection::Facades::FRenderingFacade& Facade)
+		{
+			if (Facade.CanRenderSurface())
+			{
+				const int32 NumTriangles = Facade.NumTriangles();
+				const int32 NumVertices = Facade.NumVertices();
+
+				// We can only override vertices attributes (position, normals, colors)
+				if ((NumTriangles == DynamicMesh.TriangleCount()) && (NumVertices == DynamicMesh.VertexCount()))
+				{
+					TManagedArray<FVector3f>& Positions = Facade.ModifyVertices();
+					TManagedArray<FVector3f>& Normals = Facade.ModifyNormals();
+					TManagedArray<FLinearColor>& Colors = Facade.ModifyVertexColor();
+
+					for (int32 VertexIndex = 0; VertexIndex < NumVertices; ++VertexIndex)
+					{
+						Positions[VertexIndex] = FVector3f(DynamicMesh.GetVertex(VertexIndex));
+						Normals[VertexIndex] = DynamicMesh.GetVertexNormal(VertexIndex);
+						Colors[VertexIndex] = DynamicMesh.GetVertexColor(VertexIndex);
+					}
+				}
+			}
+		}
+
+		// Convert a dynamic mesh to a dataflow component
+		void DynamicMeshToDataflow(const FDynamicMesh3& DynamicMesh, UDataflow* Dataflow)
+		{
+			//todo
+		}
+	}
+
+}	// namespace Dataflow
\ No newline at end of file
diff --git a/Engine/Plugins/Experimental/Dataflow/Source/DataflowAssetTools/Private/Dataflow/DataflowAssetToolsModule.cpp b/Engine/Plugins/Experimental/Dataflow/Source/DataflowAssetTools/Private/Dataflow/DataflowAssetToolsModule.cpp
new file mode 100644
--- /dev/null
+++ b/Engine/Plugins/Experimental/Dataflow/Source/DataflowAssetTools/Private/Dataflow/DataflowAssetToolsModule.cpp
@@ -0,0 +1,22 @@
+// Copyright Epic Games, Inc. All Rights Reserved.
+
+#include "Features/IModularFeatures.h"
+#include "Modules/ModuleManager.h"
+
+namespace Dataflow
+{
+	class FDataflowAssetToolsModule : public IModuleInterface
+	{
+	public:
+		// IModuleInterface implementation
+		virtual void StartupModule() override
+		{
+		}
+
+		virtual void ShutdownModule() override
+		{
+		}
+	};
+}
+
+IMPLEMENT_MODULE(Dataflow::FDataflowAssetToolsModule, DataflowAssetTools);
diff --git a/Engine/Plugins/Experimental/Dataflow/Source/DataflowAssetTools/Public/Dataflow/CollectionPatternToDynamicMesh.h b/Engine/Plugins/Experimental/Dataflow/Source/DataflowAssetTools/Public/Dataflow/CollectionPatternToDynamicMesh.h
new file mode 100644
--- /dev/null
+++ b/Engine/Plugins/Experimental/Dataflow/Source/DataflowAssetTools/Public/Dataflow/CollectionPatternToDynamicMesh.h
@@ -0,0 +1,31 @@
+// Copyright Epic Games, Inc. All Rights Reserved.
+
+#pragma once
+
+#include "DataflowPatternVertexType.h"
+#include "Templates/SharedPointer.h"
+
+
+namespace UE::Geometry { class FDynamicMesh3; }
+struct FManagedArrayCollection;
+
+namespace Dataflow
+{
+
+/**
+* Convert a single pattern from a colleciton asset to a FDynamicMesh3. 
+* When PatternIndex = INDEX_NONE, convert the entire collection asset to an FDynamicMesh3.
+* When VertexDataType = EDataflowPatternVertexType::Sim3D, this generates a welded mesh (with native welded indexing) when PatternIndex = INDEX_NONE,
+*	and an unwelded (with native 2D unwelded indexing) when PatternIndex is a valid SimPattern Index.
+*/
+class DATAFLOWASSETTOOLS_API FCollectionPatternToDynamicMesh
+{
+public:
+
+	void Convert(const FManagedArrayCollection* FleshCollection, int32 PatternIndex, EDataflowPatternVertexType VertexDataType, UE::Geometry::FDynamicMesh3& MeshOut, bool bDisableAttributes = false);
+
+	void Convert(const UObject* AssetIn, int32 LODIndex, int32 PatternIndex, EDataflowPatternVertexType VertexDataType, UE::Geometry::FDynamicMesh3& MeshOut);
+};
+
+}	// namespace Dataflow
+
diff --git a/Engine/Plugins/Experimental/Dataflow/Source/DataflowAssetTools/Public/Dataflow/CollectionRenderingPatternUtility.h b/Engine/Plugins/Experimental/Dataflow/Source/DataflowAssetTools/Public/Dataflow/CollectionRenderingPatternUtility.h
new file mode 100644
--- /dev/null
+++ b/Engine/Plugins/Experimental/Dataflow/Source/DataflowAssetTools/Public/Dataflow/CollectionRenderingPatternUtility.h
@@ -0,0 +1,30 @@
+// Copyright Epic Games, Inc. All Rights Reserved. 
+#pragma once
+
+#include "CoreMinimal.h"
+#include "Dataflow/DataflowObjectInterface.h"
+
+
+namespace UE::Geometry { class FDynamicMesh3; }
+namespace GeometryCollection::Facades { class FRenderingFacade; }
+class UDataflow;
+class UObject;
+
+namespace Dataflow
+{
+	namespace Conversion
+	{
+		// Convert a rendering facade to a dynamic mesh
+		void DATAFLOWASSETTOOLS_API RenderingFacadeToDynamicMesh(const GeometryCollection::Facades::FRenderingFacade& Facade, UE::Geometry::FDynamicMesh3& DynamicMesh);
+
+		// Convert a dataflow component to a dynamic mesh
+		void DATAFLOWASSETTOOLS_API DataflowToDynamicMesh(TSharedPtr<::Dataflow::FEngineContext> DataflowContext, UObject* Asset, UDataflow* Dataflow, UE::Geometry::FDynamicMesh3& DynamicMesh);
+
+		// Convert a dynamic mesh to a rendering facade
+		void DATAFLOWASSETTOOLS_API DynamicMeshToRenderingFacade(const UE::Geometry::FDynamicMesh3& DynamicMesh, GeometryCollection::Facades::FRenderingFacade& Facade);
+
+		// Convert a dynamic mesh to a dataflow component
+		void DATAFLOWASSETTOOLS_API DynamicMeshToDataflow(const UE::Geometry::FDynamicMesh3& DynamicMesh, UDataflow* Dataflow);
+	}
+
+}	// namespace Dataflow
\ No newline at end of file
diff --git a/Engine/Plugins/Experimental/Dataflow/Source/DataflowAssetTools/Public/Dataflow/ColorScheme.h b/Engine/Plugins/Experimental/Dataflow/Source/DataflowAssetTools/Public/Dataflow/ColorScheme.h
new file mode 100644
--- /dev/null
+++ b/Engine/Plugins/Experimental/Dataflow/Source/DataflowAssetTools/Public/Dataflow/ColorScheme.h
@@ -0,0 +1,16 @@
+// Copyright Epic Games, Inc. All Rights Reserved.
+
+#pragma once
+
+#include "Math/Color.h"
+
+namespace Dataflow
+{
+	struct FColorScheme
+	{
+		static inline const FLinearColor Asset = FColor(180, 120, 110);
+		static inline const FLinearColor NodeHeader = FColor(180, 120, 110);
+		static inline const FLinearColor NodeBody = FColor(18, 12, 11, 127);
+	};
+}
+
diff --git a/Engine/Plugins/Experimental/Dataflow/Source/DataflowAssetTools/Public/Dataflow/DataflowPatternVertexType.h b/Engine/Plugins/Experimental/Dataflow/Source/DataflowAssetTools/Public/Dataflow/DataflowPatternVertexType.h
new file mode 100644
--- /dev/null
+++ b/Engine/Plugins/Experimental/Dataflow/Source/DataflowAssetTools/Public/Dataflow/DataflowPatternVertexType.h
@@ -0,0 +1,16 @@
+// Copyright Epic Games, Inc. All Rights Reserved.
+
+#pragma once
+
+#include "CoreTypes.h"
+
+namespace Dataflow
+{
+	enum class EDataflowPatternVertexType : uint8
+	{
+		Sim2D = 0,
+		Sim3D = 1,
+		Render = 2
+	};
+}
+
diff --git a/Engine/Plugins/Experimental/Dataflow/Source/DataflowEditor/DataflowEditor.Build.cs b/Engine/Plugins/Experimental/Dataflow/Source/DataflowEditor/DataflowEditor.Build.cs
--- a/Engine/Plugins/Experimental/Dataflow/Source/DataflowEditor/DataflowEditor.Build.cs
+++ b/Engine/Plugins/Experimental/Dataflow/Source/DataflowEditor/DataflowEditor.Build.cs
@@ -25,40 +25,42 @@ namespace UnrealBuildTool.Rules
 					"Chaos",
 					"Core",
 					"CoreUObject",
+					"DataflowAssetTools",
 					"DataflowCore",
 					"DataflowEngine",
 					"DataflowEnginePlugin",
 					"DataflowNodes",
+					"DynamicMesh",
 					"Engine",
 					"EditorFramework",
 					"EditorInteractiveToolsFramework",
 					"EditorStyle",
+					"GeometryCore",
+					"GeometryFramework",
 					"GraphEditor",
 					"InputCore",
 					"InteractiveToolsFramework",
+					"LevelEditor",
+					"MeshDescription",
+					"MeshConversion",
 					"MeshModelingToolsEditorOnlyExp",
+					"MeshModelingToolsExp",
 					"ModelingComponentsEditorOnly",
 					"ModelingComponents",
-					"GeometryCore",
-					"GeometryFramework",
-					"MeshDescription",
-					"StaticMeshDescription",
-					"MeshConversion",
-					"LevelEditor",
-					"Slate",
-				    "SlateCore",
-					"UnrealEd",
 					"Projects",
 					"PropertyEditor",
-				    "RenderCore",
-				    "RHI",
-				    "SceneOutliner",
+					"RenderCore",
+					"RHI",
+					"SceneOutliner",
+					"Slate",
+					"SlateCore",
+					"StaticMeshDescription",
 					"ToolMenus",
 					"ToolWidgets",
-					"Slate",
-					"XmlParser",
 					"TypedElementRuntime",
-					"MeshModelingToolsExp"
+					"UnrealEd",
+					"WorkspaceMenuStructure",
+					"XmlParser"
 				}
 			);
 		}
diff --git a/Engine/Plugins/Experimental/Dataflow/Source/DataflowEditor/Private/Dataflow/DataflowEditorCommands.cpp b/Engine/Plugins/Experimental/Dataflow/Source/DataflowEditor/Private/Dataflow/DataflowEditorCommands.cpp
--- a/Engine/Plugins/Experimental/Dataflow/Source/DataflowEditor/Private/Dataflow/DataflowEditorCommands.cpp
+++ b/Engine/Plugins/Experimental/Dataflow/Source/DataflowEditor/Private/Dataflow/DataflowEditorCommands.cpp
@@ -12,6 +12,7 @@
 #include "Dataflow/DataflowOverrideNode.h"
 #include "Dataflow/DataflowSCommentNode.h"
 #include "Dataflow/DataflowSNode.h"
+#include "DataflowEditorTools/DataflowEditorWeightMapPaintTool.h"
 #include "EdGraphNode_Comment.h"
 #include "EdGraph/EdGraphNode.h"
 #include "Editor.h"
@@ -19,8 +20,12 @@
 
 #define LOCTEXT_NAMESPACE "DataflowEditorCommands"
 
-const FString FDataflowEditorCommandsImpl::BeginAttributeEditorToolIdentifier = TEXT("BeginAttributeEditorTool");
-const FString FDataflowEditorCommandsImpl::BeginMeshSelectionToolIdentifier = TEXT("BeginMeshSelectionTool");
+const FString FDataflowEditorCommandsImpl::BeginWeightMapPaintToolIdentifier = TEXT("BeginWeightMapPaintTool");
+const FString FDataflowEditorCommandsImpl::AddWeightMapNodeIdentifier = TEXT("AddWeightMapNode");
+
+// @todo(brice) Remove Example Tools
+//const FString FDataflowEditorCommandsImpl::BeginAttributeEditorToolIdentifier = TEXT("BeginAttributeEditorTool");
+//const FString FDataflowEditorCommandsImpl::BeginMeshSelectionToolIdentifier = TEXT("BeginMeshSelectionTool");
 
 FDataflowEditorCommandsImpl::FDataflowEditorCommandsImpl()
 	: TBaseCharacterFXEditorCommands<FDataflowEditorCommandsImpl>("DataflowEditor", 
@@ -44,8 +49,9 @@ void FDataflowEditorCommandsImpl::RegisterCommands()
 	UI_COMMAND(RemoveOptionPin, "RemoveOptionPin", "Remove the last option pin from the selected nodes.", EUserInterfaceActionType::Button, FInputChord());
 	UI_COMMAND(ZoomToFitGraph, "ZoomToFitGraph", "Fit the graph in the graph editor viewport.", EUserInterfaceActionType::None, FInputChord(EKeys::F));
 
-	UI_COMMAND(BeginAttributeEditorTool, "AttrEd", "Edit/configure mesh attributes", EUserInterfaceActionType::Button, FInputChord());
-	UI_COMMAND(BeginMeshSelectionTool, "SelVtx", "Select mesh vertices", EUserInterfaceActionType::Button, FInputChord());
+	UI_COMMAND(BeginWeightMapPaintTool, "Add Weight Map", "Paint weight maps on the mesh", EUserInterfaceActionType::None, FInputChord());
+	UI_COMMAND(AddWeightMapNode, "Add Weight Map", "Paint weight maps on the mesh", EUserInterfaceActionType::Button, FInputChord());
+
 
 	if (Dataflow::FNodeFactory* Factory = Dataflow::FNodeFactory::GetInstance())
 	{
@@ -69,6 +75,22 @@ void FDataflowEditorCommandsImpl::RegisterCommands()
 
 void FDataflowEditorCommandsImpl::GetToolDefaultObjectList(TArray<UInteractiveTool*>& ToolCDOs)
 {
+	ToolCDOs.Add(GetMutableDefault<UDataflowEditorWeightMapPaintTool>());
+}
+
+void FDataflowEditorCommandsImpl::UpdateToolCommandBinding(UInteractiveTool* Tool, TSharedPtr<FUICommandList> UICommandList, bool bUnbind)
+{
+	if (FDataflowEditorCommandsImpl::IsRegistered())
+	{
+		if (bUnbind)
+		{
+			FDataflowEditorCommandsImpl::Get().UnbindActiveCommands(UICommandList);
+		}
+		else
+		{
+			FDataflowEditorCommandsImpl::Get().BindCommandsForCurrentTool(UICommandList, Tool);
+		}
+	}
 }
 
 const FDataflowEditorCommandsImpl& FDataflowEditorCommands::Get()
diff --git a/Engine/Plugins/Experimental/Dataflow/Source/DataflowEditor/Private/Dataflow/DataflowEditorContent.cpp b/Engine/Plugins/Experimental/Dataflow/Source/DataflowEditor/Private/Dataflow/DataflowEditorContent.cpp
--- a/Engine/Plugins/Experimental/Dataflow/Source/DataflowEditor/Private/Dataflow/DataflowEditorContent.cpp
+++ b/Engine/Plugins/Experimental/Dataflow/Source/DataflowEditor/Private/Dataflow/DataflowEditorContent.cpp
@@ -8,5 +8,11 @@
 
 
 UDataflowEditorContent::UDataflowEditorContent()
+	: UDataflowEditorContextObject()
 {
-}
\ No newline at end of file
+}
+
+
+
+
+
diff --git a/Engine/Plugins/Experimental/Dataflow/Source/DataflowEditor/Private/Dataflow/DataflowEditorContent.h b/Engine/Plugins/Experimental/Dataflow/Source/DataflowEditor/Private/Dataflow/DataflowEditorContent.h
deleted file mode 100644
diff --git a/Engine/Plugins/Experimental/Dataflow/Source/DataflowEditor/Private/Dataflow/DataflowEditorMode.cpp b/Engine/Plugins/Experimental/Dataflow/Source/DataflowEditor/Private/Dataflow/DataflowEditorMode.cpp
--- a/Engine/Plugins/Experimental/Dataflow/Source/DataflowEditor/Private/Dataflow/DataflowEditorMode.cpp
+++ b/Engine/Plugins/Experimental/Dataflow/Source/DataflowEditor/Private/Dataflow/DataflowEditorMode.cpp
@@ -2,19 +2,26 @@
 
 #include "Dataflow/DataflowEditorMode.h"
 
+#include "AssetEditorModeManager.h"
 #include "AttributeEditorTool.h"
+#include "BaseGizmos/TransformGizmoUtil.h"
 #include "Components/DynamicMeshComponent.h"
+#include "ContextObjectStore.h"
+#include "DataflowEditorTools/DataflowEditorWeightMapPaintTool.h"
 #include "Dataflow/DataflowComponentToolTarget.h"
+#include "Dataflow/DataflowCollectionAddScalarVertexPropertyNode.h"
 #include "Dataflow/DataflowEditor.h"
 #include "Dataflow/DataflowEditorContent.h"
 #include "Dataflow/DataflowEditorCommands.h"
 #include "Dataflow/DataflowEditorModeToolkit.h"
+#include "Dataflow/DataflowEditorViewportClient.h"
 #include "Dataflow/DataflowGraphEditor.h"
 #include "Dataflow/DataflowPreviewScene.h"
 #include "Dataflow/DataflowSNode.h"
 #include "Dataflow/DataflowToolTarget.h"
 #include "EditorModeManager.h"
 #include "EdModeInteractiveToolsContext.h"
+#include "Elements/Framework/EngineElementsLibrary.h"
 #include "MeshSelectionTool.h"
 #include "MeshVertexPaintTool.h"
 #include "MeshAttributePaintTool.h"
@@ -30,6 +37,7 @@
 #include "ToolTargets/SkeletalMeshToolTarget.h"
 #include "Tools/UEdMode.h"
 #include "Selection.h"
+#include "UnrealClient.h"
 
 #include UE_INLINE_GENERATED_CPP_BY_NAME(DataflowEditorMode)
 
@@ -41,33 +49,11 @@ UDataflowEditorMode::UDataflowEditorMode()
 {
 	Info = FEditorModeInfo(
 		EM_DataflowEditorModeId,
-		LOCTEXT("DataflowEditorModeName", "Dataflow Tools"),
+		LOCTEXT("DataflowEditorModeName", "Dataflow"),
 		FSlateIcon(),
 		false);
 }
 
-void UDataflowEditorMode::Exit()
-{
-	UActorComponent::MarkRenderStateDirtyEvent.RemoveAll(this);
-	PreviewScene->ResetDynamicMeshComponents();
-	
-	Super::Exit();
-}
-
-void UDataflowEditorMode::Enter()
-{
-	Super::Enter();
-}
-
-void UDataflowEditorMode::ModeTick(float DeltaTime)
-{
-	Super::ModeTick(DeltaTime);
-	if (PreviewScene)
-	{
-		PreviewScene->GetWorld()->Tick(ELevelTick::LEVELTICK_All, DeltaTime);
-	}
-}
-
 const FToolTargetTypeRequirements& UDataflowEditorMode::GetToolTargetRequirements()
 {
 	static const FToolTargetTypeRequirements ToolTargetRequirements =
@@ -80,6 +66,14 @@ const FToolTargetTypeRequirements& UDataflowEditorMode::GetToolTargetRequirement
 	return ToolTargetRequirements;
 }
 
+void UDataflowEditorMode::Enter()
+{
+	UBaseCharacterFXEditorMode::Enter();
+
+	// Register gizmo ContextObject for use inside interactive tools
+	UE::TransformGizmoUtil::RegisterTransformGizmoContextObject(GetInteractiveToolsContext());
+}
+
 void UDataflowEditorMode::AddToolTargetFactories()
 {
 	GetInteractiveToolsContext()->TargetManager->AddTargetFactory(NewObject<UDynamicMeshComponentToolTargetFactory>(GetToolManager()));
@@ -91,15 +85,143 @@ void UDataflowEditorMode::AddToolTargetFactories()
 	GetInteractiveToolsContext()->TargetManager->AddTargetFactory(NewObject<UDataflowToolTargetFactory>(GetToolManager()));
 }
 
+void UDataflowEditorMode::RegisterDataflowTool(TSharedPtr<FUICommandInfo> UICommand,
+	FString ToolIdentifier,
+	UInteractiveToolBuilder* Builder,
+	const IDataflowEditorToolBuilder* DataflowToolBuilder,
+	UEditorInteractiveToolsContext* const ToolsContext,
+	EToolsContextScope ToolScope)
+{
+	if (!Toolkit.IsValid())
+	{
+		return;
+	}
+
+	if (!ToolsContext)
+	{
+		return;
+	}
+
+	if (ToolScope == EToolsContextScope::Default)
+	{
+		ToolScope = GetDefaultToolScope();
+	}
+	ensure(ToolScope != EToolsContextScope::Editor);
+
+	ToolsContext->ToolManager->RegisterToolType(ToolIdentifier, Builder);
+
+	const TSharedRef<FUICommandList>& CommandList = Toolkit->GetToolkitCommands();
+
+	CommandList->MapAction(UICommand, 
+		FExecuteAction::CreateWeakLambda(ToolsContext, [this, ToolsContext, ToolIdentifier, DataflowToolBuilder]()
+		{
+			// Check if we need to switch view modes before starting the tool
+			TArray<Dataflow::EDataflowPatternVertexType> SupportedModes;
+			DataflowToolBuilder->GetSupportedViewModes(SupportedModes);
+
+			if (SupportedModes.Num() > 0 && !SupportedModes.Contains(this->GetConstructionViewMode()))
+			{
+				if (!bShouldRestoreSavedConstructionViewMode)
+				{
+					// remember the current view mode so we can restore it later
+					SavedConstructionViewMode = this->GetConstructionViewMode();
+					bShouldRestoreSavedConstructionViewMode = true;
+				}
+
+				// switch to the preferred view mode for the tool that's about to start
+				this->SetConstructionViewMode(SupportedModes[0]);
+			}
+
+			// Check if we need to disable wireframe mode before starting tool.
+			const bool bCanSetWireframeActive = DataflowToolBuilder->CanSetConstructionViewWireframeActive();
+			if (!bCanSetWireframeActive)
+			{
+				if (!bShouldRestoreConstructionViewWireframe)
+				{
+					bShouldRestoreConstructionViewWireframe = bConstructionViewWireframe;
+				}
+				bConstructionViewWireframe = false;
+			}
+
+			ActiveToolsContext = ToolsContext;
+			ToolsContext->StartTool(ToolIdentifier);
+		}),
+		FCanExecuteAction::CreateWeakLambda(ToolsContext, [this, ToolIdentifier, ToolsContext]()
+		{
+			return ShouldToolStartBeAllowed(ToolIdentifier) &&
+			ToolsContext->ToolManager->CanActivateTool(EToolSide::Mouse, ToolIdentifier);
+		}),
+		FIsActionChecked::CreateUObject(ToolsContext, &UEdModeInteractiveToolsContext::IsToolActive, EToolSide::Mouse, ToolIdentifier),
+		EUIActionRepeatMode::RepeatDisabled
+	);
+}
+
+void UDataflowEditorMode::RegisterAddNodeCommand(TSharedPtr<FUICommandInfo> AddNodeCommand, const FName& NewNodeType, TSharedPtr<FUICommandInfo> StartToolCommand)
+{
+	auto AddNode = [this](const FName& NewNodeType)
+	{
+		const FName ConnectionType = FManagedArrayCollection::StaticType();
+		const FName ConnectionName("Collection");
+
+		UEdGraphNode* const CurrentlySelectedNode = GetSingleSelectedNodeWithOutputType(ConnectionType);
+		checkf(CurrentlySelectedNode, TEXT("No node with FManagedArrayCollection output is currently selected in the Dataflow graph"));
+
+		const UEdGraphNode* const NewNode = CreateAndConnectNewNode(NewNodeType, *CurrentlySelectedNode, ConnectionType, ConnectionName);
+		verifyf(NewNode, TEXT("Failed to create a new node: %s"), *NewNodeType.ToString());
+
+		StartToolForSelectedNode(NewNode);
+	};
+
+	auto CanAddNode = [this](const FName& NewNodeType) -> bool
+	{
+		const UEdGraphNode* const CurrentlySelectedNode = GetSingleSelectedNodeWithOutputType(FManagedArrayCollection::StaticType());
+		return (CurrentlySelectedNode != nullptr);
+	};
+
+	const TSharedRef<FUICommandList>& CommandList = Toolkit->GetToolkitCommands();
+
+	CommandList->MapAction(AddNodeCommand,
+		FExecuteAction::CreateWeakLambda(this, AddNode, NewNodeType),
+		FCanExecuteAction::CreateWeakLambda(this, CanAddNode, NewNodeType)
+	);
+
+	NodeTypeToToolCommandMap.Add(NewNodeType, StartToolCommand);
+}
+
 void UDataflowEditorMode::RegisterTools()
 {
 	const FDataflowEditorCommandsImpl& CommandInfos = FDataflowEditorCommands::Get();
-	RegisterTool(CommandInfos.BeginAttributeEditorTool, FDataflowEditorCommandsImpl::BeginAttributeEditorToolIdentifier, NewObject<UAttributeEditorToolBuilder>());
 
-	// For testing...
+	UEditorInteractiveToolsContext* const ConstructionViewportToolsContext = GetInteractiveToolsContext();
+
+	UDataflowEditorWeightMapPaintToolBuilder* WeightMapPaintToolBuilder = NewObject<UDataflowEditorWeightMapPaintToolBuilder>();
+	RegisterDataflowTool(CommandInfos.BeginWeightMapPaintTool, FDataflowEditorCommandsImpl::BeginWeightMapPaintToolIdentifier, WeightMapPaintToolBuilder, WeightMapPaintToolBuilder, ConstructionViewportToolsContext);
+	RegisterAddNodeCommand(CommandInfos.AddWeightMapNode, FDataflowCollectionAddScalarVertexPropertyNode::StaticType(), CommandInfos.BeginWeightMapPaintTool);
+
+	// @todo(brice) Remove Example Tools
+	//RegisterTool(CommandInfos.BeginAttributeEditorTool, FDataflowEditorCommandsImpl::BeginAttributeEditorToolIdentifier, NewObject<UAttributeEditorToolBuilder>());
 	//RegisterTool(CommandInfos.BeginMeshSelectionTool, FDataflowEditorCommandsImpl::BeginMeshSelectionToolIdentifier, NewObject<UMeshSelectionToolBuilder>());
 	//RegisterTool(CommandInfos.BeginMeshSelectionTool, FDataflowEditorCommandsImpl::BeginMeshSelectionToolIdentifier, NewObject<UMeshVertexPaintToolBuilder>());
-	RegisterTool(CommandInfos.BeginMeshSelectionTool, FDataflowEditorCommandsImpl::BeginMeshSelectionToolIdentifier, NewObject<UMeshAttributePaintToolBuilder>());
+	//RegisterTool(CommandInfos.BeginMeshSelectionTool, FDataflowEditorCommandsImpl::BeginMeshSelectionToolIdentifier, NewObject<UMeshAttributePaintToolBuilder>());
+}
+
+bool UDataflowEditorMode::ShouldToolStartBeAllowed(const FString& ToolIdentifier) const
+{
+	// Allow switching away from tool if no changes have been made in the tool yet (which we infer from the CanAccept status)
+	if (GetInteractiveToolsContext()->CanAcceptActiveTool())
+	{
+		return false;
+	}
+
+	if (PreviewScene && PreviewScene->GetDataflowModeManager() && PreviewScene->GetDataflowModeManager()->GetInteractiveToolsContext())
+	{
+		if (PreviewScene->GetDataflowModeManager()->GetInteractiveToolsContext()->HasActiveTool())
+		{
+			return false;
+		}
+	}
+
+	return Super::ShouldToolStartBeAllowed(ToolIdentifier);
 }
 
 void UDataflowEditorMode::CreateToolkit()
@@ -107,6 +229,37 @@ void UDataflowEditorMode::CreateToolkit()
 	Toolkit = MakeShared<FDataflowEditorModeToolkit>();
 }
 
+void UDataflowEditorMode::OnToolStarted(UInteractiveToolManager* Manager, UInteractiveTool* Tool)
+{
+	FDataflowEditorCommandsImpl::UpdateToolCommandBinding(Tool, ToolCommandList, false);
+}
+
+void UDataflowEditorMode::OnToolEnded(UInteractiveToolManager* Manager, UInteractiveTool* Tool)
+{
+	FDataflowEditorCommandsImpl::UpdateToolCommandBinding(Tool, ToolCommandList, true);
+
+	if (bShouldRestoreConstructionViewWireframe)
+	{
+		bConstructionViewWireframe = true;
+		bShouldRestoreConstructionViewWireframe = false;
+	}
+
+	if (bShouldRestoreSavedConstructionViewMode)
+	{
+		SetConstructionViewMode(SavedConstructionViewMode);
+		bShouldRestoreSavedConstructionViewMode = false;
+	}
+	else
+	{
+		PreviewScene->ReinitializeDynamicMeshComponents();
+	}
+
+	if (TSharedPtr<SDataflowGraphEditor> GraphEditor = DataflowGraphEditor.Pin())
+	{
+		GraphEditor->SetEnabled(true);
+	}
+}
+
 void UDataflowEditorMode::BindCommands()
 {
 	const FDataflowEditorCommandsImpl& CommandInfos = FDataflowEditorCommands::Get();
@@ -115,14 +268,10 @@ void UDataflowEditorMode::BindCommands()
 	// Hook up to Enter/Esc key presses
 	CommandList->MapAction(
 		CommandInfos.AcceptOrCompleteActiveTool,
-		FExecuteAction::CreateLambda([this]()
-		{
-			AcceptActiveToolActionOrTool();
-			
-		}),
+		FExecuteAction::CreateLambda([this](){ AcceptActiveToolActionOrTool();}),
 		FCanExecuteAction::CreateLambda([this]() {
-		return GetInteractiveToolsContext()->CanAcceptActiveTool() || GetInteractiveToolsContext()->CanCompleteActiveTool();
-	}),
+				return GetInteractiveToolsContext()->CanAcceptActiveTool() || GetInteractiveToolsContext()->CanCompleteActiveTool();
+			}),
 		FGetActionCheckState(),
 		FIsActionButtonVisible(),
 		EUIActionRepeatMode::RepeatDisabled);
@@ -131,13 +280,45 @@ void UDataflowEditorMode::BindCommands()
 		CommandInfos.CancelOrCompleteActiveTool,
 		FExecuteAction::CreateLambda([this]() { CancelActiveToolActionOrTool(); }),
 		FCanExecuteAction::CreateLambda([this]() {
-		return GetInteractiveToolsContext()->CanCompleteActiveTool() || GetInteractiveToolsContext()->CanCancelActiveTool();
-	}),
+			return GetInteractiveToolsContext()->CanCompleteActiveTool() || GetInteractiveToolsContext()->CanCancelActiveTool();
+			}),
 		FGetActionCheckState(),
 		FIsActionButtonVisible(),
 		EUIActionRepeatMode::RepeatDisabled);
 }
 
+void UDataflowEditorMode::Exit()
+{
+	UActorComponent::MarkRenderStateDirtyEvent.RemoveAll(this);
+
+	PreviewScene->Exit();
+	PreviewScene = nullptr;
+
+	Super::Exit();
+}
+
+void UDataflowEditorMode::SetDataflowPreviewScene(FDataflowPreviewScene* InPreviewScene)
+{
+	PreviewScene = InPreviewScene;
+
+	UEditorInteractiveToolsContext* const PreviewToolsContext = PreviewScene->GetDataflowModeManager()->GetInteractiveToolsContext();
+	UInteractiveToolManager* const PreviewToolManager = PreviewToolsContext->ToolManager;
+	//#todo(brice): Make sure AddToolTargetFactories has been called. 
+	//PreviewToolsContext->TargetManager->AddTargetFactory(NewObject<UClothComponentToolTargetFactory>(PreviewToolManager));
+	//PreviewToolsContext->TargetManager->AddTargetFactory(NewObject<USkeletalMeshComponentToolTargetFactory>(PreviewToolManager));
+
+	PreviewToolManager->OnToolStarted.AddUObject(this, &UDataflowEditorMode::OnToolStarted);
+	PreviewToolManager->OnToolEnded.AddUObject(this, &UDataflowEditorMode::OnToolEnded);
+
+	check(Toolkit.IsValid());
+
+	// FBaseToolkit's OnToolStarted and OnToolEnded are protected, so we use the subclass to get at them
+	FDataflowEditorModeToolkit* const DataflowModeToolkit = static_cast<FDataflowEditorModeToolkit*>(Toolkit.Get());
+
+	PreviewToolManager->OnToolStarted.AddSP(DataflowModeToolkit, &FDataflowEditorModeToolkit::OnToolStarted);
+	PreviewToolManager->OnToolEnded.AddSP(DataflowModeToolkit, &FDataflowEditorModeToolkit::OnToolEnded);
+}
+
 void UDataflowEditorMode::CreateToolTargets(const TArray<TObjectPtr<UObject>>& AssetsIn)
 {
 	ToolTargets.Reset();
@@ -150,15 +331,100 @@ void UDataflowEditorMode::CreateToolTargets(const TArray<TObjectPtr<UObject>>& A
 	}
 }
 
-void UDataflowEditorMode::InitializeTargets(const TArray<TObjectPtr<UObject>>& AssetsIn)
+bool UDataflowEditorMode::IsComponentSelected(const UPrimitiveComponent* InComponent)
+{
+	if (const FEditorModeTools* const ModeManager = GetModeManager())
+	{
+		if (const UTypedElementSelectionSet* const TypedElementSelectionSet = ModeManager->GetEditorSelectionSet())
+		{
+			if (const FTypedElementHandle ComponentElement = UEngineElementsLibrary::AcquireEditorComponentElementHandle(InComponent))
+			{
+				const bool bElementSelected = TypedElementSelectionSet->IsElementSelected(ComponentElement, FTypedElementIsSelectedOptions());
+				return bElementSelected;
+			}
+		}
+	}
+
+	return false;
+}
+
+void UDataflowEditorMode::SetSelectedCollection(TSharedPtr<FManagedArrayCollection> Collection)
+{
+	SelectedCollection = Collection;
+	PreviewScene->ReinitializeDynamicMeshComponents();
+
+	// The first time we get a valid mesh, refocus the camera on it
+	FirstTimeFocusRestSpaceViewport();
+}
+
+TSharedPtr<FManagedArrayCollection> UDataflowEditorMode::GetSelectedCollection()
+{
+	return SelectedCollection;
+
+	// TODO: If no cloth collection node is selected, show the ClothAsset's collection. In this case, also ensure that any interactive tools are disabled. (UE-181574)
+}
+
+void UDataflowEditorMode::RefocusRestSpaceViewportClient()
+{
+	TSharedPtr<FDataflowEditorViewportClient, ESPMode::ThreadSafe> PinnedVC = ConstructionViewportClient.Pin();
+	if (PinnedVC.IsValid())
+	{
+		// This will happen in FocusViewportOnBox anyways; do it now to get a consistent end result
+		PinnedVC->ToggleOrbitCamera(false);
+
+		const FBox SceneBounds = SceneBoundingBox();
+		const bool bPattern2DMode = (ConstructionViewMode == Dataflow::EDataflowPatternVertexType::Sim2D);
+		if (bPattern2DMode)
+		{
+			// 2D pattern
+			PinnedVC->SetInitialViewTransform(ELevelViewportType::LVT_Perspective, FVector(0, 0, -100), FRotator(90, -90, 0), DEFAULT_ORTHOZOOM);
+		}
+		else
+		{
+			// 3D rest space
+			PinnedVC->SetInitialViewTransform(ELevelViewportType::LVT_Perspective, FVector(0, 150, 200), FRotator(0, 0, 0), DEFAULT_ORTHOZOOM);
+		}
+
+		constexpr bool bInstant = true;
+		PinnedVC->FocusViewportOnBox(SceneBounds, bInstant);
+
+		// Recompute near/far clip planes
+		PinnedVC->SetConstructionViewMode(ConstructionViewMode);
+	}
+}
+
+void UDataflowEditorMode::FirstTimeFocusRestSpaceViewport()
+{
+	// If this is the first time seeing a valid 2D or 3D mesh, refocus the camera on it.
+	const bool bIsValid = (SelectedCollection && PreviewScene->HasRenderableGeometry());
+	const bool bIs2D = ConstructionViewMode == Dataflow::EDataflowPatternVertexType::Sim3D;
+
+	if (bIsValid)
+	{
+		if (bIs2D && bFirstValid2DMesh)
+		{
+			bFirstValid2DMesh = false;
+			RefocusRestSpaceViewportClient();
+		}
+		else if (!bIs2D && bFirstValid3DMesh)
+		{
+			bFirstValid3DMesh = false;
+			RefocusRestSpaceViewportClient();
+		}
+	}
+}
+
+void UDataflowEditorMode::InitializeTargets(const TArray<TObjectPtr<UObject>>& ObjectsToEdit)
 {
-	UBaseCharacterFXEditorMode::InitializeTargets(AssetsIn);
-	
+	UBaseCharacterFXEditorMode::InitializeTargets(ObjectsToEdit);
+
 	USelection* SelectedComponents = GetModeManager()->GetSelectedComponents();
 	SelectedComponents->Modify();
 	SelectedComponents->BeginBatchSelectOperation();
 	SelectedComponents->DeselectAll();
 
+	// @todo(brice) : Consider initializing the Content here from the ObjectsToEdit
+
 	// @todo(brice) : What are the ToolTargets storing?
 	// ... for(ToolTarget& : ToolTargets){
 	// ... UE::ToolTarget::GetDynamicMeshCopy(Target)
@@ -176,53 +442,381 @@ void UDataflowEditorMode::InitializeTargets(const TArray<TObjectPtr<UObject>>& A
 	SelectedComponents->EndBatchSelectOperation();
 }
 
-// UEdGraphNode* UDataflowEditorMode::CreateNewNode(const FName& NewNodeTypeName)
-// {
-// 	const TSharedPtr<const SDataflowGraphEditor> PinnedDataflowGraphEditor = GraphEditor.Pin();
-// 	if (!PinnedDataflowGraphEditor)
-// 	{
-// 		return nullptr;
-// 	}
-//
-// 	const TSharedPtr<FAssetSchemaAction_Dataflow_CreateNode_DataflowEdNode> NodeAction =
-// 		FAssetSchemaAction_Dataflow_CreateNode_DataflowEdNode::CreateAction(PreviewScene->GetDataflowDatas().DataflowAsset, NewNodeTypeName);
-// 	constexpr UEdGraphPin* FromPin = nullptr;
-// 	constexpr bool bSelectNewNode = true;
-// 	UEdGraphNode* const NewEdNode = NodeAction->PerformAction(PreviewScene->GetDataflowDatas().DataflowAsset, FromPin,
-// 			PinnedDataflowGraphEditor->GetPasteLocation(), bSelectNewNode);
-//
-// 	PreviewScene->GetDataflowDatas().DataflowAsset->NotifyGraphChanged();
-//
-// 	return NewEdNode;
-// }
-//
-// void UDataflowEditorMode::RegisterAddNodeCommand(TSharedPtr<FUICommandInfo> AddNodeCommand, const FName& NewNodeType, TSharedPtr<FUICommandInfo> StartToolCommand)
-// {
-// 	auto AddNode = [this](const FName& NewNodeType)
-// 	{
-// 		const FName ConnectionType = FManagedArrayCollection::StaticType();
-// 		const FName ConnectionName("Collection");
-// 		
-// 		const UEdGraphNode* const NewNode = CreateNewNode(NewNodeType);
-// 		verifyf(NewNode, TEXT("Failed to create a new node: %s"), *NewNodeType.ToString());
-//
-// 		StartToolForSelectedNode(NewNode);
-// 	};
-//
-// 	const TSharedRef<FUICommandList>& CommandList = Toolkit->GetToolkitCommands();
-//
-// 	CommandList->MapAction(AddNodeCommand,
-// 		FExecuteAction::CreateWeakLambda(this, AddNode, NewNodeType),
-// 		FCanExecuteAction::CreateLambda([]() {return true;})
-// 	);
-//
-// 	NodeToolMap.Add(NewNodeType, StartToolCommand);
-// }
+
+void UDataflowEditorMode::ModeTick(float DeltaTime)
+{
+	Super::ModeTick(DeltaTime);
+
+	if (TSharedPtr<SDataflowGraphEditor> GraphEditor = DataflowGraphEditor.Pin())
+	{
+		// For now don't allow selection change once the tool has uncommitted changes
+		// TODO: We might want to auto-accept unsaved changes and allow switching between nodes
+		if (GetInteractiveToolsContext()->CanAcceptActiveTool())
+		{
+			GraphEditor->SetEnabled(false);
+		}
+		else
+		{
+			GraphEditor->SetEnabled(true);
+		}
+	}
+
+	if (!NodeTypeForPendingToolStart.IsNone() && !GetToolManager()->HasActiveTool(EToolSide::Left))
+	{
+		const TSharedRef<FUICommandList> CommandList = Toolkit->GetToolkitCommands();
+		const FDataflowEditorCommandsImpl& CommandInfos = FDataflowEditorCommandsImpl::Get();
+
+		if (const TSharedPtr<const FUICommandInfo>* const Command = NodeTypeToToolCommandMap.Find(NodeTypeForPendingToolStart))
+		{
+			CommandList->TryExecuteAction(Command->ToSharedRef());
+		}
+
+		NodeTypeForPendingToolStart = FName();
+	}
+
+	if (PreviewScene)
+	{
+		PreviewScene->GetWorld()->Tick(ELevelTick::LEVELTICK_All, DeltaTime);
+	}
+
+}
+
+void UDataflowEditorMode::RestSpaceViewportResized(FViewport* RestspaceViewport, uint32 /*Unused*/)
+{
+	// We'd like to call RefocusRestSpaceViewportClient() when the viewport is first created, however in Ortho mode the
+	// viewport needs to have non-zero size for FocusViewportOnBox() to work properly. So we wait until the viewport is resized here.
+	if (bShouldFocusRestSpaceView && RestspaceViewport && RestspaceViewport->GetSizeXY().X > 0 && RestspaceViewport->GetSizeXY().Y > 0)
+	{
+		RefocusRestSpaceViewportClient();
+		bShouldFocusRestSpaceView = false;
+	}
+}
 
 FBox UDataflowEditorMode::SceneBoundingBox() const
 {
 	return PreviewScene->GetBoundingBox();
 }
 
+FBox UDataflowEditorMode::SelectionBoundingBox() const
+{
+	FBox Bounds = PreviewScene->SelectedComponentBounds();
+	if (Bounds.IsValid)
+	{
+		return Bounds;
+	}
+
+	// Nothing selected, return the whole scene
+	return SceneBoundingBox();
+}
+
+void UDataflowEditorMode::SetConstructionViewMode(Dataflow::EDataflowPatternVertexType InMode)
+{
+	// We will first check if there is an active tool. If so, we'll shut down the tool and save the results to the Node, then change view modes, then restart the tool again.
+	bool bEndedActiveTool = false;
+	UInteractiveToolManager* const ToolManager = GetInteractiveToolsContext()->ToolManager;
+	checkf(ToolManager, TEXT("No valid ToolManager found for UDataflowEditorMode"));
+	if (UInteractiveTool* const ActiveTool = ToolManager->GetActiveTool(EToolSide::Left))
+	{
+		// avoid switching back to the previous view mode when the tool ends here
+		const bool bTempShouldRestoreVal = bShouldRestoreSavedConstructionViewMode;
+		bShouldRestoreSavedConstructionViewMode = false;
+
+		ToolManager->PostActiveToolShutdownRequest(ActiveTool, EToolShutdownType::Accept);
+		bEndedActiveTool = true;
+
+		// now we can restore the previous view mode the next time the tool ends
+		bShouldRestoreSavedConstructionViewMode = bTempShouldRestoreVal;
+	}
+
+	ConstructionViewMode = InMode;
+	PreviewScene->ReinitializeDynamicMeshComponents();
+
+	const TSharedPtr<FDataflowEditorViewportClient> VC = ConstructionViewportClient.Pin();
+	if (VC.IsValid())
+	{
+		VC->SetConstructionViewMode(ConstructionViewMode);
+	}
+
+	// If we are switching to a mode with a valid mesh for the first time, focus the camera on it
+	FirstTimeFocusRestSpaceViewport();
+
+	if (bEndedActiveTool)
+	{
+		// If we ended the active tool in order to change modes, restart it now
+		if (const TSharedPtr<const SDataflowGraphEditor> PinnedGraphEditor = DataflowGraphEditor.Pin())
+		{
+			const FGraphPanelSelectionSet& SelectedNodes = PinnedGraphEditor->GetSelectedNodes();
+			if (SelectedNodes.Num() == 1)
+			{
+				StartToolForSelectedNode(*SelectedNodes.CreateConstIterator());
+			}
+		}
+	}
+}
+
+Dataflow::EDataflowPatternVertexType UDataflowEditorMode::GetConstructionViewMode() const
+{
+	return ConstructionViewMode;
+}
+
+
+bool UDataflowEditorMode::CanChangeConstructionViewModeTo(Dataflow::EDataflowPatternVertexType NewViewMode) const
+{
+	check(false);
+
+	if (!GetToolManager()->HasActiveTool(EToolSide::Left))
+	{
+		return true;
+	}
+
+	return false;
+	/*
+	const UInteractiveToolBuilder* const ActiveToolBuilder = GetToolManager()->GetActiveToolBuilder(EToolSide::Left);
+	checkf(ActiveToolBuilder, TEXT("No Active Tool Builder found despite having an Active Tool"));
+
+	const IDataflowEditorToolBuilder* const DataflowToolBuilder = Cast<const IDataflowEditorToolBuilder>(ActiveToolBuilder);
+	checkf(ClothToolBuilder, TEXT("Cloth Editor has an active Tool Builder that does not implement IDataflowEditorToolBuilder"));
+
+	TArray<Dataflow::EDataflowPatternVertexType> SupportedViewModes;
+	ClothToolBuilder->GetSupportedViewModes(SupportedViewModes);
+	return SupportedViewModes.Contains(NewViewMode);
+	*/
+}
+
+void UDataflowEditorMode::ToggleConstructionViewWireframe()
+{
+	check(false);
+	bConstructionViewWireframe = !bConstructionViewWireframe;
+	PreviewScene->ReinitializeDynamicMeshComponents();
+}
+
+bool UDataflowEditorMode::CanSetConstructionViewWireframeActive() const
+{
+	if (!GetToolManager()->HasActiveTool(EToolSide::Left))
+	{
+		return true;
+	}
+
+	const UInteractiveToolBuilder* const ActiveToolBuilder = GetToolManager()->GetActiveToolBuilder(EToolSide::Left);
+	checkf(ActiveToolBuilder, TEXT("No Active Tool Builder found despite having an Active Tool"));
+
+	const IDataflowEditorToolBuilder* const DataflowToolBuilder = Cast<const IDataflowEditorToolBuilder>(ActiveToolBuilder);
+	checkf(DataflowToolBuilder, TEXT("Cloth Editor has an active Tool Builder that does not implement IDataflowEditorToolBuilder"));
+	return DataflowToolBuilder->CanSetConstructionViewWireframeActive();
+}
+
+void UDataflowEditorMode::SetRestSpaceViewportClient(TWeakPtr<FDataflowEditorViewportClient, ESPMode::ThreadSafe> InViewportClient)
+{
+	ConstructionViewportClient = InViewportClient;
+
+	TSharedPtr<FDataflowEditorViewportClient> VC = ConstructionViewportClient.Pin();
+	if (VC.IsValid())
+	{
+		VC->SetConstructionViewMode(ConstructionViewMode);
+		VC->SetToolCommandList(ToolCommandList);
+
+		if (VC->Viewport)
+		{
+			VC->Viewport->ViewportResizedEvent.AddUObject(this, &UDataflowEditorMode::RestSpaceViewportResized);
+		}
+	}
+}
+
+void UDataflowEditorMode::InitializeContextObject()
+{
+	check(PreviewScene);
+
+	if (TObjectPtr<UDataflowEditorContent> DataflowEditorContent = PreviewScene->GetDataflowEditorContent())
+	{
+		UEditorInteractiveToolsContext* const RestSpaceToolsContext = GetInteractiveToolsContext();
+
+		UDataflowEditorContent* EditorContextObject = RestSpaceToolsContext->ContextObjectStore->FindContext<UDataflowEditorContent>();
+		if (!EditorContextObject)
+		{
+			EditorContextObject = PreviewScene->GetDataflowEditorContent();
+			RestSpaceToolsContext->ContextObjectStore->AddContextObject(PreviewScene->GetDataflowEditorContent());
+		}
+
+		check(EditorContextObject);
+
+		EditorContextObject->SetConstructionViewMode(ConstructionViewMode);
+	}
+}
+
+void UDataflowEditorMode::DeleteContextObject()
+{
+	UEditorInteractiveToolsContext* const RestSpaceToolsContext = GetInteractiveToolsContext();
+	if (UDataflowEditorContent* EditorContextObject = RestSpaceToolsContext->ContextObjectStore->FindContext<UDataflowEditorContent>())
+	{
+		RestSpaceToolsContext->ContextObjectStore->RemoveContextObject(EditorContextObject);
+	}
+}
+
+void UDataflowEditorMode::SetDataflowGraphEditor(TSharedPtr<SDataflowGraphEditor> InGraphEditor)
+{
+	if (InGraphEditor)
+	{
+		DataflowGraphEditor = InGraphEditor;
+		InitializeContextObject();
+	}
+	else
+	{
+		DeleteContextObject();
+	}
+}
+
+void UDataflowEditorMode::StartToolForSelectedNode(const UObject* SelectedNode)
+{
+	if (const UDataflowEdNode* const EdNode = Cast<UDataflowEdNode>(SelectedNode))
+	{
+		if (const TSharedPtr<const FDataflowNode> DataflowNode = EdNode->GetDataflowNode())
+		{
+			const FName DataflowNodeType = DataflowNode->GetType();
+			NodeTypeForPendingToolStart = DataflowNodeType;
+		}
+	}
+}
+
+
+void UDataflowEditorMode::OnDataflowNodeDeleted(const TSet<UObject*>& DeletedNodes)
+{
+	UEditorInteractiveToolsContext* const ToolsContext = GetInteractiveToolsContext();
+	checkf(ToolsContext, TEXT("No valid ToolsContext found for UDataflowEditorMode"));
+	const bool bCanCancel = ToolsContext->CanCancelActiveTool();
+	ToolsContext->EndTool(bCanCancel ? EToolShutdownType::Cancel : EToolShutdownType::Completed);
+}
+
+UEdGraphNode* UDataflowEditorMode::GetSingleSelectedNodeWithOutputType(const FName& SelectedNodeOutputTypeName) const
+{
+	const TSharedPtr<const SDataflowGraphEditor> PinnedDataflowGraphEditor = DataflowGraphEditor.Pin();
+	if (!PinnedDataflowGraphEditor)
+	{
+		return nullptr;
+	}
+
+	UEdGraphNode* const SelectedNode = PinnedDataflowGraphEditor->GetSingleSelectedNode();
+	if (!SelectedNode)
+	{
+		return nullptr;
+	}
+
+	const UDataflowEdNode* const SelectedDataflowEdNode = CastChecked<UDataflowEdNode>(SelectedNode);
+	const TSharedPtr<const FDataflowNode> SelectedDataflowNode = SelectedDataflowEdNode->GetDataflowNode();
+
+	if (!SelectedDataflowNode)
+	{
+		// This can happen when the user deletes a node. Seems like the Dataflow FGraph is updated with the removed node before the graph editor can update.
+		return nullptr;
+	}
+
+	for (const FDataflowOutput* const Output : SelectedDataflowNode->GetOutputs())
+	{
+		if (Output->GetType() == SelectedNodeOutputTypeName)
+		{
+			return SelectedNode;
+		}
+	}
+
+	return nullptr;
+}
+
+UEdGraphNode* UDataflowEditorMode::CreateNewNode(const FName& NewNodeTypeName)
+{
+	const TSharedPtr<const SDataflowGraphEditor> PinnedDataflowGraphEditor = DataflowGraphEditor.Pin();
+	if (!PinnedDataflowGraphEditor)
+	{
+		return nullptr;
+	}
+
+	if (TObjectPtr<UDataflowEditorContent> EditorContent = PreviewScene->GetDataflowEditorContent())
+	{
+		if (TObjectPtr<UDataflow> DataflowGraph = EditorContent->GetDataflowAsset())
+		{
+			const TSharedPtr<FAssetSchemaAction_Dataflow_CreateNode_DataflowEdNode> NodeAction =
+				FAssetSchemaAction_Dataflow_CreateNode_DataflowEdNode::CreateAction(DataflowGraph, NewNodeTypeName);
+			constexpr UEdGraphPin* FromPin = nullptr;
+			constexpr bool bSelectNewNode = true;
+			return NodeAction->PerformAction(DataflowGraph, FromPin, PinnedDataflowGraphEditor->GetPasteLocation(), bSelectNewNode);
+		}
+	}
+	return nullptr;
+}
+
+
+UEdGraphNode* UDataflowEditorMode::CreateAndConnectNewNode(const FName& NewNodeTypeName, UEdGraphNode& UpstreamNode, const FName& ConnectionTypeName, const FName& NewNodeConnectionName)
+{
+	if (TObjectPtr<UDataflowEditorContent> EditorContent = PreviewScene->GetDataflowEditorContent())
+	{
+		if (TObjectPtr<UDataflow> DataflowGraph = EditorContent->GetDataflowAsset())
+		{
+			// First find the specified output of the upstream node, plus any pins it's connected to
+
+			UEdGraphPin* UpstreamNodeOutputPin = nullptr;
+			TArray<UEdGraphPin*> ExistingNodeInputPins;
+
+			const UDataflowEdNode* const UpstreamDataflowEdNode = CastChecked<UDataflowEdNode>(&UpstreamNode);
+			const TSharedPtr<const FDataflowNode> UpstreamDataflowNode = UpstreamDataflowEdNode->GetDataflowNode();
+
+			for (const FDataflowOutput* const Output : UpstreamDataflowNode->GetOutputs())
+			{
+				if (Output->GetType() == ConnectionTypeName)
+				{
+					UpstreamNodeOutputPin = UpstreamDataflowEdNode->FindPin(*Output->GetName().ToString(), EGPD_Output);
+					ExistingNodeInputPins = UpstreamNodeOutputPin->LinkedTo;
+					break;
+				}
+			}
+
+			// Add the new node 
+
+			UEdGraphNode* const NewEdNode = CreateNewNode(NewNodeTypeName);
+			checkf(NewEdNode, TEXT("Failed to create a new node in the DataflowGraph"));
+
+			UDataflowEdNode* const NewDataflowEdNode = CastChecked<UDataflowEdNode>(NewEdNode);
+			const TSharedPtr<FDataflowNode> NewDataflowNode = NewDataflowEdNode->GetDataflowNode();
+
+			// Re-wire the graph
+
+			if (UpstreamNodeOutputPin)
+			{
+				UEdGraphPin* NewNodeInputPin = nullptr;
+				for (const FDataflowInput* const NewNodeInput : NewDataflowNode->GetInputs())
+				{
+					if (NewNodeInput->GetType() == ConnectionTypeName && NewNodeInput->GetName() == NewNodeConnectionName)
+					{
+						NewNodeInputPin = NewDataflowEdNode->FindPin(*NewNodeInput->GetName().ToString(), EGPD_Input);
+					}
+				}
+
+				UEdGraphPin* NewNodeOutputPin = nullptr;
+				for (const FDataflowOutput* const NewNodeOutput : NewDataflowNode->GetOutputs())
+				{
+					if (NewNodeOutput->GetType() == ConnectionTypeName && NewNodeOutput->GetName() == NewNodeConnectionName)
+					{
+						NewNodeOutputPin = NewDataflowEdNode->FindPin(*NewNodeOutput->GetName().ToString(), EGPD_Output);
+						break;
+					}
+				}
+
+				check(NewNodeInputPin);
+				check(NewNodeOutputPin);
+
+				DataflowGraph->GetSchema()->TryCreateConnection(UpstreamNodeOutputPin, NewNodeInputPin);
+
+				for (UEdGraphPin* DownstreamInputPin : ExistingNodeInputPins)
+				{
+					DataflowGraph->GetSchema()->TryCreateConnection(NewNodeOutputPin, DownstreamInputPin);
+				}
+			}
+
+			DataflowGraph->NotifyGraphChanged();
+
+			return NewEdNode;
+
+		}
+	}
+	return nullptr;
+}
+
 #undef LOCTEXT_NAMESPACE
 
diff --git a/Engine/Plugins/Experimental/Dataflow/Source/DataflowEditor/Private/Dataflow/DataflowEditorModeToolkit.cpp b/Engine/Plugins/Experimental/Dataflow/Source/DataflowEditor/Private/Dataflow/DataflowEditorModeToolkit.cpp
--- a/Engine/Plugins/Experimental/Dataflow/Source/DataflowEditor/Private/Dataflow/DataflowEditorModeToolkit.cpp
+++ b/Engine/Plugins/Experimental/Dataflow/Source/DataflowEditor/Private/Dataflow/DataflowEditorModeToolkit.cpp
@@ -4,35 +4,144 @@
 
 #include "AssetEditorModeManager.h"
 #include "Dataflow/DataflowEditorCommands.h"
+#include "Dataflow/DataflowEditorMode.h"
 #include "Dataflow/DataflowEditorStyle.h"
+#include "Dataflow/DataflowEditorViewport.h"
+#include "Dataflow/DataflowPreviewScene.h"
+#include "EdModeInteractiveToolsContext.h"
+#include "Framework/Application/SlateApplication.h"
 #include "InteractiveToolManager.h"
 #include "MeshAttributePaintTool.h"
+#include "SBaseCharacterFXEditorViewport.h"
+#include "SPrimaryButton.h"
+#include "Tools/UEdMode.h"
+#include "Widgets/Images/SImage.h"
 
 #define LOCTEXT_NAMESPACE "FDataflowEditorModeToolkit"
 
-void FDataflowEditorModeToolkit::OnToolStarted(UInteractiveToolManager* Manager, UInteractiveTool* Tool)
+FName FDataflowEditorModeToolkit::GetToolkitFName() const
 {
-	FBaseCharacterFXEditorModeToolkit::OnToolStarted(Manager, Tool);
+	return FName("DataflowEditorMode");
 }
 
-void FDataflowEditorModeToolkit::BuildToolPalette(FName PaletteIndex, class FToolBarBuilder& ToolbarBuilder)
+FText FDataflowEditorModeToolkit::GetBaseToolkitName() const
 {
-	const FDataflowEditorCommandsImpl& Commands = FDataflowEditorCommands::Get();
-	if (PaletteIndex == ToolsTabName)
-	{
-		ToolbarBuilder.AddToolBarButton(Commands.BeginAttributeEditorTool);
-		ToolbarBuilder.AddToolBarButton(Commands.BeginMeshSelectionTool);
-	}
+	return NSLOCTEXT("DataflowEditorModeToolkit", "DisplayName", "DataflowEditorMode");
 }
 
-FName FDataflowEditorModeToolkit::GetToolkitFName() const
+void FDataflowEditorModeToolkit::Init(const TSharedPtr<IToolkitHost>& InitToolkitHost, TWeakObjectPtr<UEdMode> InOwningMode)
 {
-	return FName("DataflowEditorMode");
+	FModeToolkit::Init(InitToolkitHost, InOwningMode);	
+
+	// Set up tool message areas
+	ClearNotification();
+	ClearWarning();
+	GetScriptableEditorMode()->GetInteractiveToolsContext()->OnToolNotificationMessage.AddSP(this, &FBaseCharacterFXEditorModeToolkit::PostNotification);
+	GetScriptableEditorMode()->GetInteractiveToolsContext()->OnToolWarningMessage.AddSP(this, &FBaseCharacterFXEditorModeToolkit::PostWarning);
+
+	// Hook up the tool detail panel
+	ToolDetailsContainer->SetContent(DetailsView.ToSharedRef());
+	
+	// Set up the overlay -- it will poll different ToolsContexts depending on which one launched the currently active tool
+
+	SAssignNew(ViewportOverlayWidget, SHorizontalBox)
+
+	+SHorizontalBox::Slot()
+	.HAlign(HAlign_Center)
+	.VAlign(VAlign_Bottom)
+	.Padding(FMargin(0.0f, 0.0f, 0.f, 15.f))
+	[
+		SNew(SBorder)
+		.BorderImage(FAppStyle::Get().GetBrush("EditorViewport.OverlayBrush"))
+		.Padding(8.f)
+		[
+			SNew(SHorizontalBox)
+
+			+SHorizontalBox::Slot()
+			.AutoWidth()
+			.VAlign(VAlign_Center)
+			.Padding(FMargin(0.f, 0.f, 8.f, 0.f))
+			[
+				SNew(SImage)
+				.Image_Lambda([this] () { return ActiveToolIcon; })
+			]
+
+			+SHorizontalBox::Slot()
+			.AutoWidth()
+			.VAlign(VAlign_Center)
+			.Padding(FMargin(0.f, 0.f, 8.f, 0.f))
+			[
+				SNew(STextBlock)
+				.Text(this, &FBaseCharacterFXEditorModeToolkit::GetActiveToolDisplayName)
+			]
+
+			+SHorizontalBox::Slot()
+			.AutoWidth()
+			.Padding(FMargin(0.0, 0.f, 2.f, 0.f))
+			[
+				SNew(SPrimaryButton)
+				.Text(LOCTEXT("OverlayAccept", "Accept"))
+				.ToolTipText(LOCTEXT("OverlayAcceptTooltip", "Accept/Commit the results of the active Tool [Enter]"))
+				.OnClicked_Lambda([this]() { 
+				    GetCurrentToolsContext()->EndTool(EToolShutdownType::Accept);
+					return FReply::Handled(); 
+					})
+				.IsEnabled_Lambda([this]() { return GetCurrentToolsContext()->CanAcceptActiveTool(); })
+				.Visibility_Lambda([this]() { return GetCurrentToolsContext()->ActiveToolHasAccept() ? EVisibility::Visible : EVisibility::Collapsed; })
+			]
+
+			+SHorizontalBox::Slot()
+			.AutoWidth()
+			.Padding(FMargin(2.0, 0.f, 0.f, 0.f))
+			[
+				SNew(SButton)
+				.ButtonStyle(FAppStyle::Get(), "Button")
+				.TextStyle( FAppStyle::Get(), "DialogButtonText" )
+				.Text(LOCTEXT("OverlayCancel", "Cancel"))
+				.ToolTipText(LOCTEXT("OverlayCancelTooltip", "Cancel the active Tool [Esc]"))
+				.HAlign(HAlign_Center)
+				.OnClicked_Lambda([this]() { 
+					GetCurrentToolsContext()->EndTool(EToolShutdownType::Cancel);
+					return FReply::Handled(); 
+					})
+				.IsEnabled_Lambda([this]() { return GetCurrentToolsContext()->CanCancelActiveTool(); })
+				.Visibility_Lambda([this]() { return GetCurrentToolsContext()->ActiveToolHasAccept() ? EVisibility::Visible : EVisibility::Collapsed; })
+			]
+
+			+ SHorizontalBox::Slot()
+			.AutoWidth()
+			.Padding(FMargin(2.0, 0.f, 0.f, 0.f))
+			[
+				SNew(SButton)
+				.ButtonStyle(FAppStyle::Get(), "PrimaryButton")
+				.TextStyle(FAppStyle::Get(), "DialogButtonText")
+				.Text(LOCTEXT("OverlayComplete", "Complete"))
+				.ToolTipText(LOCTEXT("OverlayCompleteTooltip", "Exit the active Tool [Enter]"))
+				.HAlign(HAlign_Center)
+				.OnClicked_Lambda([this]() { 
+					GetCurrentToolsContext()->EndTool(EToolShutdownType::Completed);
+					return FReply::Handled(); 
+					})
+				.IsEnabled_Lambda([this]() {
+					return GetCurrentToolsContext()->CanCompleteActiveTool();
+				})
+				.Visibility_Lambda([this]() { return GetCurrentToolsContext()->CanCompleteActiveTool() ? EVisibility::Visible : EVisibility::Collapsed; })
+			]
+		]	
+	];
 }
 
-FText FDataflowEditorModeToolkit::GetBaseToolkitName() const
+void FDataflowEditorModeToolkit::BuildEditorToolBar(const FName& EditorToolBarName)
 {
-	return NSLOCTEXT("DataflowEditorModeToolkit", "DisplayName", "DataflowEditorMode");
+	check(EditorToolBarName != FName());
+
+	const TSharedRef<const FUICommandList> EdModeToolkitCommands = GetToolkitCommands();
+
+	UToolMenu* const ToolBarMenu = UToolMenus::Get()->ExtendMenu(EditorToolBarName);
+	FToolMenuSection& Section = ToolBarMenu->FindOrAddSection("DataflowTools");
+
+	FToolMenuEntry& WeightMapButtonEntry = Section.AddEntry(FToolMenuEntry::InitToolBarButton(FDataflowEditorCommandsImpl::Get().AddWeightMapNode));
+	WeightMapButtonEntry.SetCommandList(EdModeToolkitCommands);
 }
 
 const FSlateBrush* FDataflowEditorModeToolkit::GetActiveToolIcon(const FString& ActiveToolIdentifier) const
@@ -41,4 +150,119 @@ const FSlateBrush* FDataflowEditorModeToolkit::GetActiveToolIcon(const FString&
 	return FDataflowEditorStyle::Get().GetOptionalBrush(ActiveToolIconName);
 }
 
+SBaseCharacterFXEditorViewport* FDataflowEditorModeToolkit::GetViewportWidgetForManager(UInteractiveToolManager* Manager)
+{
+	if (OwningEditorMode.IsValid(false))
+	{
+		const UEdMode* const Mode = OwningEditorMode.Get();
+
+		if (const UDataflowEditorMode* const DataflowEdMode = Cast<UDataflowEditorMode>(Mode))
+		{
+			if (const FDataflowPreviewScene* const PreviewScene = DataflowEdMode->GetDataflowPreviewScene())
+			{
+				if (const UEditorInteractiveToolsContext* const PreviewToolsContext = PreviewScene->GetDataflowModeManager()->GetInteractiveToolsContext())
+				{
+					const UInteractiveToolManager* const PreviewToolManager = PreviewToolsContext->ToolManager;
+
+					//@todo(SimulationViewport) : Update to include simulation view when its added
+					//if (Manager == PreviewToolManager)
+					//{
+					//	if (const TSharedPtr<SChaosClothAssetEditor3DViewport> Widget = PreviewViewportWidget.Pin())
+					//	{
+					//		return Widget.Get();
+					//	}
+					//}
+					//else
+					//{
+					if (const TSharedPtr<SDataflowEditorViewport> Widget = RestSpaceViewportWidget.Pin())
+					{
+						return Widget.Get();
+					}
+					//}
+				}
+			}
+		}
+	}
+
+	return nullptr;
+}
+
+UEditorInteractiveToolsContext* FDataflowEditorModeToolkit::GetCurrentToolsContext()
+{
+	if (OwningEditorMode.IsValid(/*bEvenIfPendingKill*/ false))
+	{
+		UEdMode* const Mode = OwningEditorMode.Get();
+		if (UDataflowEditorMode* const DataflowEdMode = Cast<UDataflowEditorMode>(Mode))
+		{
+			return DataflowEdMode->GetActiveToolsContext();
+		}
+	}
+
+	// this should not happen, but don't crash if it does
+	return GetEditorModeManager().GetInteractiveToolsContext();
+}
+void FDataflowEditorModeToolkit::OnToolStarted(UInteractiveToolManager* Manager, UInteractiveTool* Tool)
+{
+	ensure(Tool == Manager->GetActiveTool(EToolSide::Left));
+
+	if (!Tool)
+	{
+		return;
+	}
+
+	if (DetailsView)
+	{
+		DetailsView->SetObjects(Tool->GetToolProperties());
+
+		Tool->OnPropertySetsModified.AddLambda([this, Manager]()
+			{
+				const UInteractiveTool* const Tool = Manager->GetActiveTool(EToolSide::Left);
+		if (Tool != nullptr)
+		{
+			DetailsView->SetObjects(Tool->GetToolProperties(true));
+		}
+			});
+	}
+
+	Tool->OnPropertyModifiedDirectlyByTool.AddSP(this, &FDataflowEditorModeToolkit::InvalidateCachedDetailPanelState);
+
+	ActiveToolName = Tool->GetToolInfo().ToolDisplayName;
+	FString ActiveToolIdentifier = Manager->GetActiveToolName(EToolSide::Mouse);
+	ActiveToolIdentifier.InsertAt(0, ".");
+	ActiveToolIcon = GetActiveToolIcon(ActiveToolIdentifier);
+
+	if (SBaseCharacterFXEditorViewport* Widget = GetViewportWidgetForManager(Manager))
+	{
+		Widget->AddOverlayWidget(ViewportOverlayWidget.ToSharedRef());
+
+		// Manually set keyboard focus to this viewport. Otherwise the user has to click on the viewport before Tool keyboard shortcuts will work.
+		FSlateApplication::Get().SetUserFocus(FSlateApplication::Get().GetUserIndexForKeyboard(), Widget->AsShared(), EFocusCause::SetDirectly);
+	}
+}
+
+void FDataflowEditorModeToolkit::OnToolEnded(UInteractiveToolManager* Manager, UInteractiveTool* Tool)
+{
+	FModeToolkit::OnToolEnded(Manager, Tool);
+
+	ActiveToolName = FText::GetEmpty();
+	ClearNotification();
+	ClearWarning();
+
+	if (SBaseCharacterFXEditorViewport* const Widget = GetViewportWidgetForManager(Manager))
+	{
+		Widget->RemoveOverlayWidget(ViewportOverlayWidget.ToSharedRef());
+	}
+
+	if (UInteractiveTool* const CurTool = Manager->GetActiveTool(EToolSide::Left))
+	{
+		CurTool->OnPropertySetsModified.RemoveAll(this);
+		CurTool->OnPropertyModifiedDirectlyByTool.RemoveAll(this);
+	}
+}
+
+void FDataflowEditorModeToolkit::SetRestSpaceViewportWidget(TWeakPtr<SDataflowEditorViewport> InRestSpaceViewportWidget)
+{
+	RestSpaceViewportWidget = InRestSpaceViewportWidget;
+}
+
 #undef LOCTEXT_NAMESPACE
diff --git a/Engine/Plugins/Experimental/Dataflow/Source/DataflowEditor/Private/Dataflow/DataflowEditorModeUILayer.cpp b/Engine/Plugins/Experimental/Dataflow/Source/DataflowEditor/Private/Dataflow/DataflowEditorModeUILayer.cpp
new file mode 100644
--- /dev/null
+++ b/Engine/Plugins/Experimental/Dataflow/Source/DataflowEditor/Private/Dataflow/DataflowEditorModeUILayer.cpp
@@ -0,0 +1,33 @@
+// Copyright Epic Games, Inc. All Rights Reserved.
+
+#include "Dataflow/DataflowEditorModeUILayer.h"
+#include "Dataflow/DataflowEditorToolkit.h"
+#include "Dataflow/DataflowEditorModule.h"
+#include "WorkspaceMenuStructure.h"
+#include "WorkspaceMenuStructureModule.h"
+#include "Toolkits/IToolkit.h"
+
+#include UE_INLINE_GENERATED_CPP_BY_NAME(DataflowEditorModeUILayer)
+
+const FName UDataflowEditorUISubsystem::EditorSidePanelAreaName = "DataflowEditorSidePanelArea";
+
+
+void UDataflowEditorUISubsystem::Initialize(FSubsystemCollectionBase& Collection)
+{
+	FDataflowEditorModule& DataflowEditorModule = FModuleManager::GetModuleChecked<FDataflowEditorModule>("DataflowEditor");
+	DataflowEditorModule.OnRegisterLayoutExtensions().AddUObject(this, &UDataflowEditorUISubsystem::RegisterLayoutExtensions);
+}
+
+void UDataflowEditorUISubsystem::Deinitialize()
+{
+	FDataflowEditorModule& DataflowEditorModule = FModuleManager::GetModuleChecked<FDataflowEditorModule>("DataflowEditor");
+	DataflowEditorModule.OnRegisterLayoutExtensions().RemoveAll(this);
+}
+
+void UDataflowEditorUISubsystem::RegisterLayoutExtensions(FLayoutExtender& Extender)
+{
+	FTabManager::FTab NewTab(FTabId(UAssetEditorUISubsystem::TopLeftTabID), ETabState::ClosedTab);
+	Extender.ExtendStack(UDataflowEditorUISubsystem::EditorSidePanelAreaName, ELayoutExtensionPosition::After, NewTab);
+}
+
+
diff --git a/Engine/Plugins/Experimental/Dataflow/Source/DataflowEditor/Private/Dataflow/DataflowEditorStyle.h b/Engine/Plugins/Experimental/Dataflow/Source/DataflowEditor/Private/Dataflow/DataflowEditorStyle.h
--- a/Engine/Plugins/Experimental/Dataflow/Source/DataflowEditor/Private/Dataflow/DataflowEditorStyle.h
+++ b/Engine/Plugins/Experimental/Dataflow/Source/DataflowEditor/Private/Dataflow/DataflowEditorStyle.h
@@ -21,6 +21,8 @@ public:
 		const FVector2D Icon40x40(40.f, 40.f);
 		const FVector2D Icon64x64(64.f, 64.f);
 
+		const FVector2D ToolbarIconSize = Icon20x20;
+
 		SetContentRoot(IPluginManager::Get().FindPlugin("Dataflow")->GetBaseDir() / TEXT("Resources"));
 		Set("ClassIcon.Dataflow", new FSlateVectorImageBrush(RootToContentDir(TEXT("DataflowAsset_16.svg")), Icon16x16));
 		Set("ClassThumbnail.Dataflow", new FSlateVectorImageBrush(RootToContentDir(TEXT("DataflowAsset_64.svg")), Icon64x64));
@@ -36,11 +38,14 @@ public:
 		Set("Dataflow.SelectFace", new FSlateImageBrush(RootToContentDir(TEXT("Slate/Dataflow_SelectFace_40x.png")), Icon40x40));
 		Set("Dataflow.SelectVertex", new FSlateImageBrush(RootToContentDir(TEXT("Slate/Dataflow_SelectVertex_40x.png")), Icon40x40));
 
-		const FString AttributeEditorPropertyName = "DataflowEditor." + FDataflowEditorCommandsImpl::BeginAttributeEditorToolIdentifier;
-		Set(*AttributeEditorPropertyName, new FSlateImageBrush(RootToContentDir(TEXT("Slate/Dataflow_SelectObject_20x.png")), Icon20x20));
+		// @todo(brice) Remove Example Tools
+		//const FString AttributeEditorPropertyName = "DataflowEditor." + FDataflowEditorCommandsImpl::BeginAttributeEditorToolIdentifier;
+		//Set(*AttributeEditorPropertyName, new FSlateImageBrush(RootToContentDir(TEXT("Slate/Dataflow_SelectObject_20x.png")), Icon20x20));
+		//const FString MeshSelectionPropertyName = "DataflowEditor." + FDataflowEditorCommandsImpl::BeginMeshSelectionToolIdentifier;
+		//Set(*MeshSelectionPropertyName, new FSlateImageBrush(RootToContentDir(TEXT("Slate/Dataflow_SelectVertex_20x.png")), Icon20x20));
 
-		const FString MeshSelectionPropertyName = "DataflowEditor." + FDataflowEditorCommandsImpl::BeginMeshSelectionToolIdentifier;
-		Set(*MeshSelectionPropertyName, new FSlateImageBrush(RootToContentDir(TEXT("Slate/Dataflow_SelectVertex_20x.png")), Icon20x20));
+		FString PropertyNameString = "DataflowEditor." + FDataflowEditorCommandsImpl::AddWeightMapNodeIdentifier;
+		Set(*PropertyNameString, new IMAGE_BRUSH_SVG("Icons/PaintMaps", ToolbarIconSize));
 
 		DefaultMaterial = Cast<UMaterial>(StaticLoadObject(UMaterial::StaticClass(), NULL, TEXT("/Engine/BasicShapes/BasicShapeMaterial")));
 
diff --git a/Engine/Plugins/Experimental/Dataflow/Source/DataflowEditor/Private/Dataflow/DataflowEditorToolkit.cpp b/Engine/Plugins/Experimental/Dataflow/Source/DataflowEditor/Private/Dataflow/DataflowEditorToolkit.cpp
--- a/Engine/Plugins/Experimental/Dataflow/Source/DataflowEditor/Private/Dataflow/DataflowEditorToolkit.cpp
+++ b/Engine/Plugins/Experimental/Dataflow/Source/DataflowEditor/Private/Dataflow/DataflowEditorToolkit.cpp
@@ -9,6 +9,8 @@
 #include "Dataflow/DataflowEditorContent.h"
 #include "Dataflow/DataflowEditorCommands.h"
 #include "Dataflow/DataflowEditorMode.h"
+#include "Dataflow/DataflowEditorModeToolkit.h"
+#include "Dataflow/DataflowEditorModeUILayer.h"
 #include "Dataflow/DataflowEditorViewport.h"
 #include "Dataflow/DataflowEditorViewportClient.h"
 #include "Dataflow/DataflowEdNode.h"
@@ -24,6 +26,7 @@
 #include "EditorViewportTabContent.h"
 #include "EditorViewportLayout.h"
 #include "EditorViewportCommands.h"
+#include "EdModeInteractiveToolsContext.h"
 #include "Engine/SkeletalMesh.h"
 #include "GraphEditorActions.h"
 #include "Modules/ModuleManager.h"
@@ -58,38 +61,53 @@ FDataflowEditorToolkit::FDataflowEditorToolkit(UAssetEditor* InOwningAssetEditor
 {
 	check(Cast<UDataflowEditor>(InOwningAssetEditor));
 
-	StandaloneDefaultLayout = FTabManager::NewLayout("Dataflow_Layout.V5")
+	StandaloneDefaultLayout = FTabManager::NewLayout(FName("DataflowEditorLayout.V9"))
 		->AddArea
 		(
-			FTabManager::NewPrimaryArea()->SetOrientation(Orient_Vertical)
+			FTabManager::NewPrimaryArea()->SetOrientation(Orient_Horizontal)
 			->Split
 			(
 				FTabManager::NewSplitter()->SetOrientation(Orient_Vertical)
+				->SetSizeCoefficient(0.8f)	// Relative width of (Tools Panel, Construction Viewport, Preview Viewport, Dataflow Graph Editor, Outliner) vs (Asset Details, Preview Scene Details, Dataflow Node Details)
 				->Split
 				(
 					FTabManager::NewSplitter()->SetOrientation(Orient_Horizontal)
-					->SetSizeCoefficient(0.9f)
+					->SetSizeCoefficient(0.55f)	// Relative height of (Tools Panel, Construction Viewport, Preview Viewport) vs (Dataflow Graph Editor, Outliner)
 					->Split
 					(
 						FTabManager::NewStack()
-						->SetSizeCoefficient(0.6f)
-						->AddTab(GraphCanvasTabId, ETabState::OpenedTab)
+						->SetSizeCoefficient(0.1f)		// Relative width of (Tools Panel) vs (Construction Viewport, Preview Viewport)
+						->SetExtensionId(UDataflowEditorUISubsystem::EditorSidePanelAreaName)
+						->SetHideTabWell(true)
 					)
 					->Split
 					(
-						FTabManager::NewSplitter()->SetOrientation(Orient_Vertical)
-						->SetSizeCoefficient(0.2f)
-						->Split
-						(
-							FTabManager::NewStack()
-							->SetSizeCoefficient(0.7f)
-							->AddTab(DetailsTabID, ETabState::OpenedTab)
-						)
+						FTabManager::NewStack()
+						->SetSizeCoefficient(0.9f)		// Relative width of (Construction Viewport) vs (Tools Panel, Preview Viewport)
+						->AddTab(ViewportTabID, ETabState::OpenedTab)
+						->SetExtensionId("RestSpaceViewportArea")
+						->SetHideTabWell(false)
 					)
 				)
+				->Split
+				(
+					FTabManager::NewStack()
+					->SetSizeCoefficient(0.45f)	// Relative height of (Dataflow Graph Editor, Outliner) vs (Tools Panel, Construction Viewport, Preview Viewport)
+					->AddTab(GraphCanvasTabId, ETabState::OpenedTab)
+					->SetExtensionId("GraphEditorArea")
+					->SetHideTabWell(false)
+					->SetForegroundTab(GraphCanvasTabId)
+				)
+			)
+			->Split
+			(
+				FTabManager::NewStack()
+				->SetSizeCoefficient(0.2f)	// Relative height of (Dataflow Node Details) vs (Asset Details, Preview Scene Details)
+				->AddTab(NodeDetailsTabId, ETabState::OpenedTab)
+				->SetExtensionId("NodeDetailsArea")
+				->SetHideTabWell(false)
 			)
 		);
-
 	FAdvancedPreviewScene::ConstructionValues PreviewSceneArgs;
 	PreviewSceneArgs.bShouldSimulatePhysics = 1;
 	PreviewSceneArgs.bCreatePhysicsScene = 1;
@@ -270,32 +288,58 @@ void FDataflowEditorToolkit::PostInitAssetEditor()
 	};
 	SetCommonViewportClientOptions(ViewportClient.Get());
 
+	//@todo(brice) : FChaosClothAssetEditorToolkit::PostInitAssetEditor() has alot more going on. 
+
 	// Set up 3D viewport
 	TSharedPtr<FDataflowEditorViewportClient> DataflowViewportClient = StaticCastSharedPtr<FDataflowEditorViewportClient>(ViewportClient);
 	DataflowViewportClient->SetDataflowEditorToolkit(StaticCastSharedRef<FDataflowEditorToolkit>(this->AsShared()));
 
+	//@todo(brice) : can we remove this coupling of the viewport and mode
+	UDataflowEditorMode* const DataflowMode = CastChecked<UDataflowEditorMode>(EditorModeManager->GetActiveScriptableMode(UDataflowEditorMode::EM_DataflowEditorModeId));
+	const TWeakPtr<FViewportClient> WeakViewportClient(ViewportClient);
+	DataflowMode->SetRestSpaceViewportClient(StaticCastWeakPtr<FDataflowEditorViewportClient>(WeakViewportClient));
+
+	check(GetDataflowEditorContent());
+
+
 	//SetEditingObject(GetDataflowEditorContent().DataflowOwner);
 
-	// @todo(DataflowMode) : Do this in UDataflowEditorMode::RefocusViewportClient and use the DataflowComponents bounds
-	//		... see ChaosClothAssetEditorMode::RefocusRestSpaceViewportClient() for reference
-	// FBoxSphereBounds SphereBounds = FBoxSphereBounds(EForceInit::ForceInitToZero);
-	// DataflowViewportClient->OverrideFarClipPlane(0);
-	// DataflowViewportClient->FocusViewportOnBox(SphereBounds.GetBox());
 }
 
 void FDataflowEditorToolkit::InitializeEdMode(UBaseCharacterFXEditorMode* EdMode)
 {
+
+	UDataflowEditorMode* DataflowMode = Cast<UDataflowEditorMode>(EdMode);
+	check(DataflowMode);
+
 	// We first set the preview scene in order to store the dynamic mesh elements
 	// generated by the tools
-	Cast<UDataflowEditorMode>(EdMode)->SetDataflowPreviewScene(
-		static_cast<FDataflowPreviewScene*>(ObjectScene.Get()));
+	DataflowMode->SetDataflowPreviewScene(static_cast<FDataflowPreviewScene*>(ObjectScene.Get()));
 
 	// Set of the graph editor to be able to add nodes
-	Cast<UDataflowEditorMode>(EdMode)->SetDataflowGraphEditor(GraphEditor);
+	DataflowMode->SetDataflowGraphEditor(GraphEditor);
 	
-	FBaseCharacterFXEditorToolkit::InitializeEdMode(EdMode);
+	TArray<TObjectPtr<UObject>> ObjectsToEdit;
+	OwningAssetEditor->GetObjectsToEdit(MutableView(ObjectsToEdit));
+	DataflowMode->InitializeTargets(ObjectsToEdit);
+
+	if (TSharedPtr<FModeToolkit> ModeToolkit = DataflowMode->GetToolkit().Pin())
+	{
+		FDataflowEditorModeToolkit* DataflowModeToolkit = static_cast<FDataflowEditorModeToolkit*>(ModeToolkit.Get());
+		DataflowModeToolkit->SetRestSpaceViewportWidget(DataflowEditorViewport);
+		//ClothModeToolkit->SetPreviewViewportWidget(PreviewViewportWidget);
+
+		FName ParentToolbarName;
+		const FName ToolBarName = GetToolMenuToolbarName(ParentToolbarName);
+		DataflowModeToolkit->BuildEditorToolBar(ToolBarName);
+	}
+
+	// @todo(brice) : This used to crash when comnmented out. 
+	//FBaseCharacterFXEditorToolkit::InitializeEdMode(EdMode);
 }
 
+
+
 void FDataflowEditorToolkit::CreateEditorModeUILayer()
 {
 	FBaseCharacterFXEditorToolkit::CreateEditorModeUILayer();
@@ -373,7 +417,7 @@ void FDataflowEditorToolkit::OnPropertyValueChanged(const FPropertyChangedEvent&
 		{
 			TSharedPtr<Dataflow::FEngineContext>& DataflowContext = EditorContent->GetDataflowContext();
 			Dataflow::FTimestamp& LastNodeTimestamp = EditorContent->GetLastModifiedTimestamp();
-			FDataflowEditorCommands::OnPropertyValueChanged(DataflowAsset, DataflowContext, LastNodeTimestamp, PropertyChangedEvent, PrevNodeSelection);
+			FDataflowEditorCommands::OnPropertyValueChanged(DataflowAsset, DataflowContext, LastNodeTimestamp, PropertyChangedEvent, SelectedDataflowNodes);
 		}
 	}
 }
@@ -396,70 +440,157 @@ void FDataflowEditorToolkit::OnNodeTitleCommitted(const FText& InNewText, ETextC
 	FDataflowEditorCommands::OnNodeTitleCommitted(InNewText, InCommitType, GraphNode);
 }
 
-void FDataflowEditorToolkit::OnNodeSelectionChanged(const TSet<UObject*>& NewSelection)
+void FDataflowEditorToolkit::OnNodeSelectionChanged(const TSet<UObject*>& InNewSelection)
 {
-	if (TObjectPtr<UDataflowEditorContent> EditorContent = GetDataflowEditorContent(); EditorContent->GetDataflowAsset())
-	{
-		// Only keep UDataflowEdNode from NewSelection
-		TSet<UObject*> ValidatedSelection;
 
-		for (UObject* Item : NewSelection)
+	auto FindDataflowNodesInSet = [](const TSet<UObject*>& InSet) {
+		TSet<UObject*> Results;
+		for (UObject* Item : InSet)
 		{
-			if (UDataflowEdNode* Node = Cast<UDataflowEdNode>(Item))
+			if (Cast<UDataflowEdNode>(Item))
 			{
-				ValidatedSelection.Add(Item);
+				Results.Add(Item);
 			}
 		}
+		return Results;
+	};
 
-		if (ValidatedSelection.Num() > 0)
+	auto ResetListeners = [&ViewListeners = ViewListeners](UDataflowEdNode* Node = nullptr)
+	{
+		for (IDataflowViewListener* Listener : ViewListeners)
 		{
-			TSet<UObject*> SelectionToUse, SelectionDifference;
-
-			if (PrevNodeSelection.Num() > 0)
+			Listener->OnSelectedNodeChanged(nullptr);
+		}
+		if (Node)
+		{
+			for (IDataflowViewListener* Listener : ViewListeners)
 			{
-				SelectionDifference = ValidatedSelection.Difference(PrevNodeSelection);
+				Listener->OnSelectedNodeChanged(Node);
+			}
+		}
+	};
 
-				if (SelectionDifference.Num() > 0)
+	auto EvaluateCollection = [](TObjectPtr<UDataflowEdNode> InDataflowEdNode, const TSharedPtr<Dataflow::FEngineContext> Context)
+	{
+		if (InDataflowEdNode && Context)
+		{
+			if (const TSharedPtr<FDataflowNode> InDataflowNode = InDataflowEdNode->GetDataflowNode())
+			{
+				for (const FDataflowOutput* const Output : InDataflowNode->GetOutputs())
 				{
-					SelectionToUse.Add(SelectionDifference.Array()[SelectionDifference.Num() - 1]);
+					if (Output->GetType() == FName("FManagedArrayCollection"))
+					{
+						const FManagedArrayCollection DefaultValue;
+						return TSharedPtr<FManagedArrayCollection>(new FManagedArrayCollection(Output->GetValue<FManagedArrayCollection>(*Context, DefaultValue)));
+					}
 				}
-				else
+			}
+		}
+
+		return TSharedPtr<FManagedArrayCollection>(new FManagedArrayCollection());
+	};
+
+	// Despite this function's name, we might not have actually changed which node is selected
+	bool bPrimarySelectionChanged = false;
+
+	if (TObjectPtr<UDataflowEditorContent> EditorContent = GetDataflowEditorContent(); EditorContent->GetDataflowAsset())
+	{
+		// Only keep UDataflowEdNode from NewSelection
+		TSet<UObject*> NodeSelection = FindDataflowNodesInSet(InNewSelection);
+
+		if (!NodeSelection.Num())
+		{
+			// The selection is empty. 
+			ResetListeners();
+			SelectedDataflowNodes = TSet<UObject*>();
+			PrimarySelection = nullptr;
+		}
+		else
+		{
+			TSet<UObject*> DeselectedNodes = SelectedDataflowNodes.Difference(NodeSelection);
+			TSet<UObject*> StillSelectedNodes = SelectedDataflowNodes.Intersect(NodeSelection);
+			TSet<UObject*> NewlySelectedNodes = NodeSelection.Difference(SelectedDataflowNodes);
+
+			// Something has been removed
+			if (DeselectedNodes.Num())
+			{
+				if (DeselectedNodes.Contains(PrimarySelection))
 				{
-					SelectionToUse.Add(ValidatedSelection.Array()[ValidatedSelection.Num() - 1]);
+					ResetListeners();
+
+					if (PrimarySelection) bPrimarySelectionChanged = true;
+					PrimarySelection = nullptr;
+
+					// pick a new primary if nothing new was selected
+					if (!NewlySelectedNodes.Num() && StillSelectedNodes.Num())
+					{
+						PrimarySelection = Cast< UDataflowEdNode>(StillSelectedNodes.Array()[0]);
+						ResetListeners(PrimarySelection);
+						bPrimarySelectionChanged = true;
+					}
 				}
 			}
-			else
+
+			// Something new has been selected.
+			if (NewlySelectedNodes.Num() == 1)
 			{
-				SelectionToUse.Add(ValidatedSelection.Array()[ValidatedSelection.Num() - 1]);
+				PrimarySelection = Cast< UDataflowEdNode>(NewlySelectedNodes.Array()[0]);
+				ResetListeners(PrimarySelection);
+				bPrimarySelectionChanged = true;
 			}
 
-			if (UDataflowEdNode* Node = Cast<UDataflowEdNode>(SelectionToUse.Array()[0]))
+			SelectedDataflowNodes = NodeSelection;
+		}
+
+		if (bPrimarySelectionChanged)
+		{
+			EditorContent->SetPrimarySelectedNode(nullptr);
+
+			UDataflowEditorMode* const DataflowMode = CastChecked<UDataflowEditorMode>(EditorModeManager->GetActiveScriptableMode(UDataflowEditorMode::EM_DataflowEditorModeId));
+			if (DataflowMode)
 			{
-				for (IDataflowViewListener* Listener : ViewListeners)
+				// Close any running tool. OnNodeSingleClicked() will start a new tool if a new node was clicked.
+				UEditorInteractiveToolsContext* const ToolsContext = DataflowMode->GetInteractiveToolsContext();
+				checkf(ToolsContext, TEXT("No valid ToolsContext found for FDataflowEditorToolkit"));
+				if (ToolsContext->HasActiveTool())
 				{
-					Listener->OnSelectedNodeChanged(Node);
+					ToolsContext->EndTool(EToolShutdownType::Completed);
+				}
+
+				// Update the Construction viewport with the newly selected node's Collection
+				// @todo(brice) : Is this necessary? FDataflowPreviewScene::Update will generate a new RenderCollection 
+				TSharedPtr<FManagedArrayCollection> Collection = EvaluateCollection(PrimarySelection, EditorContent->GetDataflowContext());
+				if(Collection->HasGroup("Geometry"))
+				{
+					EditorContent->SetPrimarySelectedNode(PrimarySelection);
+					DataflowMode->SetSelectedCollection(Collection);
 				}
 			}
 		}
-		else
+	}
+}
+
+void FDataflowEditorToolkit::OnNodeSingleClicked(UObject* ClickedNode) const
+{
+	UDataflowEditorMode* DataflowMode = Cast<UDataflowEditorMode>(EditorModeManager->GetActiveScriptableMode(UDataflowEditorMode::EM_DataflowEditorModeId));
+	if (DataflowMode)
+	{
+		if (GraphEditor && GraphEditor->GetSingleSelectedNode() == ClickedNode)
 		{
-			for (IDataflowViewListener* Listener : ViewListeners)
-			{
-				Listener->OnSelectedNodeChanged(nullptr);
-			}
+			// Start the corresponding tool
+			DataflowMode->StartToolForSelectedNode(ClickedNode);
 		}
-
-		PrevNodeSelection = ValidatedSelection;
 	}
 }
 
+
 void FDataflowEditorToolkit::OnNodeDeleted(const TSet<UObject*>& NewSelection)
 {
 	for (UObject* Node : NewSelection)
 	{
-		if (PrevNodeSelection.Contains(Node))
+		if (SelectedDataflowNodes.Contains(Node))
 		{
-			PrevNodeSelection.Remove(Node);
+			SelectedDataflowNodes.Remove(Node);
 		}
 	}
 }
@@ -513,6 +644,7 @@ TSharedRef<SDataflowGraphEditor> FDataflowEditorToolkit::CreateGraphEditorWidget
 	SGraphEditor::FGraphEditorEvents InEvents;
 	InEvents.OnVerifyTextCommit = FOnNodeVerifyTextCommit::CreateSP(this, &FDataflowEditorToolkit::OnNodeVerifyTitleCommit);
 	InEvents.OnTextCommitted = FOnNodeTextCommitted::CreateSP(this, &FDataflowEditorToolkit::OnNodeTitleCommitted);
+	InEvents.OnNodeSingleClicked = SGraphEditor::FOnNodeSingleClicked::CreateSP(this, &FDataflowEditorToolkit::OnNodeSingleClicked);
 
 	TSharedRef<SDataflowGraphEditor> NewGraphEditor = SNew(SDataflowGraphEditor, DataflowToEdit)
 		.GraphToEdit(DataflowToEdit)
@@ -522,6 +654,7 @@ TSharedRef<SDataflowGraphEditor> FDataflowEditorToolkit::CreateGraphEditorWidget
 
 	OnSelectionChangedMulticastDelegateHandle = NewGraphEditor->OnSelectionChangedMulticast.AddSP(this, &FDataflowEditorToolkit::OnNodeSelectionChanged);
 	OnNodeDeletedMulticastDelegateHandle = NewGraphEditor->OnNodeDeletedMulticast.AddSP(this, &FDataflowEditorToolkit::OnNodeDeleted);
+	GetDataflowEditorContent()->SetDataflowGraphEditor(NewGraphEditor);
 
 	return NewGraphEditor;
 }
diff --git a/Engine/Plugins/Experimental/Dataflow/Source/DataflowEditor/Private/Dataflow/DataflowEditorUtil.cpp b/Engine/Plugins/Experimental/Dataflow/Source/DataflowEditor/Private/Dataflow/DataflowEditorUtil.cpp
--- a/Engine/Plugins/Experimental/Dataflow/Source/DataflowEditor/Private/Dataflow/DataflowEditorUtil.cpp
+++ b/Engine/Plugins/Experimental/Dataflow/Source/DataflowEditor/Private/Dataflow/DataflowEditorUtil.cpp
@@ -92,88 +92,6 @@ namespace UE
 		}
 	}
 
-
-	namespace Conversion
-	{
-		// Convert a rendering facade to a dynamic mesh
-		void RenderingFacadeToDynamicMesh(const GeometryCollection::Facades::FRenderingFacade& Facade, FDynamicMesh3& DynamicMesh)
-		{
-			if (Facade.CanRenderSurface())
-			{
-				const int32 NumTriangles = Facade.NumTriangles();
-				const int32 NumVertices = Facade.NumVertices();
-
-				const TManagedArray<FIntVector>& Indices = Facade.GetIndices();
-				const TManagedArray<FVector3f>& Positions = Facade.GetVertices();
-				const TManagedArray<FVector3f>& Normals = Facade.GetNormals();
-				const TManagedArray<FLinearColor>& Colors = Facade.GetVertexColor();
-
-				DynamicMesh.Clear();
-				for (int32 VertexIndex = 0; VertexIndex < NumVertices; ++VertexIndex)
-				{
-					DynamicMesh.AppendVertex(FVertexInfo(FVector3d(Positions[VertexIndex]), Normals[VertexIndex],
-											 FVector3f(Colors[VertexIndex].R, Colors[VertexIndex].G, Colors[VertexIndex].B)));
-				}
-				for (int32 TriangleIndex = 0; TriangleIndex < NumTriangles; ++TriangleIndex)
-				{
-					DynamicMesh.AppendTriangle(FIndex3i(Indices[TriangleIndex].X, Indices[TriangleIndex].Y, Indices[TriangleIndex].Z));
-				}
-
-				FMeshNormals::QuickComputeVertexNormals(DynamicMesh);
-			}
-		}
-
-		// Convert a dataflow component to a dynamic mesh
-		void DataflowToDynamicMesh(TSharedPtr<::Dataflow::FEngineContext> DataflowContext, UObject* Asset, UDataflow* Dataflow, FDynamicMesh3& DynamicMesh)
-		{
-			// just call update on the preview scene.
-
-			FManagedArrayCollection RenderCollection;
-			GeometryCollection::Facades::FRenderingFacade Facade(RenderCollection);
-			Facade.DefineSchema();
-
-
-			for (const UDataflowEdNode* Target : Dataflow->GetRenderTargets())
-			{
-				if (Target)
-				{
-					Target->Render(Facade, DataflowContext);
-				}
-			}
-			RenderingFacadeToDynamicMesh(Facade, DynamicMesh);
-		}
-
-		// Convert a dynamic mesh to a rendering facade
-		void DynamicMeshToRenderingFacade(const FDynamicMesh3& DynamicMesh, GeometryCollection::Facades::FRenderingFacade& Facade)
-		{
-			if (Facade.CanRenderSurface())
-			{
-				const int32 NumTriangles = Facade.NumTriangles();
-				const int32 NumVertices = Facade.NumVertices();
-
-				// We can only override vertices attributes (position, normals, colors)
-				if ((NumTriangles == DynamicMesh.TriangleCount()) && (NumVertices == DynamicMesh.VertexCount()))
-				{
-					TManagedArray<FVector3f>& Positions = Facade.ModifyVertices();
-					TManagedArray<FVector3f>& Normals = Facade.ModifyNormals();
-					TManagedArray<FLinearColor>& Colors = Facade.ModifyVertexColor();
-
-					for (int32 VertexIndex = 0; VertexIndex < NumVertices; ++VertexIndex)
-					{
-						Positions[VertexIndex] = FVector3f(DynamicMesh.GetVertex(VertexIndex));
-						Normals[VertexIndex] = DynamicMesh.GetVertexNormal(VertexIndex);
-						Colors[VertexIndex] = DynamicMesh.GetVertexColor(VertexIndex);
-					}
-				}
-			}
-		}
-
-		// Convert a dynamic mesh to a dataflow component
-		void DynamicMeshToDataflow(const FDynamicMesh3& DynamicMesh, UDataflow* Dataflow)
-		{
-			//todo
-		}
-	}
 }
 
 namespace Dataflow
diff --git a/Engine/Plugins/Experimental/Dataflow/Source/DataflowEditor/Private/Dataflow/DataflowEditorUtil.h b/Engine/Plugins/Experimental/Dataflow/Source/DataflowEditor/Private/Dataflow/DataflowEditorUtil.h
--- a/Engine/Plugins/Experimental/Dataflow/Source/DataflowEditor/Private/Dataflow/DataflowEditorUtil.h
+++ b/Engine/Plugins/Experimental/Dataflow/Source/DataflowEditor/Private/Dataflow/DataflowEditorUtil.h
@@ -36,21 +36,6 @@ namespace UE
 	{
 		UMaterial* LoadMaterialFromPath( const FName& Path, UObject* Outer);
 	}
-
-	namespace Conversion
-	{
-		// Convert a rendering facade to a dynamic mesh
-		void RenderingFacadeToDynamicMesh(const GeometryCollection::Facades::FRenderingFacade& Facade, UE::Geometry::FDynamicMesh3& DynamicMesh);
-
-		// Convert a dataflow component to a dynamic mesh
-		void DataflowToDynamicMesh(TSharedPtr<::Dataflow::FEngineContext> DataflowContext, UObject* Asset, UDataflow* Dataflow, UE::Geometry::FDynamicMesh3& DynamicMesh);
-
-		// Convert a dynamic mesh to a rendering facade
-		void DynamicMeshToRenderingFacade(const UE::Geometry::FDynamicMesh3& DynamicMesh, GeometryCollection::Facades::FRenderingFacade& Facade);
-
-		// Convert a dynamic mesh to a dataflow component
-		void DynamicMeshToDataflow(const UE::Geometry::FDynamicMesh3& DynamicMesh, UDataflow* Dataflow);
-	}
 }
 
 namespace Dataflow
diff --git a/Engine/Plugins/Experimental/Dataflow/Source/DataflowEditor/Private/Dataflow/DataflowEditorViewportClient.cpp b/Engine/Plugins/Experimental/Dataflow/Source/DataflowEditor/Private/Dataflow/DataflowEditorViewportClient.cpp
--- a/Engine/Plugins/Experimental/Dataflow/Source/DataflowEditor/Private/Dataflow/DataflowEditorViewportClient.cpp
+++ b/Engine/Plugins/Experimental/Dataflow/Source/DataflowEditor/Private/Dataflow/DataflowEditorViewportClient.cpp
@@ -10,9 +10,11 @@
 #include "Dataflow/DataflowPreviewScene.h"
 #include "Dataflow/DataflowXml.h"
 #include "EditorModeManager.h"
+#include "EdModeInteractiveToolsContext.h"
 #include "HAL/PlatformApplicationMisc.h"
 #include "PreviewScene.h"
 #include "Selection.h"
+#include "InputBehaviorSet.h"
 
 FDataflowEditorViewportClient::FDataflowEditorViewportClient(FEditorModeTools* InModeTools,
                                                              FPreviewScene* InPreviewScene,
@@ -33,6 +35,15 @@ void FDataflowEditorViewportClient::SetDataflowEditorToolkit(TWeakPtr<FDataflowE
 	DataflowEditorToolkitPtr = InDataflowEditorToolkitPtr;
 }
 
+void FDataflowEditorViewportClient::SetToolCommandList(TWeakPtr<FUICommandList> InToolCommandList)
+{
+	ToolCommandList = InToolCommandList;
+}
+
+//const UInputBehaviorSet* FDataflowEditorViewportClient::GetInputBehaviors() const
+//{
+//	return BehaviorSet;
+//}
 
 Dataflow::FTimestamp LatestTimestamp(const UDataflow* Dataflow, const ::Dataflow::FContext* Context)
 {
@@ -73,6 +84,48 @@ void FDataflowEditorViewportClient::Tick(float DeltaSeconds)
 	}
 }
 
+void FDataflowEditorViewportClient::SetConstructionViewMode(Dataflow::EDataflowPatternVertexType InViewMode)
+{
+	// @todo(Dataflow) : Add support for Sim2D
+	//const bool bSwitching2D3D = (ConstructionViewMode == Dataflow::EDataflowPatternVertexType::Sim2D) != (InViewMode == Dataflow::EDataflowPatternVertexType::Sim2D);
+	//if (bSwitching2D3D)
+	//{
+	//	Swap(SavedInactiveViewTransform, ViewTransformPerspective);
+	//}
+
+	ConstructionViewMode = Dataflow::EDataflowPatternVertexType::Sim3D;
+
+	
+	//if (ConstructionViewMode == EDataflowPatternVertexType::Sim2D)
+	//{
+	//	for (UInputBehavior* const Behavior : BehaviorsFor2DMode)
+	//	{
+	//		BehaviorSet->Add(Behavior);
+	//	}
+	//
+	//	const double AbsZ = FMath::Abs(ViewTransformPerspective.GetLocation().Z);
+	//	constexpr double CameraFarPlaneWorldZ = -10.0;
+	//	constexpr double CameraNearPlaneProportionZ = 0.8;
+	//	OverrideFarClipPlane(static_cast<float>(AbsZ - CameraFarPlaneWorldZ));
+	//	OverrideNearClipPlane(static_cast<float>(AbsZ * (1.0 - CameraNearPlaneProportionZ)));
+	//}
+	//else
+	//{
+	OverrideFarClipPlane(0);
+	OverrideNearClipPlane(UE_KINDA_SMALL_NUMBER);
+	//}
+
+	//ModeTools->GetInteractiveToolsContext()->InputRouter->DeregisterSource(this);
+	//ModeTools->GetInteractiveToolsContext()->InputRouter->RegisterSource(this);
+
+}
+
+
+Dataflow::EDataflowPatternVertexType FDataflowEditorViewportClient::GetConstructionViewMode() const
+{
+	return Dataflow::EDataflowPatternVertexType::Sim3D;// ConstructionViewMode;
+}
+
 void FDataflowEditorViewportClient::AddReferencedObjects(FReferenceCollector& Collector)
 {
 	Super::AddReferencedObjects(Collector);
diff --git a/Engine/Plugins/Experimental/Dataflow/Source/DataflowEditor/Private/Dataflow/DataflowPreviewScene.cpp b/Engine/Plugins/Experimental/Dataflow/Source/DataflowEditor/Private/Dataflow/DataflowPreviewScene.cpp
--- a/Engine/Plugins/Experimental/Dataflow/Source/DataflowEditor/Private/Dataflow/DataflowPreviewScene.cpp
+++ b/Engine/Plugins/Experimental/Dataflow/Source/DataflowEditor/Private/Dataflow/DataflowPreviewScene.cpp
@@ -6,14 +6,18 @@
 #include "AssetEditorModeManager.h"
 #include "Components/SkeletalMeshComponent.h"
 #include "Components/DynamicMeshComponent.h"
+#include "Dataflow/CollectionRenderingPatternUtility.h"
 #include "Dataflow/DataflowActor.h"
 #include "Dataflow/DataflowComponent.h"
 #include "Dataflow/DataflowEditor.h"
 #include "Dataflow/DataflowEditorContent.h"
 #include "Dataflow/DataflowEditorStyle.h"
 #include "Dataflow/DataflowEditorUtil.h"
+#include "Drawing/MeshElementsVisualizer.h"
 #include "Elements/Framework/EngineElementsLibrary.h"
+#include "InteractiveTool.h"
 #include "ModelingToolTargetUtil.h"
+#include "Selection.h"
 
 #define LOCTEXT_NAMESPACE "FDataflowPreviewScene"
 
@@ -51,63 +55,106 @@ FDataflowPreviewScene::~FDataflowPreviewScene()
 	ResetDynamicMeshComponents();
 }
 
-void FDataflowPreviewScene::AddReferencedObjects(FReferenceCollector& Collector)
+void FDataflowPreviewScene::Tick(float DeltaTime)
 {
-	FAdvancedPreviewScene::AddReferencedObjects(Collector);
-	Collector.AddReferencedObject(SkeletalMeshComponent);
-	Collector.AddReferencedObject(SkeletalMeshActor);
-	Collector.AddReferencedObject(EditorContent);
-	Collector.AddReferencedObject(DynamicMeshActor);
-	Collector.AddReferencedObjects(DynamicMeshComponents);
+	//@todo(brice) : Make sure this is being called. 
+	for (TObjectPtr<UInteractiveToolPropertySet>& Propset : PropertyObjectsToTick)
+	{
+		if (Propset)
+		{
+			if (Propset->IsPropertySetEnabled())
+			{
+				Propset->CheckAndUpdateWatched();
+			}
+			else
+			{
+				Propset->SilentUpdateWatched();
+			}
+		}
+	}
+
+	if (WireframeDraw)
+	{
+		WireframeDraw->OnTick(DeltaTime);
+	}
+}
+
+void FDataflowPreviewScene::ReinitializeDynamicMeshComponents()
+{
+	Update();
 }
 
 void FDataflowPreviewScene::Update()
 {
-	using namespace UE::Geometry;//FDynamicMesh3
+	// Some objects, like the UMeshElementsVisualizer and Settings Objects
+	// are not part of a tool, so they won't get ticked.This member holds
+	// ticked objects that get rebuilt on Update
+	PropertyObjectsToTick.Empty();
 
 	// Update the SkeletalMeshComponent for animation 
 	// changes.
 	UpdateSkeletalMeshComponent();
 
-
 	// The preview scene for the construction view will be
 	// cleared and rebuilt from scratch. This will genrate a 
 	// list of UPrimitiveComponents for rendering.
-	ResetDynamicMeshComponents();
+	UpdateDynamicMeshComponents();
 
-	if (EditorContent)
+	// Attach a wireframe renderer to the DynamicMeshComponents
+	UpdateWireframeMeshElementsVisualizer();
+
+
+	// Manage Selection and Tool Interaction
+	if (DataflowModeManager.IsValid())
 	{
-		TObjectPtr<UDataflow> DataflowAsset = EditorContent->GetDataflowAsset();
-		TSharedPtr<Dataflow::FEngineContext> DataflowContext = EditorContent->GetDataflowContext();
-		if(DataflowAsset && DataflowContext)
+		USelection* SelectedComponents = DataflowModeManager->GetSelectedComponents();
+		for (const TObjectPtr<UDynamicMeshComponent>& DynamicMeshComponent : DynamicMeshComponents)
 		{
-			for (const UDataflowEdNode* Target : DataflowAsset->GetRenderTargets())
-			{
-				if (Target)
-				{
-					FDynamicMesh3 DynamicMesh;
-					FManagedArrayCollection RenderCollection;
-					GeometryCollection::Facades::FRenderingFacade Facade(RenderCollection);
-					Facade.DefineSchema();
-
-					Target->Render(Facade, DataflowContext);
-					UE::Conversion::RenderingFacadeToDynamicMesh(Facade, DynamicMesh);
-					AddDynamicMeshComponent(MoveTemp(DynamicMesh), {});
-				}
-			}
+			SelectedComponents->DeselectAll();
+			SelectedComponents->Select(DynamicMeshComponent);
+			DynamicMeshComponent->PushSelectionToProxy();
 		}
 
-		EditorContent->SetIsDirty(false);
+		// @todo(brice) : Deal with this
+		// Update the context object with the ConstructionViewMode and Collection used to build the DynamicMeshComponents, so
+		// tools know how to use the components.
+		//UEditorInteractiveToolsContext* RestSpaceToolsContext = DataflowModeManager->GetInteractiveToolsContext();
+		//UClothEditorContextObject* EditorContextObject = RestSpaceToolsContext->ContextObjectStore->FindContext<UClothEditorContextObject>();
+		//if (ensure(EditorContextObject))
+		//{
+		//	EditorContextObject->SetCollection(ConstructionViewMode, Collection);
+		//}
+	}
+}
+
+void FDataflowPreviewScene::Exit()
+{
+	PropertyObjectsToTick.Empty();
+
+	if (WireframeDraw)
+	{
+		WireframeDraw->Disconnect();
 	}
+	WireframeDraw = nullptr;
+
+	ResetDynamicMeshComponents();
 }
 
+
 void FDataflowPreviewScene::ResetDynamicMeshComponents()
 {
+	USelection* SelectedComponents = DataflowModeManager->GetSelectedComponents();
 	for(const TObjectPtr<UDynamicMeshComponent>& DynamicMeshComponent : DynamicMeshComponents)
 	{
 		DynamicMeshComponent->SelectionOverrideDelegate.Unbind();
 		DynamicMeshComponent->UnregisterComponent();
 		DynamicMeshComponent->DestroyComponent();
+
+		if (SelectedComponents->IsSelected(DynamicMeshComponent))
+		{
+			SelectedComponents->Deselect(DynamicMeshComponent);
+			DynamicMeshComponent->PushSelectionToProxy();
+		}
 	}
 	DynamicMeshComponents.Reset();
 }
@@ -118,7 +165,8 @@ TObjectPtr<UDynamicMeshComponent>& FDataflowPreviewScene::AddDynamicMeshComponen
 		
 	DynamicMeshComponent->SetMesh(MoveTemp(DynamicMesh));
 
-	// @todo(Material) This is just to have a material, we should transfer the materials from the assets if they have them. 
+	// @todo(Dataflow) : Material support
+	// This is just to have a material, we should transfer the materials from the assets if they have them. 
 	if (FDataflowEditorStyle::Get().DefaultMaterial)
 	{
 		DynamicMeshComponent->ConfigureMaterialSet({ FDataflowEditorStyle::Get().DefaultMaterial });
@@ -127,15 +175,139 @@ TObjectPtr<UDynamicMeshComponent>& FDataflowPreviewScene::AddDynamicMeshComponen
 	{
 		DynamicMeshComponent->ValidateMaterialSlots(true, false);
 	}
+
 	DynamicMeshComponent->SelectionOverrideDelegate = UPrimitiveComponent::FSelectionOverride::CreateRaw(this, &FDataflowPreviewScene::IsComponentSelected);
-	
 	DynamicMeshComponent->RegisterComponentWithWorld(GetWorld());
+	DynamicMeshComponent->OnMeshChanged.Add(FSimpleMulticastDelegate::FDelegate::CreateLambda([this](){}));
+
 	DynamicMeshComponent->UpdateBounds();
 		
 	const int32 ElementIndex = DynamicMeshComponents.Emplace(DynamicMeshComponent);
 	return DynamicMeshComponents[ElementIndex];
 }
 
+void FDataflowPreviewScene::UpdateDynamicMeshComponents()
+{
+	using namespace UE::Geometry;//FDynamicMesh3
+
+	ResetDynamicMeshComponents();
+
+	if (EditorContent)
+	{
+		TObjectPtr<UDataflow> DataflowAsset = EditorContent->GetDataflowAsset();
+		TSharedPtr<Dataflow::FEngineContext> DataflowContext = EditorContent->GetDataflowContext();
+		if (DataflowAsset && DataflowContext)
+		{
+			for (const UDataflowEdNode* Target : DataflowAsset->GetRenderTargets())
+			{
+				if (Target)
+				{
+					FDynamicMesh3 DynamicMesh;
+					DynamicMesh.EnableAttributes();
+
+					TSharedPtr<FManagedArrayCollection> RenderCollection(new FManagedArrayCollection);
+					GeometryCollection::Facades::FRenderingFacade Facade(*RenderCollection);
+					Facade.DefineSchema();
+
+					Target->Render(Facade, DataflowContext);
+					Dataflow::Conversion::RenderingFacadeToDynamicMesh(Facade, DynamicMesh);
+
+					if (Target == EditorContent->GetPrimarySelectedNode())
+					{
+						EditorContent->SetPrimaryRenderCollection(RenderCollection);
+					}
+
+					// post updates
+					{
+						// Use per-triangle normals for the 2D view
+						//UE::Geometry::FMeshNormals::InitializeMeshToPerTriangleNormals(&LodMesh);
+					}
+					{
+						//@todo(Dataflow) :: Add material support
+						//SetUpDynamicMeshComponentMaterial(ClothFacade, *DynamicMeshComponent);
+					}
+
+
+					AddDynamicMeshComponent(MoveTemp(DynamicMesh), {});
+				}
+			}
+		}
+
+		EditorContent->SetIsDirty(false);
+	}
+}
+
+
+void FDataflowPreviewScene::AddWireframeMeshElementsVisualizer()
+{
+	ensure(WireframeDraw==nullptr);
+	if (DynamicMeshComponents.Num())
+	{
+		// Set up the wireframe display of the rest space mesh.
+
+		WireframeDraw = NewObject<UMeshElementsVisualizer>(DynamicMeshActor);
+		WireframeDraw->CreateInWorld(GetWorld(), FTransform::Identity);
+
+		WireframeDraw->Settings->DepthBias = 2.0;
+		WireframeDraw->Settings->bAdjustDepthBiasUsingMeshSize = false;
+		WireframeDraw->Settings->bShowWireframe = true;
+		WireframeDraw->Settings->bShowBorders = true;
+		WireframeDraw->Settings->bShowUVSeams = false;
+
+		WireframeDraw->WireframeComponent->BoundaryEdgeThickness = 2;
+
+		WireframeDraw->SetMeshAccessFunction([this](UMeshElementsVisualizer::ProcessDynamicMeshFunc ProcessFunc) 
+		{
+			for (auto DynamicMeshComponent : DynamicMeshComponents) ProcessFunc(*DynamicMeshComponent->GetMesh());
+		});
+
+		for (auto DynamicMeshComponent : DynamicMeshComponents)
+		{
+			DynamicMeshComponent->OnMeshChanged.Add(FSimpleMulticastDelegate::FDelegate::CreateLambda([this]()
+			{
+				WireframeDraw->NotifyMeshChanged();
+			}));
+
+			const bool bRestSpaceMeshVisible = DynamicMeshComponent->GetVisibleFlag();
+			WireframeDraw->Settings->bVisible = bRestSpaceMeshVisible && bConstructionViewWireframe;
+		}
+
+		// Some interactive tools will hide the input DynamicMeshComponent and create their own temporary PreviewMesh for visualization. If this
+		// occurs, we should also hide the corresponding Wireframe and Seam drawing (and un-hide it when the tool finishes).
+		/*
+		* // @todo(brice) : Deal with this.
+		UActorComponent::MarkRenderStateDirtyEvent.AddWeakLambda(this, [this](UActorComponent& ActorComponent)
+		{
+			if (!DynamicMeshComponent)
+			{
+				return;
+			}
+		    const bool bRestSpaceMeshVisible = DynamicMeshComponent->GetVisibleFlag();
+			if (WireframeDraw)
+			{
+				WireframeDraw->Settings->bVisible = bRestSpaceMeshVisible && bConstructionViewWireframe;
+			}
+		});
+		*/
+		PropertyObjectsToTick.Add(WireframeDraw->Settings);
+	}
+}
+
+void FDataflowPreviewScene::ResetWireframeMeshElementsVisualizer()
+{
+	if (WireframeDraw)
+	{
+		WireframeDraw->Disconnect();
+	}
+	WireframeDraw = nullptr;
+}
+
+void FDataflowPreviewScene::UpdateWireframeMeshElementsVisualizer()
+{
+	ResetWireframeMeshElementsVisualizer();
+	AddWireframeMeshElementsVisualizer();
+}
+
 void FDataflowPreviewScene::UpdateSkeletalMeshComponent()
 {
 	if (SkeletalMeshComponent)
@@ -215,5 +387,47 @@ bool FDataflowPreviewScene::IsComponentSelected(const UPrimitiveComponent* InCom
 	return false;
 }
 
+FBox FDataflowPreviewScene::SelectedComponentBounds() const
+{
+	FBox Bounds(ForceInit);
+	if (DataflowModeManager.IsValid())
+	{
+		const USelection* const SelectedComponents = DataflowModeManager->GetSelectedComponents();
+		for (int32 i = 0; i < SelectedComponents->Num(); ++i)
+		{
+			const UObject* const SelectedObject = SelectedComponents->GetSelectedObject(i);
+			if (const UDynamicMeshComponent* const DynamicMeshComponent = Cast<UDynamicMeshComponent>(SelectedObject))
+			{
+				Bounds += DynamicMeshComponent->Bounds.GetBox();
+			}
+		}
+	}
+	return Bounds;
+}
+
+bool FDataflowPreviewScene::HasRenderableGeometry()
+{
+	for (auto& DynamicMeshComponent : DynamicMeshComponents)
+	{
+		if (DynamicMeshComponent->GetMesh()->TriangleCount() > 0)
+		{
+			return true;
+		}
+	}
+	return false;
+}
+
+
+void FDataflowPreviewScene::AddReferencedObjects(FReferenceCollector& Collector)
+{
+	FAdvancedPreviewScene::AddReferencedObjects(Collector);
+	Collector.AddReferencedObject(SkeletalMeshComponent);
+	Collector.AddReferencedObject(SkeletalMeshActor);
+	Collector.AddReferencedObject(EditorContent);
+	Collector.AddReferencedObject(DynamicMeshActor);
+	Collector.AddReferencedObjects(DynamicMeshComponents);
+	Collector.AddReferencedObject(WireframeDraw);
+}
+
 #undef LOCTEXT_NAMESPACE
 
diff --git a/Engine/Plugins/Experimental/Dataflow/Source/DataflowEditor/Private/Dataflow/DataflowToolTarget.cpp b/Engine/Plugins/Experimental/Dataflow/Source/DataflowEditor/Private/Dataflow/DataflowToolTarget.cpp
--- a/Engine/Plugins/Experimental/Dataflow/Source/DataflowEditor/Private/Dataflow/DataflowToolTarget.cpp
+++ b/Engine/Plugins/Experimental/Dataflow/Source/DataflowEditor/Private/Dataflow/DataflowToolTarget.cpp
@@ -3,6 +3,7 @@
 #include "Dataflow/DataflowToolTarget.h"
 
 #include "ConversionUtils/DynamicMeshViaMeshDescriptionUtil.h"
+#include "Dataflow/CollectionRenderingPatternUtility.h"
 #include "Dataflow/DataflowEditor.h"
 #include "Dataflow/DataflowEditorContent.h"
 #include "Dataflow/DataflowEditorToolkit.h"
@@ -104,7 +105,7 @@ FDynamicMesh3 UDataflowReadOnlyToolTarget::GetDynamicMesh()
 	FDynamicMesh3 DynamicMesh;
 	if(IsValid())
 	{
-		UE::Conversion::DataflowToDynamicMesh(Context, Asset, Dataflow, DynamicMesh);
+		Dataflow::Conversion::DataflowToDynamicMesh(Context, Asset, Dataflow, DynamicMesh);
 	}
 	return DynamicMesh;
 }
@@ -146,7 +147,7 @@ void UDataflowToolTarget::CommitDynamicMesh(const FDynamicMesh3& DynamicMesh, co
 {
 	if(IsValid())
 	{
-		UE::Conversion::DynamicMeshToDataflow(DynamicMesh, Dataflow);
+		Dataflow::Conversion::DynamicMeshToDataflow(DynamicMesh, Dataflow);
 	}
 }
 
diff --git a/Engine/Plugins/Experimental/Dataflow/Source/DataflowEditor/Private/DataflowEditorTools/DataflowEditorWeightMapPaintBrushOps.cpp b/Engine/Plugins/Experimental/Dataflow/Source/DataflowEditor/Private/DataflowEditorTools/DataflowEditorWeightMapPaintBrushOps.cpp
new file mode 100644
--- /dev/null
+++ b/Engine/Plugins/Experimental/Dataflow/Source/DataflowEditor/Private/DataflowEditorTools/DataflowEditorWeightMapPaintBrushOps.cpp
@@ -0,0 +1,148 @@
+// Copyright Epic Games, Inc. All Rights Reserved.
+
+#include "DataflowEditorTools/DataflowEditorWeightMapPaintBrushOps.h"
+#include "Async/ParallelFor.h"
+#include "DynamicMesh/DynamicMesh3.h"
+
+//
+// Paint Brush
+//
+
+void FDataflowWeightMapPaintBrushOp::ApplyStampByVertices(
+	const FDynamicMesh3* Mesh,
+	const FSculptBrushStamp& Stamp,
+	const TArray<int32>& Vertices,
+	TArray<double>& NewAttributesOut)
+{
+	UDataflowWeightMapPaintBrushOpProps* const Props = GetPropertySetAs<UDataflowWeightMapPaintBrushOpProps>();
+	const double TargetValue = (double)Props->GetAttribute();
+
+	const int32 NumVertices = Vertices.Num();
+	ParallelFor(NumVertices, [&](int32 k)
+	{
+		if (bApplyRadiusLimit)
+		{
+			const FVector3d& StampPos = Stamp.LocalFrame.Origin;
+			const int32 VertIdx = Vertices[k];
+			const FVector3d VertexPos = Mesh->GetVertex(VertIdx);
+			const double DistanceSquared = (VertexPos - StampPos).SquaredLength();
+			if (DistanceSquared >= Stamp.Radius * Stamp.Radius)
+			{
+				return;
+			}
+		}
+
+		const double ExistingValue = NewAttributesOut[k];
+		const double ValueDiff = TargetValue - ExistingValue;
+		NewAttributesOut[k] = FMath::Clamp(ExistingValue + Stamp.Power * ValueDiff, 0.0, 1.0);
+	});
+}
+
+
+//
+// Erase Brush
+//
+
+void FDataflowWeightMapEraseBrushOp::ApplyStampByVertices(
+	const FDynamicMesh3* Mesh,
+	const FSculptBrushStamp& Stamp,
+	const TArray<int32>& Vertices,
+	TArray<double>& NewAttributesOut)
+{
+	UDataflowWeightMapEraseBrushOpProps* Props = GetPropertySetAs<UDataflowWeightMapEraseBrushOpProps>();
+	const double EraseAttribute = (double)Props->GetAttribute();
+
+	// TODO: Add something here to get the old value so we can subtract (clamped) the AttributeValue from it.
+	// TODO: Handle the stamp's properties for fall off, etc..
+
+	check(NewAttributesOut.Num() == Vertices.Num());
+
+	for (int32 k = 0; k < Vertices.Num(); ++k)
+	{
+		if (bApplyRadiusLimit)
+		{
+			const FVector3d& StampPos = Stamp.LocalFrame.Origin;
+			const int32 VertIdx = Vertices[k];
+			const FVector3d VertexPos = Mesh->GetVertex(VertIdx);
+			const double DistanceSquared = (VertexPos - StampPos).SquaredLength();
+			if (DistanceSquared >= Stamp.Radius * Stamp.Radius)
+			{
+				continue;
+			}
+		}
+		NewAttributesOut[k] = EraseAttribute;
+	}
+}
+
+
+//
+// Smooth Brush
+//
+
+void FDataflowWeightMapSmoothBrushOp::ApplyStampByVertices(
+	const FDynamicMesh3* Mesh,
+	const FSculptBrushStamp& Stamp,
+	const TArray<int32>& Vertices,
+	TArray<double>& VertexWeightValues)
+{
+	// Converted from FClothPaintTool_Smooth::SmoothVertices
+	
+	const int32 NumVertices = Vertices.Num();
+
+	TMap<int32, int32> VertexToBufferIndexMap;
+	for (int32 BufferIndex = 0; BufferIndex < NumVertices; ++BufferIndex)
+	{
+		const int32 VertexIndex = Vertices[BufferIndex];
+		VertexToBufferIndexMap.Add(VertexIndex, BufferIndex);
+	}
+
+	// Compute average values of one-rings for all vertices
+	TArray<double> OneRingAverages;
+	OneRingAverages.SetNumUninitialized(NumVertices);
+
+	for (int32 BufferIndex = 0; BufferIndex < NumVertices; ++BufferIndex)
+	{
+		const int32 VertexIndex = Vertices[BufferIndex];
+		double Accumulator = 0.0f;
+		int NumNeighbors = 0;
+
+		for (const int32 NeighborIndex : Mesh->VtxVerticesItr(VertexIndex))
+		{
+			if (VertexToBufferIndexMap.Contains(NeighborIndex))
+			{
+				const int32 NeighborBufferIndex = VertexToBufferIndexMap[NeighborIndex];
+				Accumulator += VertexWeightValues[NeighborBufferIndex];
+				++NumNeighbors;
+			}
+		}
+
+		if (NumNeighbors > 0)
+		{
+			OneRingAverages[BufferIndex] = Accumulator / NumNeighbors;
+		}
+		else
+		{
+			// Don't change the vertex value if it has no neighbors
+			OneRingAverages[BufferIndex] = VertexWeightValues[BufferIndex];
+		}
+	}
+
+	// Blend vertex value with its average one-ring value
+	for (int32 BufferIndex = 0; BufferIndex < NumVertices; ++BufferIndex)
+	{
+		const double Diff = OneRingAverages[BufferIndex] - VertexWeightValues[BufferIndex];
+		if (bApplyRadiusLimit)
+		{
+			const FVector3d& StampPos = Stamp.LocalFrame.Origin;
+			const int32 VertexIndex = Vertices[BufferIndex];
+			const FVector3d VertexPos = Mesh->GetVertex(VertexIndex);
+			const double DistanceSquared = (VertexPos - StampPos).SquaredLength();
+			if (DistanceSquared >= Stamp.Radius * Stamp.Radius)
+			{
+				continue;
+			}
+		}
+		VertexWeightValues[BufferIndex] += Stamp.Power * Diff;
+	}
+}
+
diff --git a/Engine/Plugins/Experimental/Dataflow/Source/DataflowEditor/Private/DataflowEditorTools/DataflowEditorWeightMapPaintBrushOps.h b/Engine/Plugins/Experimental/Dataflow/Source/DataflowEditor/Private/DataflowEditorTools/DataflowEditorWeightMapPaintBrushOps.h
new file mode 100644
--- /dev/null
+++ b/Engine/Plugins/Experimental/Dataflow/Source/DataflowEditor/Private/DataflowEditorTools/DataflowEditorWeightMapPaintBrushOps.h
@@ -0,0 +1,155 @@
+// Copyright Epic Games, Inc. All Rights Reserved.
+
+#pragma once
+
+#include "Sculpting/MeshBrushOpBase.h"
+
+#include "DataflowEditorWeightMapPaintBrushOps.generated.h"
+
+
+
+
+class DATAFLOWEDITOR_API FDataflowMeshVertexWeightMapEditBrushOp : public FMeshSculptBrushOp
+{
+public:
+
+	// not supported for this kind of brush op
+	virtual void ApplyStamp(
+		const FDynamicMesh3* Mesh,
+		const FSculptBrushStamp& Stamp,
+		const TArray<int32>& Vertices,
+		TArray<FVector3d>& NewColorsOut) override
+	{
+		check(false);
+	}
+
+
+	virtual void ApplyStampByVertices(
+		const FDynamicMesh3* Mesh,
+		const FSculptBrushStamp& Stamp,
+		const TArray<int32>& Vertices,
+		TArray<double>& NewAttributesOut) = 0;
+
+	/** If bApplyRadiusLimit is enabled, only vertices within the stamp radius are considered */
+	bool bApplyRadiusLimit = true;
+
+};
+
+
+
+//
+// Erase Brush
+// 
+
+UCLASS()
+class DATAFLOWEDITOR_API UDataflowWeightMapEraseBrushOpProps : public UMeshSculptBrushOpProps
+{
+	GENERATED_BODY()
+public:
+	// TODO: This AttributeValue is not currenly connected to the UI (jira UE-198413)
+	/** Attribute to set as Erased value */
+	UPROPERTY(EditAnywhere, Category = EraseBrush, meta = (DisplayName = "Erase Attribute", UIMin = 0))
+	double AttributeValue = 0;
+
+	virtual double GetAttribute() { return AttributeValue; }
+
+	virtual float GetFalloff() override { return 0.0f; }
+};
+
+
+
+class FDataflowWeightMapEraseBrushOp : public FDataflowMeshVertexWeightMapEditBrushOp
+{
+public:
+
+	virtual void ApplyStampByVertices(
+		const FDynamicMesh3* Mesh,
+		const FSculptBrushStamp& Stamp,
+		const TArray<int32>& Vertices,
+		TArray<double>& NewAttributesOut) override;
+
+};
+
+
+//
+// Paint Brush
+// 
+
+UCLASS()
+class DATAFLOWEDITOR_API UDataflowWeightMapPaintBrushOpProps : public UMeshSculptBrushOpProps
+{
+	GENERATED_BODY()
+
+public:
+
+	/** The Attribute that will be assigned to triangles within the brush region */
+	UPROPERTY(EditAnywhere, Category = PaintBrush, meta = (DisplayName = "Value", UIMin = 0))
+	double AttributeValue = 1;
+
+	/** Strength of the Brush */
+	UPROPERTY(EditAnywhere, Category = PaintBrush, meta = (DisplayName = "Strength", UIMin = "0.0", UIMax = "10.", ClampMin = "0.0", ClampMax = "10."))
+	float Strength = 0.5;
+
+	virtual double GetAttribute() { return AttributeValue; }
+
+	virtual float GetStrength() override { return Strength; }
+	virtual void SetStrength(float NewStrength) override { Strength = NewStrength;  }
+
+	virtual float GetFalloff() override { return 0.0f; }
+};
+
+
+class FDataflowWeightMapPaintBrushOp : public FDataflowMeshVertexWeightMapEditBrushOp
+{
+public:
+
+	virtual bool IgnoreZeroMovements() const override
+	{
+		return false;
+	}
+
+	virtual void ApplyStampByVertices(
+		const FDynamicMesh3* Mesh,
+		const FSculptBrushStamp& Stamp,
+		const TArray<int32>& Vertices,
+		TArray<double>& NewAttributesOut) override;
+
+};
+
+//
+// Smooth Brush
+// 
+
+UCLASS()
+class DATAFLOWEDITOR_API UDataflowWeightMapSmoothBrushOpProps : public UMeshSculptBrushOpProps
+{
+	GENERATED_BODY()
+
+public:
+
+	UPROPERTY(EditAnywhere, Category = PaintBrush, meta = (DisplayName = "Strength", UIMin = "0.0", UIMax = "10.", ClampMin = "0.0", ClampMax = "10."))
+	float Strength = 0.5;
+
+	virtual float GetStrength() override { return Strength; }
+	virtual void SetStrength(float NewStrength) override { Strength = NewStrength; }
+
+	virtual float GetFalloff() override { return 0.0f; }
+};
+
+
+class FDataflowWeightMapSmoothBrushOp : public FDataflowMeshVertexWeightMapEditBrushOp
+{
+public:
+
+	virtual bool IgnoreZeroMovements() const override
+	{
+		return false;
+	}
+
+	virtual void ApplyStampByVertices(
+		const FDynamicMesh3* Mesh,
+		const FSculptBrushStamp& Stamp,
+		const TArray<int32>& Vertices,
+		TArray<double>& NewAttributesOut) override;
+
+};
diff --git a/Engine/Plugins/Experimental/Dataflow/Source/DataflowEditor/Private/DataflowEditorTools/DataflowEditorWeightMapPaintTool.cpp b/Engine/Plugins/Experimental/Dataflow/Source/DataflowEditor/Private/DataflowEditorTools/DataflowEditorWeightMapPaintTool.cpp
new file mode 100644
--- /dev/null
+++ b/Engine/Plugins/Experimental/Dataflow/Source/DataflowEditor/Private/DataflowEditorTools/DataflowEditorWeightMapPaintTool.cpp
@@ -0,0 +1,1960 @@
+// Copyright Epic Games, Inc. All Rights Reserved.
+
+#include "DataflowEditorTools/DataflowEditorWeightMapPaintTool.h"
+
+#include "Async/ParallelFor.h"
+#include "Async/Async.h"
+#include "BaseGizmos/BrushStampIndicator.h"
+#include "CanvasTypes.h"
+#include "CanvasItem.h"
+#include "Changes/MeshVertexChange.h"
+#include "Changes/BasicChanges.h"
+#include "ContextObjectStore.h"
+#include "DataflowEditorTools/DataflowEditorWeightMapPaintBrushOps.h"
+#include "Dataflow/DataflowEditorContent.h"
+#include "Dataflow/DataflowEdNode.h"
+#include "Dataflow/DataflowObject.h"
+#include "Dataflow/DataflowCollectionAddScalarVertexPropertyNode.h"
+#include "Dataflow/DataflowPatternVertexType.h"
+#include "Dataflow/DataflowSNode.h"
+#include "Dataflow/DataflowGraphEditor.h"
+#include "Drawing/MeshElementsVisualizer.h"
+#include "DynamicMesh/MeshNormals.h"
+#include "DynamicMesh/MeshIndexUtil.h"
+#include "DynamicMesh/NonManifoldMappingSupport.h"
+#include "Engine/World.h"
+#include "Framework/Notifications/NotificationManager.h"
+#include "GraphEditor.h"
+#include "InteractiveToolManager.h"
+#include "InteractiveGizmoManager.h"
+#include "Intersection/IntrLine2Line2.h"
+#include "Intersection/IntrSegment2Segment2.h"
+#include "ModelingToolTargetUtil.h"
+#include "MeshWeights.h"
+#include "ToolSetupUtil.h"
+#include "Polygon2.h"
+#include "Sculpting/StampFalloffs.h"
+#include "Sculpting/MeshSculptUtil.h"
+#include "Selections/MeshConnectedComponents.h"
+#include "Selections/MeshFaceSelection.h"
+#include "Selections/MeshVertexSelection.h"
+#include "Selection/PolygonSelectionMechanic.h"
+#include "Util/BufferUtil.h"
+#include "Util/ColorConstants.h"
+#include "Widgets/Notifications/SNotificationList.h"
+
+#include UE_INLINE_GENERATED_CPP_BY_NAME(DataflowEditorWeightMapPaintTool)
+
+using namespace UE::Geometry;
+
+#define LOCTEXT_NAMESPACE "UDataflowEditorWeightMapPaintTool"
+
+namespace Dataflow::Private
+{
+	// probably should be something defined for the whole tool framework...
+#if WITH_EDITOR
+	static EAsyncExecution WeightPaintToolAsyncExecTarget = EAsyncExecution::LargeThreadPool;
+#else
+	static EAsyncExecution WeightPaintToolAsyncExecTarget = EAsyncExecution::ThreadPool;
+#endif
+}
+
+
+/*
+ * ToolBuilder
+ */
+
+void UDataflowEditorWeightMapPaintToolBuilder::GetSupportedViewModes(TArray<Dataflow::EDataflowPatternVertexType>& Modes) const
+{
+	Modes.Add(Dataflow::EDataflowPatternVertexType::Sim3D);
+	//Modes.Add(Dataflow::EDataflowPatternVertexType::Sim2D);
+}
+
+UMeshSurfacePointTool* UDataflowEditorWeightMapPaintToolBuilder::CreateNewTool(const FToolBuilderState& SceneState) const
+{
+	UDataflowEditorWeightMapPaintTool* PaintTool = NewObject<UDataflowEditorWeightMapPaintTool>(SceneState.ToolManager);
+	PaintTool->SetWorld(SceneState.World);
+
+	if (UDataflowEditorContextObject* ContextObject = SceneState.ToolManager->GetContextObjectStore()->FindContext<UDataflowEditorContextObject>())
+	{
+		PaintTool->SetDataflowEditorContextObject(ContextObject);
+	}
+
+	return PaintTool;
+}
+
+
+/*
+ * ToolActions
+ */
+void UDataflowEditorMeshWeightMapPaintToolActions::PostAction(EDataflowEditorWeightMapPaintToolActions Action)
+{
+	if (ParentTool.IsValid())
+	{
+		ParentTool->RequestAction(Action);
+	}
+}
+
+
+/*
+ * Properties
+ */
+void UDataflowEditorUpdateWeightMapProperties::PostEditChangeProperty(FPropertyChangedEvent& PropertyChangedEvent)
+{
+	// @todo(brice) Fix this
+	if (PropertyChangedEvent.GetPropertyName() == GET_MEMBER_NAME_CHECKED(UDataflowEditorUpdateWeightMapProperties, Name))
+	{
+		//Dataflow::FWeightMapTools::MakeWeightMapName(Name);
+	}
+}
+
+
+/*
+ * Tool
+ */
+
+void UDataflowEditorWeightMapPaintTool::Setup()
+{
+	UMeshSculptToolBase::Setup();
+
+	// Get the selected weight map node
+	WeightMapNodeToUpdate = DataflowEditorContextObject->GetPrimarySelectedNodeOfType<FDataflowCollectionAddScalarVertexPropertyNode>();
+	checkf(WeightMapNodeToUpdate, TEXT("No Weight Map Node is currently selected, or more than one node is selected"));
+
+	SetToolDisplayName(LOCTEXT("ToolName", "Paint Weight Maps"));
+
+	// create dynamic mesh component to use for live preview
+	FActorSpawnParameters SpawnInfo;
+	PreviewMeshActor = TargetWorld->SpawnActor<AInternalToolFrameworkActor>(FVector::ZeroVector, FRotator::ZeroRotator, SpawnInfo);
+	DynamicMeshComponent = NewObject<UDynamicMeshComponent>(PreviewMeshActor);
+	InitializeSculptMeshComponent(DynamicMeshComponent, PreviewMeshActor);
+
+	// assign materials
+	FComponentMaterialSet MaterialSet = UE::ToolTarget::GetMaterialSet(Target);
+	for (int k = 0; k < MaterialSet.Materials.Num(); ++k)
+	{
+		DynamicMeshComponent->SetMaterial(k, MaterialSet.Materials[k]);
+	}
+
+	DynamicMeshComponent->SetInvalidateProxyOnChangeEnabled(false);
+	OnDynamicMeshComponentChangedHandle = DynamicMeshComponent->OnMeshVerticesChanged.AddUObject(this, &UDataflowEditorWeightMapPaintTool::OnDynamicMeshComponentChanged);
+
+	FDynamicMesh3* Mesh = GetSculptMesh();
+	Mesh->EnableVertexColors(FVector3f::One());
+	Mesh->Attributes()->EnablePrimaryColors();
+	Mesh->Attributes()->PrimaryColors()->CreateFromPredicate([](int ParentVID, int TriIDA, int TriIDB) {return true; }, 0.f);
+	FAxisAlignedBox3d Bounds = Mesh->GetBounds(true);
+
+	TFuture<void> PrecomputeFuture = Async(Dataflow::Private::WeightPaintToolAsyncExecTarget, [&]()
+	{
+		PrecomputeFilterData();
+	});
+
+	TFuture<void> OctreeFuture = Async(Dataflow::Private::WeightPaintToolAsyncExecTarget, [&]()
+	{
+		// initialize dynamic octree
+		if (Mesh->TriangleCount() > 100000)
+		{
+			Octree.RootDimension = Bounds.MaxDim() / 10.0;
+			Octree.SetMaxTreeDepth(4);
+		}
+		else
+		{
+			Octree.RootDimension = Bounds.MaxDim();
+			Octree.SetMaxTreeDepth(8);
+		}
+		Octree.Initialize(Mesh);
+	});
+
+	// initialize render decomposition
+	TUniquePtr<FMeshRenderDecomposition> Decomp = MakeUnique<FMeshRenderDecomposition>();
+	FMeshRenderDecomposition::BuildChunkedDecomposition(Mesh, &MaterialSet, *Decomp);
+	Decomp->BuildAssociations(Mesh);
+	DynamicMeshComponent->SetExternalDecomposition(MoveTemp(Decomp));
+
+	// initialize brush radius range interval, brush properties
+	UMeshSculptToolBase::InitializeBrushSizeRange(Bounds);
+
+	// Set up control points mechanic
+	PolyLassoMechanic = NewObject<UPolyLassoMarqueeMechanic>(this);
+	PolyLassoMechanic->Setup(this);
+	PolyLassoMechanic->SetIsEnabled(false);
+	PolyLassoMechanic->SpacingTolerance = 10.0f;
+	PolyLassoMechanic->OnDrawPolyLassoFinished.AddUObject(this, &UDataflowEditorWeightMapPaintTool::OnPolyLassoFinished);
+
+
+	// Set up vertex selection mechanic
+	PolygonSelectionMechanic = NewObject<UPolygonSelectionMechanic>(this);
+	PolygonSelectionMechanic->bAddSelectionFilterPropertiesToParentTool = false;
+	PolygonSelectionMechanic->Setup(this);
+	PolygonSelectionMechanic->SetIsEnabled(false);
+	PolygonSelectionMechanic->OnSelectionChanged.AddUObject(this, &UDataflowEditorWeightMapPaintTool::OnSelectionModified);
+
+	// disable CTRL to remove from selection
+	PolygonSelectionMechanic->SetShouldRemoveFromSelectionFunc([]() { return false; });
+
+	PolygonSelectionMechanic->Properties->bSelectEdges = false;
+	PolygonSelectionMechanic->Properties->bSelectFaces = false;
+	PolygonSelectionMechanic->Properties->bSelectVertices = true;
+	
+	constexpr bool bAutoBuild = true;
+	GradientSelectionTopology = MakeUnique<UE::Geometry::FTriangleGroupTopology>(DynamicMeshComponent->GetMesh(), bAutoBuild);
+	
+	MeshSpatial = MakeUnique<UE::Geometry::FDynamicMeshAABBTree3>(DynamicMeshComponent->GetMesh(), bAutoBuild);
+	PolygonSelectionMechanic->Initialize(DynamicMeshComponent, GradientSelectionTopology.Get(), [this]() { return MeshSpatial.Get(); });
+
+	UpdateWeightMapProperties = NewObject<UDataflowEditorUpdateWeightMapProperties>(this);
+	UpdateWeightMapProperties->Name = WeightMapNodeToUpdate->Name;
+
+	UpdateWeightMapProperties->WatchProperty(WeightMapNodeToUpdate->Name, [this](const FString& NewName)
+	{
+		UpdateWeightMapProperties->Name = NewName;
+	});
+	AddToolPropertySource(UpdateWeightMapProperties);
+
+	// initialize other properties
+	FilterProperties = NewObject<UDataflowEditorWeightMapPaintBrushFilterProperties>(this);
+	FilterProperties->WatchProperty(FilterProperties->SubToolType,
+		[this](EDataflowEditorWeightMapPaintInteractionType NewType) { UpdateSubToolType(NewType); });
+	FilterProperties->WatchProperty(FilterProperties->BrushSize,
+		[this](float NewSize) { UMeshSculptToolBase::BrushProperties->BrushSize.AdaptiveSize = NewSize; });
+	FilterProperties->BrushSize = UMeshSculptToolBase::BrushProperties->BrushSize.AdaptiveSize;
+	FilterProperties->RestoreProperties(this);
+	AddToolPropertySource(FilterProperties);
+
+	InitializeIndicator();
+
+	// initialize our properties
+	AddToolPropertySource(UMeshSculptToolBase::BrushProperties);
+	UMeshSculptToolBase::BrushProperties->bShowPerBrushProps = false;
+	UMeshSculptToolBase::BrushProperties->bShowFalloff = true;
+	UMeshSculptToolBase::BrushProperties->bShowLazyness = false;
+	UMeshSculptToolBase::BrushProperties->FlowRate = 0.0f;
+	CalculateBrushRadius();
+
+	PaintBrushOpProperties = NewObject<UDataflowWeightMapPaintBrushOpProps>(this);
+	RegisterBrushType((int32)EDataflowEditorWeightMapPaintBrushType::Paint, LOCTEXT("Paint", "Paint"),
+		MakeUnique<FLambdaMeshSculptBrushOpFactory>([this]() { return MakeUnique<FDataflowWeightMapPaintBrushOp>(); }),
+		PaintBrushOpProperties);
+
+	SmoothBrushOpProperties = NewObject<UDataflowWeightMapSmoothBrushOpProps>(this);
+	RegisterBrushType((int32)EDataflowEditorWeightMapPaintBrushType::Smooth, LOCTEXT("SmoothBrushType", "Smooth"),
+		MakeUnique<FLambdaMeshSculptBrushOpFactory>([this]() { return MakeUnique<FDataflowWeightMapSmoothBrushOp>(); }),
+		SmoothBrushOpProperties);
+
+	// secondary brushes
+	EraseBrushOpProperties = NewObject<UDataflowWeightMapEraseBrushOpProps>(this);
+
+	RegisterSecondaryBrushType((int32)EDataflowEditorWeightMapPaintBrushType::Erase, LOCTEXT("Erase", "Erase"),
+		MakeUnique<TBasicMeshSculptBrushOpFactory<FDataflowWeightMapEraseBrushOp>>(),
+		EraseBrushOpProperties);
+
+	AddToolPropertySource(UMeshSculptToolBase::ViewProperties);
+	SetToolPropertySourceEnabled(UMeshSculptToolBase::ViewProperties, true);
+
+	AddToolPropertySource(UMeshSculptToolBase::GizmoProperties);
+	SetToolPropertySourceEnabled(UMeshSculptToolBase::GizmoProperties, false);
+
+
+	// register watchers
+	FilterProperties->WatchProperty( FilterProperties->PrimaryBrushType,
+		[this](EDataflowEditorWeightMapPaintBrushType NewType) { UpdateBrushType(NewType); });
+
+	// must call before updating brush type so that we register all brush properties?
+	UMeshSculptToolBase::OnCompleteSetup();
+
+	UpdateBrushType(FilterProperties->PrimaryBrushType);
+	SetActiveSecondaryBrushType((int32)EDataflowEditorWeightMapPaintBrushType::Erase);
+	
+	ActionsProps = NewObject<UDataflowEditorMeshWeightMapPaintToolActions>(this);
+	ActionsProps->Initialize(this);
+	AddToolPropertySource(ActionsProps);
+
+	MeshElementsDisplay = NewObject<UMeshElementsVisualizer>(this);
+	MeshElementsDisplay->CreateInWorld(DynamicMeshComponent->GetWorld(), DynamicMeshComponent->GetComponentTransform());
+	if (ensure(MeshElementsDisplay->Settings))
+	{
+		MeshElementsDisplay->Settings->bShowNormalSeams = false;
+		MeshElementsDisplay->Settings->RestoreProperties(this, TEXT("DataflowEditorWeightMapPaintTool2"));
+		AddToolPropertySource(MeshElementsDisplay->Settings);
+	}
+	MeshElementsDisplay->SetMeshAccessFunction([this](UMeshElementsVisualizer::ProcessDynamicMeshFunc ProcessFunc) {
+		ProcessFunc(*GetSculptMesh());
+	});
+
+	// disable view properties
+	SetViewPropertiesEnabled(false);
+	UpdateMaterialMode(EMeshEditingMaterialModes::VertexColor);
+	UpdateWireframeVisibility(false);
+	UpdateFlatShadingSetting(false);
+
+	// configure panels
+	UpdateSubToolType(FilterProperties->SubToolType);
+
+	// Setup DynamicMeshToWeight conversion
+	if (DataflowEditorContextObject)
+	{
+		if (const TSharedPtr<const FManagedArrayCollection> Collection = DataflowEditorContextObject->GetPrimaryRenderCollection())
+		{
+			using namespace Dataflow;
+			const FNonManifoldMappingSupport NonManifoldMapping(*Mesh);
+
+			bHaveDynamicMeshToWeightConversion = NonManifoldMapping.IsNonManifoldVertexInSource();
+			if (bHaveDynamicMeshToWeightConversion)
+			{
+				DynamicMeshToWeight.SetNumUninitialized(Mesh->VertexCount());
+				WeightToDynamicMesh.Reset();
+				WeightToDynamicMesh.SetNum(Collection->NumElements("Vertices"));
+				for (int32 DynamicMeshVert = 0; DynamicMeshVert < Mesh->VertexCount(); ++DynamicMeshVert)
+				{
+					DynamicMeshToWeight[DynamicMeshVert] = NonManifoldMapping.GetOriginalNonManifoldVertexID(DynamicMeshVert);
+					WeightToDynamicMesh[DynamicMeshToWeight[DynamicMeshVert]].Add(DynamicMeshVert);
+				}
+			}
+		}
+	}
+
+	PrecomputeFuture.Wait();
+	OctreeFuture.Wait();
+
+
+	// Create an attribute layer to temporarily paint into
+	const int NumAttributeLayers = Mesh->Attributes()->NumWeightLayers();
+	Mesh->Attributes()->SetNumWeightLayers(NumAttributeLayers + 1);
+	ActiveWeightMap = Mesh->Attributes()->GetWeightLayer(NumAttributeLayers);
+	ActiveWeightMap->SetName(FName("PaintLayer"));
+
+	// Copy weights from selected node to the preview mesh
+	const TArray<float>& CurrentWeights = WeightMapNodeToUpdate->VertexWeights;
+	
+	if (bHaveDynamicMeshToWeightConversion)
+	{
+		if (WeightToDynamicMesh.Num() == CurrentWeights.Num())	// Only copy node weights if they match the number of mesh vertices
+		{
+			for (int32 WeightID = 0; WeightID < CurrentWeights.Num(); ++WeightID)
+			{
+				for (const int32 VertexID : WeightToDynamicMesh[WeightID])
+				{
+					ActiveWeightMap->SetValue(VertexID, &CurrentWeights[WeightID]);
+				}
+			}
+		}
+	}
+	else
+	{
+		if (Mesh->MaxVertexID() == CurrentWeights.Num())	// Only copy node weights if they match the number of mesh vertices
+		{
+			for (int32 VertexID = 0; VertexID < CurrentWeights.Num(); ++VertexID)
+			{
+				ActiveWeightMap->SetValue(VertexID, &CurrentWeights[VertexID]);
+			}
+		}
+	}
+
+	UpdateWeightMapProperties->Name = WeightMapNodeToUpdate->Name;
+	SetToolPropertySourceEnabled(UpdateWeightMapProperties, true);
+
+	// update colors
+	UpdateVertexColorOverlay();
+	DynamicMeshComponent->FastNotifyVertexAttributesUpdated(EMeshRenderAttributeFlags::VertexColors);
+	GetToolManager()->PostInvalidation();
+
+}
+
+
+void UDataflowEditorWeightMapPaintTool::SetDataflowEditorContextObject(TObjectPtr<UDataflowEditorContextObject> InDataflowEditorContextObject)
+{
+	DataflowEditorContextObject = InDataflowEditorContextObject;
+}
+
+void UDataflowEditorWeightMapPaintTool::Shutdown(EToolShutdownType ShutdownType)
+{
+	if (DynamicMeshComponent != nullptr)
+	{
+		DynamicMeshComponent->OnMeshChanged.Remove(OnDynamicMeshComponentChangedHandle);
+	}
+
+	if (ensure(MeshElementsDisplay->Settings))
+	{
+		MeshElementsDisplay->Settings->SaveProperties(this, TEXT("DataflowEditorWeightMapPaintTool2"));
+	}
+	MeshElementsDisplay->Disconnect();
+
+	FilterProperties->SaveProperties(this);
+	if (PreviewMeshActor != nullptr)
+	{
+		PreviewMeshActor->Destroy();
+		PreviewMeshActor = nullptr;
+	}
+
+	if (PolygonSelectionMechanic)
+	{
+		PolygonSelectionMechanic->Shutdown();
+		PolygonSelectionMechanic = nullptr;
+	}
+
+	UMeshSculptToolBase::Shutdown(ShutdownType);
+}
+
+
+void UDataflowEditorWeightMapPaintTool::CommitResult(UBaseDynamicMeshComponent* Component, bool bModifiedTopology)
+{
+	GetToolManager()->BeginUndoTransaction(LOCTEXT("WeightPaintToolTransactionName", "Paint Weights"));
+
+	UpdateSelectedNode();
+
+	GetToolManager()->EndUndoTransaction();
+}
+
+
+void UDataflowEditorWeightMapPaintTool::RegisterActions(FInteractiveToolActionSet& ActionSet)
+{
+	UMeshSculptToolBase::RegisterActions(ActionSet);
+	
+	ActionSet.RegisterAction(this, (int32)EStandardToolActions::BaseClientDefinedActionID + 500,
+		TEXT("PickWeightValueUnderCursor"),
+		LOCTEXT("PickWeightValueUnderCursor", "Pick Weight Value"),
+		LOCTEXT("PickWeightValueUnderCursorTooltip", "Set the active weight painting value to that currently under the cursor"),
+		EModifierKey::Shift, EKeys::G,
+		[this]() { bPendingPickWeight = true; });
+};
+
+
+TUniquePtr<FMeshSculptBrushOp>& UDataflowEditorWeightMapPaintTool::GetActiveBrushOp()
+{
+	if (GetInEraseStroke())
+	{
+		return SecondaryBrushOp;
+	}
+	else
+	{
+		return PrimaryBrushOp;
+	}
+}
+
+
+void UDataflowEditorWeightMapPaintTool::OnPropertyModified(UObject* PropertySet, FProperty* Property)
+{
+	CalculateBrushRadius();
+}
+
+
+void UDataflowEditorWeightMapPaintTool::IncreaseBrushRadiusAction()
+{
+	Super::IncreaseBrushRadiusAction();
+	FilterProperties->BrushSize = BrushProperties->BrushSize.AdaptiveSize;
+	NotifyOfPropertyChangeByTool(FilterProperties);
+}
+
+void UDataflowEditorWeightMapPaintTool::DecreaseBrushRadiusAction()
+{
+	Super::DecreaseBrushRadiusAction();
+	FilterProperties->BrushSize = BrushProperties->BrushSize.AdaptiveSize;
+	NotifyOfPropertyChangeByTool(FilterProperties);
+}
+
+void UDataflowEditorWeightMapPaintTool::IncreaseBrushRadiusSmallStepAction()
+{
+	Super::IncreaseBrushRadiusSmallStepAction();
+	FilterProperties->BrushSize = BrushProperties->BrushSize.AdaptiveSize;
+	NotifyOfPropertyChangeByTool(FilterProperties);
+}
+
+void UDataflowEditorWeightMapPaintTool::DecreaseBrushRadiusSmallStepAction()
+{
+	Super::DecreaseBrushRadiusSmallStepAction();
+	FilterProperties->BrushSize = BrushProperties->BrushSize.AdaptiveSize;
+	NotifyOfPropertyChangeByTool(FilterProperties);
+}
+
+
+bool UDataflowEditorWeightMapPaintTool::IsInBrushSubMode() const
+{
+	return FilterProperties->SubToolType == EDataflowEditorWeightMapPaintInteractionType::Brush
+		|| FilterProperties->SubToolType == EDataflowEditorWeightMapPaintInteractionType::Fill;
+}
+
+
+void UDataflowEditorWeightMapPaintTool::OnBeginStroke(const FRay& WorldRay)
+{
+	if (!ActiveWeightMap)
+	{
+		return;
+	}
+
+	UpdateBrushPosition(WorldRay);
+
+	if (PaintBrushOpProperties)
+	{
+		PaintBrushOpProperties->AttributeValue = FilterProperties->AttributeValue;
+		PaintBrushOpProperties->Strength = FilterProperties->Strength * FilterProperties->Strength;
+	}
+	if (EraseBrushOpProperties)
+	{
+		EraseBrushOpProperties->AttributeValue = 0.0;
+	}
+	if (SmoothBrushOpProperties)
+	{
+		SmoothBrushOpProperties->Strength = FilterProperties->Strength * FilterProperties->Strength;
+	}
+
+	// initialize first "Last Stamp", so that we can assume all stamps in stroke have a valid previous stamp
+	LastStamp.WorldFrame = GetBrushFrameWorld();
+	LastStamp.LocalFrame = GetBrushFrameLocal();
+	LastStamp.Radius = GetCurrentBrushRadius();
+	LastStamp.Falloff = GetCurrentBrushFalloff();
+	LastStamp.Direction = GetInInvertStroke() ? -1.0 : 1.0;
+	LastStamp.Depth = GetCurrentBrushDepth();
+	LastStamp.Power = GetActivePressure() * GetCurrentBrushStrength();
+	LastStamp.TimeStamp = FDateTime::Now();
+
+	FSculptBrushOptions SculptOptions;
+	SculptOptions.ConstantReferencePlane = GetCurrentStrokeReferencePlane();
+
+	TUniquePtr<FMeshSculptBrushOp>& UseBrushOp = GetActiveBrushOp();
+	UseBrushOp->ConfigureOptions(SculptOptions);
+	UseBrushOp->BeginStroke(GetSculptMesh(), LastStamp, VertexROI);
+
+	AccumulatedTriangleROI.Reset();
+
+	// begin change here? or wait for first stamp?
+	BeginChange();
+}
+
+void UDataflowEditorWeightMapPaintTool::OnEndStroke()
+{
+	if (!ActiveWeightMap)
+	{
+		return;
+	}
+
+	GetActiveBrushOp()->EndStroke(GetSculptMesh(), LastStamp, VertexROI);
+
+	UpdateVertexColorOverlay(&TriangleROI);
+	DynamicMeshComponent->FastNotifyVertexAttributesUpdated(EMeshRenderAttributeFlags::VertexColors);
+
+	// close change record
+	EndChange();
+}
+
+
+
+
+void UDataflowEditorWeightMapPaintTool::UpdateROI(const FSculptBrushStamp& BrushStamp)
+{
+	SCOPE_CYCLE_COUNTER(WeightMapPaintTool_UpdateROI);
+
+	const FVector3d& BrushPos = BrushStamp.LocalFrame.Origin;
+	const FDynamicMesh3* Mesh = GetSculptMesh();
+	float RadiusSqr = GetCurrentBrushRadius() * GetCurrentBrushRadius();
+	FAxisAlignedBox3d BrushBox(
+		BrushPos - GetCurrentBrushRadius() * FVector3d::One(),
+		BrushPos + GetCurrentBrushRadius() * FVector3d::One());
+
+	TriangleROI.Reset();
+
+	int32 CenterTID = GetBrushTriangleID();
+	if (Mesh->IsTriangle(CenterTID))
+	{
+		TriangleROI.Add(CenterTID);
+	}
+
+	FVector3d CenterNormal = Mesh->IsTriangle(CenterTID) ? TriNormals[CenterTID] : FVector3d::One();		// One so that normal check always passes
+	
+	bool bUseAngleThreshold = FilterProperties->AngleThreshold < 180.0f;
+	double DotAngleThreshold = FMathd::Cos(FilterProperties->AngleThreshold * FMathd::DegToRad);
+	bool bStopAtUVSeams = FilterProperties->bUVSeams;
+	bool bStopAtNormalSeams = FilterProperties->bNormalSeams;
+
+	auto CheckEdgeCriteria = [&](int32 t1, int32 t2) -> bool
+	{
+		if (bUseAngleThreshold == false || CenterNormal.Dot(TriNormals[t2]) > DotAngleThreshold)
+		{
+			int32 eid = Mesh->FindEdgeFromTriPair(t1, t2);
+			if (bStopAtUVSeams == false || UVSeamEdges[eid] == false)
+			{
+				if (bStopAtNormalSeams == false || NormalSeamEdges[eid] == false)
+				{
+					return true;
+				}
+			}
+		}
+		return false;
+	};
+
+	bool bFill = (FilterProperties->SubToolType == EDataflowEditorWeightMapPaintInteractionType::Fill);
+
+	if (Mesh->IsTriangle(CenterTID))
+	{
+		TArray<int32> StartROI;
+		StartROI.Add(CenterTID);
+		FMeshConnectedComponents::GrowToConnectedTriangles(Mesh, StartROI, TriangleROI, &TempROIBuffer,
+			[&](int t1, int t2) 
+		{ 
+			if ((Mesh->GetTriCentroid(t2) - BrushPos).SquaredLength() < RadiusSqr)
+			{
+				return CheckEdgeCriteria(t1, t2);
+			}
+			return false;
+		});
+	}
+
+	if (bFill)
+	{
+		TArray<int32> StartROI;
+		for (int32 tid : TriangleROI)
+		{
+			StartROI.Add(tid);
+		}
+		FMeshConnectedComponents::GrowToConnectedTriangles(Mesh, StartROI, TriangleROI, &TempROIBuffer,
+														   [&](int t1, int t2)
+		{
+			return CheckEdgeCriteria(t1, t2);
+		});
+	}
+
+	// construct ROI vertex set
+	VertexSetBuffer.Reset();
+	for (int32 tid : TriangleROI)
+	{
+		FIndex3i Tri = Mesh->GetTriangle(tid);
+		VertexSetBuffer.Add(Tri.A);  VertexSetBuffer.Add(Tri.B);  VertexSetBuffer.Add(Tri.C);
+	}
+	
+	// apply visibility filter
+	if (FilterProperties->VisibilityFilter != EDataflowEditorWeightMapPaintVisibilityType::None)
+	{
+		TArray<int32> ResultBuffer;
+		ApplyVisibilityFilter(VertexSetBuffer, TempROIBuffer, ResultBuffer);
+	}
+
+	VertexROI.SetNum(0, false);
+	//TODO: If we paint a 2D projection of UVs, these will need to be the 2D vertices not the 3D original mesh vertices
+	BufferUtil::AppendElements(VertexROI, VertexSetBuffer);
+
+	// construct ROI triangle and weight buffers
+	ROITriangleBuffer.Reserve(TriangleROI.Num());
+	ROITriangleBuffer.SetNum(0, false);
+	for (int32 tid : TriangleROI)
+	{
+		ROITriangleBuffer.Add(tid);
+	}
+	ROIWeightValueBuffer.SetNum(VertexROI.Num(), false);
+	SyncWeightBufferWithMesh(Mesh);
+}
+
+bool UDataflowEditorWeightMapPaintTool::UpdateStampPosition(const FRay& WorldRay)
+{
+	CalculateBrushRadius();
+
+	TUniquePtr<FMeshSculptBrushOp>& UseBrushOp = GetActiveBrushOp();
+
+	ESculptBrushOpTargetType TargetType = UseBrushOp->GetBrushTargetType();
+	switch (TargetType)
+	{
+	case ESculptBrushOpTargetType::SculptMesh:
+	case ESculptBrushOpTargetType::TargetMesh:
+		UpdateBrushPositionOnSculptMesh(WorldRay, true);
+		break;
+	case ESculptBrushOpTargetType::ActivePlane:
+		check(false);
+		UpdateBrushPositionOnActivePlane(WorldRay);
+		break;
+	}
+
+	if (UseBrushOp->GetAlignStampToView())
+	{
+		AlignBrushToView();
+	}
+
+	CurrentStamp = LastStamp;
+	CurrentStamp.DeltaTime = FMathd::Min((FDateTime::Now() - LastStamp.TimeStamp).GetTotalSeconds(), 1.0);
+	CurrentStamp.WorldFrame = GetBrushFrameWorld();
+	CurrentStamp.LocalFrame = GetBrushFrameLocal();
+	CurrentStamp.Power = GetActivePressure() * GetCurrentBrushStrength();
+
+	CurrentStamp.PrevLocalFrame = LastStamp.LocalFrame;
+	CurrentStamp.PrevWorldFrame = LastStamp.WorldFrame;
+
+	FVector3d MoveDelta = CurrentStamp.LocalFrame.Origin - CurrentStamp.PrevLocalFrame.Origin;
+
+	if (UseBrushOp->IgnoreZeroMovements() && MoveDelta.SquaredLength() < 0.1 * CurrentBrushRadius)
+	{
+		return false;
+	}
+
+	return true;
+}
+
+
+bool UDataflowEditorWeightMapPaintTool::ApplyStamp()
+{
+	SCOPE_CYCLE_COUNTER(WeightMapPaintToolApplyStamp);
+
+	TUniquePtr<FMeshSculptBrushOp>& UseBrushOp = GetActiveBrushOp();
+
+	// yuck
+	FDataflowMeshVertexWeightMapEditBrushOp* WeightBrushOp = (FDataflowMeshVertexWeightMapEditBrushOp*)UseBrushOp.Get();
+
+	if (FilterProperties->SubToolType == EDataflowEditorWeightMapPaintInteractionType::Brush)
+	{
+		WeightBrushOp->bApplyRadiusLimit = true;
+	}
+	else
+	{ 
+		WeightBrushOp->bApplyRadiusLimit = false;
+	}
+
+	FDynamicMesh3* Mesh = GetSculptMesh();
+	WeightBrushOp->ApplyStampByVertices(Mesh, CurrentStamp, VertexROI, ROIWeightValueBuffer);
+
+	bool bUpdated = SyncMeshWithWeightBuffer(Mesh);
+
+	LastStamp = CurrentStamp;
+	LastStamp.TimeStamp = FDateTime::Now();
+
+	return bUpdated;
+}
+
+
+
+
+bool UDataflowEditorWeightMapPaintTool::SyncMeshWithWeightBuffer(FDynamicMesh3* Mesh)
+{
+	int NumModified = 0;
+	const int32 NumT = VertexROI.Num();
+	if (ActiveWeightMap)
+	{
+		// change update could be async here if we collected array of <idx,orig,new> and dispatched independenlty
+		for (int32 k = 0; k < NumT; ++k)
+		{
+			int VertIdx = VertexROI[k];
+			double CurWeight = GetCurrentWeightValue(VertIdx);
+
+			if (ROIWeightValueBuffer[k] != CurWeight)
+			{
+				if (bHaveDynamicMeshToWeightConversion)
+				{
+					for (const int32 Idx : WeightToDynamicMesh[DynamicMeshToWeight[VertIdx]])
+					{
+						ActiveWeightEditChangeTracker->SaveVertexOneRingTriangles(Idx, true);
+						ActiveWeightMap->SetValue(Idx, &ROIWeightValueBuffer[k]);
+					}
+				}
+				else
+				{
+				ActiveWeightEditChangeTracker->SaveVertexOneRingTriangles(VertIdx, true);
+				ActiveWeightMap->SetValue(VertIdx, &ROIWeightValueBuffer[k]);
+				}
+				NumModified++;
+			}
+		}
+	}
+	return (NumModified > 0);
+}
+
+bool UDataflowEditorWeightMapPaintTool::SyncWeightBufferWithMesh(const FDynamicMesh3* Mesh)
+{
+	int NumModified = 0;
+	const int32 NumT = VertexROI.Num();
+	if (ActiveWeightMap)
+	{
+		// change update could be async here if we collected array of <idx,orig,new> and dispatched independenlty
+		for (int32 k = 0; k < NumT; ++k)
+		{
+			int VertIdx = VertexROI[k];
+			double CurWeight = GetCurrentWeightValue(VertIdx);
+			if (ROIWeightValueBuffer[k] != CurWeight)
+			{
+				ROIWeightValueBuffer[k] = CurWeight;
+				NumModified++;
+			}
+		}
+	}
+	return (NumModified > 0);
+}
+
+template<typename RealType>
+static bool FindPolylineSelfIntersection(
+	const TArray<UE::Math::TVector2<RealType>>& Polyline, 
+	UE::Math::TVector2<RealType>& IntersectionPointOut, 
+	FIndex2i& IntersectionIndexOut,
+	bool bParallel = true)
+{
+	int32 N = Polyline.Num();
+	std::atomic<bool> bSelfIntersects(false);
+	ParallelFor(N - 1, [&](int32 i)
+	{
+		TSegment2<RealType> SegA(Polyline[i], Polyline[i + 1]);
+		for (int32 j = i + 2; j < N - 1 && bSelfIntersects == false; ++j)
+		{
+			TSegment2<RealType> SegB(Polyline[j], Polyline[j + 1]);
+			if (SegA.Intersects(SegB) && bSelfIntersects == false)		
+			{
+				bool ExpectedValue = false;
+				if (std::atomic_compare_exchange_strong(&bSelfIntersects, &ExpectedValue, true))
+				{
+					TIntrSegment2Segment2<RealType> Intersection(SegA, SegB);
+					Intersection.Find();
+					IntersectionPointOut = Intersection.Point0;
+					IntersectionIndexOut = FIndex2i(i, j);
+					return;
+				}
+			}
+		}
+	}, (bParallel) ? EParallelForFlags::None : EParallelForFlags::ForceSingleThread );
+
+	return bSelfIntersects;
+}
+
+
+
+template<typename RealType>
+static bool FindPolylineSegmentIntersection(
+	const TArray<UE::Math::TVector2<RealType>>& Polyline,
+	const TSegment2<RealType>& Segment,
+	UE::Math::TVector2<RealType>& IntersectionPointOut,
+	int& IntersectionIndexOut)
+{
+
+	int32 N = Polyline.Num();
+	for (int32 i = 0; i < N-1; ++i)
+	{
+		TSegment2<RealType> PolySeg(Polyline[i], Polyline[i + 1]);
+		if (Segment.Intersects(PolySeg))
+		{
+			TIntrSegment2Segment2<RealType> Intersection(Segment, PolySeg);
+			Intersection.Find();
+			IntersectionPointOut = Intersection.Point0;
+			IntersectionIndexOut = i;
+			return true;
+		}
+	}
+	return false;
+}
+
+
+
+bool ApproxSelfClipPolyline(TArray<FVector2f>& Polyline)
+{
+	int32 N = Polyline.Num();
+
+	// handle already-closed polylines
+	if (Distance(Polyline[0], Polyline[N-1]) < 0.0001f)
+	{
+		return true;
+	}
+
+	FVector2f IntersectPoint;
+	FIndex2i IntersectionIndex(-1, -1);
+	bool bSelfIntersects = FindPolylineSelfIntersection(Polyline, IntersectPoint, IntersectionIndex);
+	if (bSelfIntersects)
+	{
+		TArray<FVector2f> NewPolyline;
+		NewPolyline.Add(IntersectPoint);
+		for (int32 i = IntersectionIndex.A; i <= IntersectionIndex.B; ++i)
+		{
+			NewPolyline.Add(Polyline[i]);
+		}
+		NewPolyline.Add(IntersectPoint);
+		Polyline = MoveTemp(NewPolyline);
+		return true;
+	}
+
+
+	FVector2f StartDirOut = UE::Geometry::Normalized(Polyline[0] - Polyline[1]);
+	FLine2f StartLine(Polyline[0], StartDirOut);
+	FVector2f EndDirOut = UE::Geometry::Normalized(Polyline[N - 1] - Polyline[N - 2]);
+	FLine2f EndLine(Polyline[N - 1], EndDirOut);
+	FIntrLine2Line2f LineIntr(StartLine, EndLine);
+	bool bIntersects = false;
+	if (LineIntr.Find())
+	{
+		bIntersects = LineIntr.IsSimpleIntersection() && (LineIntr.Segment1Parameter > 0) && (LineIntr.Segment2Parameter > 0);
+		if (bIntersects)
+		{
+			Polyline.Add(StartLine.PointAt(LineIntr.Segment1Parameter));
+			Polyline.Add(StartLine.Origin);
+			return true;
+		}
+	}
+
+
+	FAxisAlignedBox2f Bounds;
+	for (const FVector2f& P : Polyline)
+	{
+		Bounds.Contain(P);
+	}
+	float Size = Bounds.DiagonalLength();
+
+	FVector2f StartPos = Polyline[0] + 0.001f * StartDirOut;
+	if (FindPolylineSegmentIntersection(Polyline, FSegment2f(StartPos, StartPos + 2*Size*StartDirOut), IntersectPoint, IntersectionIndex.A))
+	{
+		return true;
+	}
+
+	FVector2f EndPos = Polyline[N-1] + 0.001f * EndDirOut;
+	if (FindPolylineSegmentIntersection(Polyline, FSegment2f(EndPos, EndPos + 2*Size*EndDirOut), IntersectPoint, IntersectionIndex.A))
+	{
+		return true;
+	}
+
+	return false;
+}
+
+
+
+void UDataflowEditorWeightMapPaintTool::OnPolyLassoFinished(const FCameraPolyLasso& Lasso, bool bCanceled)
+{
+	// construct polyline
+	TArray<FVector2f> Polyline;
+	for (FVector2D Pos : Lasso.Polyline)
+	{
+		Polyline.Add((FVector2f)Pos);
+	}
+	int32 N = Polyline.Num();
+	if (N < 2)
+	{
+		return;
+	}
+
+	// Try to clip polyline to be closed, or closed-enough for winding evaluation to work.
+	// If that returns false, the polyline is "too open". In that case we will extend
+	// outwards from the endpoints and then try to create a closed very large polygon
+	if (ApproxSelfClipPolyline(Polyline) == false)
+	{
+		FVector2f StartDirOut = UE::Geometry::Normalized(Polyline[0] - Polyline[1]);
+		FLine2f StartLine(Polyline[0], StartDirOut);
+		FVector2f EndDirOut = UE::Geometry::Normalized(Polyline[N-1] - Polyline[N-2]);
+		FLine2f EndLine(Polyline[N-1], EndDirOut);
+
+		// if we did not intersect, we are in ambiguous territory. Check if a segment along either end-direction
+		// intersects the polyline. If it does, we have something like a spiral and will be OK. 
+		// If not, make a closed polygon by interpolating outwards from each endpoint, and then in perp-directions.
+		FPolygon2f Polygon(Polyline);
+		float PerpSign = Polygon.IsClockwise() ? -1.0 : 1.0;
+
+		Polyline.Insert(StartLine.PointAt(10000.0f), 0);
+		Polyline.Insert(Polyline[0] + 1000 * PerpSign * UE::Geometry::PerpCW(StartDirOut), 0);
+
+		Polyline.Add(EndLine.PointAt(10000.0f));
+		Polyline.Add(Polyline.Last() + 1000 * PerpSign * UE::Geometry::PerpCW(EndDirOut));
+		FVector2f StartPos = Polyline[0];
+		Polyline.Add(StartPos);		// close polyline (cannot use Polyline[0] in case Add resizes!)
+	}
+
+	N = Polyline.Num();
+
+	// project each mesh vertex to view plane and evaluate winding integral of polyline
+	const FDynamicMesh3* Mesh = GetSculptMesh();
+	TempROIBuffer.SetNum(Mesh->MaxVertexID());
+	ParallelFor(Mesh->MaxVertexID(), [&](int32 vid)
+	{
+		if (Mesh->IsVertex(vid))
+		{
+			FVector3d WorldPos = CurTargetTransform.TransformPosition(Mesh->GetVertex(vid));
+			FVector2f PlanePos = (FVector2f)Lasso.GetProjectedPoint((FVector)WorldPos);
+
+			double WindingSum = 0;
+			FVector2f a = Polyline[0] - PlanePos, b = FVector2f::Zero();
+			for (int32 i = 1; i < N; ++i)
+			{
+				b = Polyline[i] - PlanePos;
+				WindingSum += (double)FMathf::Atan2(a.X*b.Y - a.Y*b.X, a.X*b.X + a.Y*b.Y);
+				a = b;
+			}
+			WindingSum /= FMathd::TwoPi;
+			bool bInside = FMathd::Abs(WindingSum) > 0.3;
+			TempROIBuffer[vid] = bInside ? 1 : 0;
+		}
+		else
+		{
+			TempROIBuffer[vid] = -1;
+		}
+	});
+
+	// convert to vertex selection, and then select fully-enclosed faces
+	FMeshVertexSelection VertexSelection(Mesh);
+	VertexSelection.SelectByVertexID([&](int32 vid) { return TempROIBuffer[vid] == 1; });
+
+	double SetWeightValue = GetInEraseStroke() ? 0.0 : FilterProperties->AttributeValue;
+	SetVerticesToWeightMap(VertexSelection.AsSet(), SetWeightValue, GetInEraseStroke());
+}
+
+
+void UDataflowEditorWeightMapPaintTool::ComputeGradient()
+{
+	if (!ensure(ActiveWeightMap))
+	{
+		UE_LOG(LogTemp, Warning, TEXT("No active weight map"));
+		return;
+	}
+
+	BeginChange();
+
+	const FDynamicMesh3* const Mesh = DynamicMeshComponent->GetMesh();
+	TempROIBuffer.SetNum(0, false);
+	for (int32 vid : Mesh->VertexIndicesItr())
+	{
+		TempROIBuffer.Add(vid);
+	}
+
+	if (bHaveDynamicMeshToWeightConversion)
+	{
+		for (int32 vid : TempROIBuffer)
+		{
+			for (const int32 Idx : WeightToDynamicMesh[DynamicMeshToWeight[vid]])
+			{
+				ActiveWeightEditChangeTracker->SaveVertexOneRingTriangles(Idx, true);
+			}
+		}
+	}
+	else
+	{
+		for (int32 vid : TempROIBuffer)
+		{
+			ActiveWeightEditChangeTracker->SaveVertexOneRingTriangles(vid, true);
+		}
+	}
+
+
+	for (const int32 VertexIndex : TempROIBuffer)
+	{
+		const FVector3d Vert = Mesh->GetVertex(VertexIndex);
+
+		// (Copied from FClothPaintTool_Gradient::ApplyGradient)
+
+		// Get distances
+		// TODO: Look into surface distance instead of 3D distance? May be necessary for some complex shapes
+		float DistanceToLowSq = MAX_flt;
+		for (const int32& LowIndex : LowValueGradientVertexSelection.SelectedCornerIDs)
+		{
+			const FVector3d LowPoint = Mesh->GetVertex(LowIndex);
+			const float DistanceSq = (LowPoint - Vert).SizeSquared();
+			if (DistanceSq < DistanceToLowSq)
+			{
+				DistanceToLowSq = DistanceSq;
+			}
+		}
+
+		float DistanceToHighSq = MAX_flt;
+		for (const int32& HighIndex : HighValueGradientVertexSelection.SelectedCornerIDs)
+		{
+			const FVector3d HighPoint = Mesh->GetVertex(HighIndex);
+			const float DistanceSq = (HighPoint - Vert).SizeSquared();
+			if (DistanceSq < DistanceToHighSq)
+			{
+				DistanceToHighSq = DistanceSq;
+			}
+		}
+
+		const float Value = FMath::LerpStable(FilterProperties->GradientLowValue, FilterProperties->GradientHighValue, DistanceToLowSq / (DistanceToLowSq + DistanceToHighSq));
+		if (bHaveDynamicMeshToWeightConversion)
+		{
+			for (const int32 Idx : WeightToDynamicMesh[DynamicMeshToWeight[VertexIndex]])
+			{
+				ActiveWeightMap->SetValue(Idx, &Value);
+			}
+		}
+		else
+		{
+		ActiveWeightMap->SetValue(VertexIndex, &Value);
+	}
+	}
+
+	// update colors
+	UpdateVertexColorOverlay();
+	DynamicMeshComponent->FastNotifyVertexAttributesUpdated(EMeshRenderAttributeFlags::VertexColors);
+	GetToolManager()->PostInvalidation();
+	EndChange();
+
+}
+
+
+void UDataflowEditorWeightMapPaintTool::OnSelectionModified()
+{
+	const bool bToolTypeIsGradient = (FilterProperties->SubToolType == EDataflowEditorWeightMapPaintInteractionType::Gradient);
+	if (bToolTypeIsGradient && PolygonSelectionMechanic)
+	{
+		const FGroupTopologySelection NewSelection = PolygonSelectionMechanic->GetActiveSelection();
+
+		const bool bSelectingLowValueGradientVertices = GetCtrlToggle();
+		if (bSelectingLowValueGradientVertices)
+		{
+			LowValueGradientVertexSelection = NewSelection;
+		}
+		else
+		{
+			HighValueGradientVertexSelection = NewSelection;
+		}
+
+		if (LowValueGradientVertexSelection.SelectedCornerIDs.Num() > 0 && HighValueGradientVertexSelection.SelectedCornerIDs.Num() > 0)
+		{
+			ComputeGradient();
+		}
+
+		constexpr bool bBroadcast = false;
+		PolygonSelectionMechanic->SetSelection(FGroupTopologySelection(), bBroadcast);
+	}
+}
+
+
+void UDataflowEditorWeightMapPaintTool::SetVerticesToWeightMap(const TSet<int32>& Vertices, double WeightValue, bool bIsErase)
+{
+	BeginChange();
+
+	TempROIBuffer.SetNum(0, false);
+	for (int32 vid : Vertices)
+	{
+		TempROIBuffer.Add(vid);
+	}
+
+	if (HaveVisibilityFilter())
+	{
+		TArray<int32> VisibleVertices;
+		VisibleVertices.Reserve(TempROIBuffer.Num());
+		ApplyVisibilityFilter(TempROIBuffer, VisibleVertices);
+		TempROIBuffer = MoveTemp(VisibleVertices);
+	}
+
+	if (bHaveDynamicMeshToWeightConversion)
+	{
+		for (int32 vid : TempROIBuffer)
+		{
+			for (const int32 Idx : WeightToDynamicMesh[DynamicMeshToWeight[vid]])
+			{
+				ActiveWeightEditChangeTracker->SaveVertexOneRingTriangles(Idx, true);
+				ActiveWeightMap->SetValue(Idx, &WeightValue);
+			}
+		}
+	}
+	else
+	{
+		for (int32 vid : TempROIBuffer)
+		{
+			ActiveWeightEditChangeTracker->SaveVertexOneRingTriangles(vid, true);
+		}
+		for (int32 vid : TempROIBuffer)
+		{		
+			ActiveWeightMap->SetValue(vid, &WeightValue);
+		}
+	}
+
+	// update colors
+	UpdateVertexColorOverlay();
+	DynamicMeshComponent->FastNotifyVertexAttributesUpdated(EMeshRenderAttributeFlags::VertexColors);
+	GetToolManager()->PostInvalidation();
+	
+
+	EndChange();
+	
+}
+
+
+
+bool UDataflowEditorWeightMapPaintTool::HaveVisibilityFilter() const
+{
+	return FilterProperties->VisibilityFilter != EDataflowEditorWeightMapPaintVisibilityType::None;
+}
+
+
+void UDataflowEditorWeightMapPaintTool::ApplyVisibilityFilter(TSet<int32>& Vertices, TArray<int32>& ROIBuffer, TArray<int32>& OutputBuffer)
+{
+	ROIBuffer.SetNum(0, false);
+	ROIBuffer.Reserve(Vertices.Num());
+	for (int32 vid : Vertices)
+	{
+		ROIBuffer.Add(vid);
+	}
+	
+	OutputBuffer.Reset();
+	ApplyVisibilityFilter(TempROIBuffer, OutputBuffer);
+
+	Vertices.Reset();
+	for (int32 vid : OutputBuffer)
+	{
+		Vertices.Add(vid);
+	}
+}
+
+void UDataflowEditorWeightMapPaintTool::ApplyVisibilityFilter(const TArray<int32>& Vertices, TArray<int32>& VisibleVertices)
+{
+	if (!HaveVisibilityFilter())
+	{
+		VisibleVertices = Vertices;
+		return;
+	}
+
+	FViewCameraState StateOut;
+	GetToolManager()->GetContextQueriesAPI()->GetCurrentViewState(StateOut);
+	FTransform3d LocalToWorld = UE::ToolTarget::GetLocalToWorldTransform(Target);
+	FVector3d LocalEyePosition(LocalToWorld.InverseTransformPosition(StateOut.Position));
+
+	const FDynamicMesh3* Mesh = GetSculptMesh();
+
+	int32 NumVertices = Vertices.Num();
+
+	VisibilityFilterBuffer.SetNum(NumVertices, false);
+	ParallelFor(NumVertices, [&](int32 idx)
+	{
+		VisibilityFilterBuffer[idx] = true;
+		UE::Geometry::FVertexInfo VertexInfo;
+		Mesh->GetVertex(Vertices[idx], VertexInfo, true, false, false);
+		FVector3d Centroid = VertexInfo.Position;
+		FVector3d FaceNormal = (FVector3d)VertexInfo.Normal;
+		if (FaceNormal.Dot((Centroid - LocalEyePosition)) > 0)
+		{
+			VisibilityFilterBuffer[idx] = false;
+		}
+		if (FilterProperties->VisibilityFilter == EDataflowEditorWeightMapPaintVisibilityType::Unoccluded)
+		{
+			int32 HitTID = Octree.FindNearestHitObject(FRay3d(LocalEyePosition, UE::Geometry::Normalized(Centroid - LocalEyePosition)));
+			if (HitTID != IndexConstants::InvalidID && Mesh->IsTriangle(HitTID))
+			{
+				// Check to see if our vertex has been occulded by another triangle.
+				FIndex3i TriVertices = Mesh->GetTriangle(HitTID);
+				if (TriVertices[0] != Vertices[idx] && TriVertices[1] != Vertices[idx] && TriVertices[2] != Vertices[idx])
+				{
+					VisibilityFilterBuffer[idx] = false;
+				}
+			}
+		}
+	});
+
+	VisibleVertices.Reset();
+	for (int32 k = 0; k < NumVertices; ++k)
+	{
+		if (VisibilityFilterBuffer[k])
+		{
+			VisibleVertices.Add(Vertices[k]);
+		}
+	}
+}
+
+
+
+int32 UDataflowEditorWeightMapPaintTool::FindHitSculptMeshTriangle(const FRay3d& LocalRay)
+{
+	// TODO: Figure out what the actual position on the triangle is when hit.
+	CurrentBaryCentricCoords = FVector3d(1.0 / 3.0, 1.0 / 3.0, 1.0 / 3.0);
+
+	if (!IsInBrushSubMode())
+	{
+		return IndexConstants::InvalidID;
+	}
+
+	if (GetBrushCanHitBackFaces())
+	{
+		return Octree.FindNearestHitObject(LocalRay);
+	}
+	else
+	{
+		FDynamicMesh3* Mesh = GetSculptMesh();
+
+		FViewCameraState StateOut;
+		GetToolManager()->GetContextQueriesAPI()->GetCurrentViewState(StateOut);
+		FVector3d LocalEyePosition(CurTargetTransform.InverseTransformPosition((FVector3d)StateOut.Position));
+		int HitTID = Octree.FindNearestHitObject(LocalRay,
+			[this, Mesh, &LocalEyePosition](int TriangleID) {
+			FVector3d Normal, Centroid;
+			double Area;
+			Mesh->GetTriInfo(TriangleID, Normal, Area, Centroid);
+			return Normal.Dot((Centroid - LocalEyePosition)) < 0;
+		});
+		return HitTID;
+	}
+}
+
+int32 UDataflowEditorWeightMapPaintTool::FindHitTargetMeshTriangle(const FRay3d& LocalRay)
+{
+	check(false);
+	return IndexConstants::InvalidID;
+}
+
+
+
+bool UDataflowEditorWeightMapPaintTool::UpdateBrushPosition(const FRay& WorldRay)
+{
+	TUniquePtr<FMeshSculptBrushOp>& UseBrushOp = GetActiveBrushOp();
+
+	bool bHit = false; 
+	ESculptBrushOpTargetType TargetType = UseBrushOp->GetBrushTargetType();
+	switch (TargetType)
+	{
+	case ESculptBrushOpTargetType::SculptMesh:
+	case ESculptBrushOpTargetType::TargetMesh:
+		bHit = UpdateBrushPositionOnSculptMesh(WorldRay, false);
+		break;
+	case ESculptBrushOpTargetType::ActivePlane:
+		check(false);
+		bHit = UpdateBrushPositionOnSculptMesh(WorldRay, false);
+		break;
+	}
+
+	if (bHit && UseBrushOp->GetAlignStampToView())
+	{
+		AlignBrushToView();
+	}
+
+	return bHit;
+}
+
+
+
+
+bool UDataflowEditorWeightMapPaintTool::OnUpdateHover(const FInputDeviceRay& DevicePos)
+{
+	PendingStampType = FilterProperties->PrimaryBrushType;
+
+	if(ensure(InStroke() == false))
+	{
+		UpdateBrushPosition(DevicePos.WorldRay);
+	}
+	return true;
+}
+
+
+void UDataflowEditorWeightMapPaintTool::DrawHUD(FCanvas* Canvas, IToolsContextRenderAPI* RenderAPI)
+{
+	if (PolyLassoMechanic)
+	{
+		// because the actual Weight change is deferred until mouse release, color the lasso to let the user know whether it will erase
+		PolyLassoMechanic->LineColor = GetInEraseStroke() ? FLinearColor::Red : FLinearColor::Green;
+		PolyLassoMechanic->DrawHUD(Canvas, RenderAPI);
+	}
+
+	if (PolygonSelectionMechanic)
+	{
+		PolygonSelectionMechanic->DrawHUD(Canvas, RenderAPI);
+	}
+}
+
+void UDataflowEditorWeightMapPaintTool::Render(IToolsContextRenderAPI* RenderAPI)
+{
+	UMeshSculptToolBase::Render(RenderAPI);
+
+	if (PolygonSelectionMechanic)
+	{
+		PolygonSelectionMechanic->RenderMarquee(RenderAPI);
+
+		const FViewCameraState RenderCameraState = RenderAPI->GetCameraState();
+		GradientSelectionRenderer.BeginFrame(RenderAPI, RenderCameraState);
+
+		const FTransform Transform = DynamicMeshComponent->GetComponentTransform();
+		GradientSelectionRenderer.SetTransform(Transform);
+			
+		GradientSelectionRenderer.SetPointParameters(FLinearColor::Red, 1.0);
+		PolygonSelectionMechanic->UMeshTopologySelectionMechanic::GetTopologySelector()->DrawSelection(LowValueGradientVertexSelection, &GradientSelectionRenderer, &RenderCameraState);
+
+		GradientSelectionRenderer.SetPointParameters(FLinearColor::Green, 1.0);
+		PolygonSelectionMechanic->UMeshTopologySelectionMechanic::GetTopologySelector()->DrawSelection(HighValueGradientVertexSelection, &GradientSelectionRenderer, &RenderCameraState);
+
+		// Now the current unsaved selection
+		if (GetCtrlToggle())
+		{
+			GradientSelectionRenderer.SetPointParameters(FLinearColor::Red, 1.0);
+		}
+		else
+		{
+			GradientSelectionRenderer.SetPointParameters(FLinearColor::Green, 1.0);
+		}
+
+		PolygonSelectionMechanic->UMeshTopologySelectionMechanic::GetTopologySelector()->DrawSelection(PolygonSelectionMechanic->GetActiveSelection(), &GradientSelectionRenderer, &RenderCameraState);
+
+		GradientSelectionRenderer.EndFrame();
+	}
+}
+
+
+void UDataflowEditorWeightMapPaintTool::UpdateMaterialMode(EMeshEditingMaterialModes MaterialMode)
+{
+	if (MaterialMode == EMeshEditingMaterialModes::VertexColor)
+	{
+		constexpr bool bUseTwoSidedMaterial = true;
+		ActiveOverrideMaterial = ToolSetupUtil::GetVertexColorMaterial(GetToolManager(), bUseTwoSidedMaterial);
+		if (ensure(ActiveOverrideMaterial != nullptr))
+		{
+			GetSculptMeshComponent()->SetOverrideRenderMaterial(ActiveOverrideMaterial);
+			ActiveOverrideMaterial->SetScalarParameterValue(TEXT("FlatShading"), (ViewProperties->bFlatShading) ? 1.0f : 0.0f);
+		}
+		GetSculptMeshComponent()->SetShadowsEnabled(false);
+	}
+	else
+	{
+		UMeshSculptToolBase::UpdateMaterialMode(MaterialMode);
+	}
+}
+
+void UDataflowEditorWeightMapPaintTool::UpdateStampPendingState()
+{
+	if (InStroke() == false) return;
+	bIsStampPending = true;
+}
+
+void UDataflowEditorWeightMapPaintTool::OnTick(float DeltaTime)
+{
+	UMeshSculptToolBase::OnTick(DeltaTime);
+	MeshElementsDisplay->OnTick(DeltaTime);
+
+	const bool bIsLasso = (FilterProperties->SubToolType == EDataflowEditorWeightMapPaintInteractionType::PolyLasso);
+	PolyLassoMechanic->SetIsEnabled(bIsLasso);
+
+	const bool bIsGradient = (FilterProperties->SubToolType == EDataflowEditorWeightMapPaintInteractionType::Gradient);
+	PolygonSelectionMechanic->SetIsEnabled(bIsGradient);
+
+	check(!(bIsLasso && bIsGradient));
+
+	ConfigureIndicator(false);
+	SetIndicatorVisibility(!bIsLasso && !bIsGradient);
+
+	if (bHavePendingAction)
+	{
+		ApplyAction(PendingAction);
+		bHavePendingAction = false;
+		PendingAction = EDataflowEditorWeightMapPaintToolActions::NoAction;
+	}
+
+	SCOPE_CYCLE_COUNTER(WeightMapPaintToolTick);
+
+	// process the undo update
+	if (bUndoUpdatePending)
+	{
+		// wait for updates
+		WaitForPendingUndoRedo();
+
+		// post rendering update
+		DynamicMeshComponent->FastNotifyTriangleVerticesUpdated(AccumulatedTriangleROI, EMeshRenderAttributeFlags::VertexColors);
+		GetToolManager()->PostInvalidation();
+
+		// ignore stamp and wait for next tick to do anything else
+		bUndoUpdatePending = false;
+		return;
+	}
+
+	// Get value at brush location
+	const bool bShouldPickWeight = bPendingPickWeight && IsStampPending() == false;
+	const bool bShouldUpdateValueAtBrush = IsInBrushSubMode();
+
+	if (bShouldPickWeight || bShouldUpdateValueAtBrush)
+	{
+		if (GetSculptMesh()->IsTriangle(GetBrushTriangleID()))
+		{
+			const double HitWeightValue = GetCurrentWeightValueUnderBrush();
+
+			if (bShouldPickWeight)
+			{
+				FilterProperties->AttributeValue = HitWeightValue;
+				NotifyOfPropertyChangeByTool(FilterProperties);
+			}
+
+			if (bShouldUpdateValueAtBrush)
+			{
+				FilterProperties->ValueAtBrush = HitWeightValue;
+			}
+		}
+		bPendingPickWeight = false;
+	}
+
+
+	if (IsInBrushSubMode())
+	{
+		if (InStroke())
+		{
+			SCOPE_CYCLE_COUNTER(WeightMapPaintTool_Tick_ApplyStampBlock);
+
+			// update brush position
+			if (UpdateStampPosition(GetPendingStampRayWorld()) == false)
+			{
+				return;
+			}
+			UpdateStampPendingState();
+			if (IsStampPending() == false)
+			{
+				return;
+			}
+
+			// update sculpt ROI
+			UpdateROI(CurrentStamp);
+
+			// append updated ROI to modified region (async)
+			FDynamicMesh3* Mesh = GetSculptMesh();
+			TFuture<void> AccumulateROI = Async(Dataflow::Private::WeightPaintToolAsyncExecTarget, [&]()
+			{
+				UE::Geometry::VertexToTriangleOneRing(Mesh, VertexROI, AccumulatedTriangleROI);
+			});
+
+			// apply the stamp
+			bool bWeightsModified = ApplyStamp();
+
+			if (bWeightsModified)
+			{
+				SCOPE_CYCLE_COUNTER(WeightMapPaintTool_Tick_UpdateMeshBlock);
+				UpdateVertexColorOverlay(&TriangleROI);
+				DynamicMeshComponent->FastNotifyTriangleVerticesUpdated(TriangleROI, EMeshRenderAttributeFlags::VertexColors);
+				GetToolManager()->PostInvalidation();
+			}
+
+			// we don't really need to wait for these to happen to end Tick()...
+			AccumulateROI.Wait();
+		}
+	}
+
+}
+
+bool UDataflowEditorWeightMapPaintTool::CanAccept() const
+{
+	return bAnyChangeMade || UpdateWeightMapProperties->Name != WeightMapNodeToUpdate->Name;
+}
+
+
+
+FColor UDataflowEditorWeightMapPaintTool::GetColorForWeightValue(double WeightValue)
+{
+	FColor MaxColor = LinearColors::White3b();
+	FColor MinColor = LinearColors::Black3b();
+	FColor Color;
+	double ClampedValue = FMath::Clamp(WeightValue, 0.0, 1.0);
+	Color.R = FMath::LerpStable(MinColor.R, MaxColor.R, ClampedValue);
+	Color.G = FMath::LerpStable(MinColor.G, MaxColor.G, ClampedValue);
+	Color.B = FMath::LerpStable(MinColor.B, MaxColor.B, ClampedValue);
+	Color.A = 1.0;
+
+	return Color;
+}
+
+void UDataflowEditorWeightMapPaintTool::FloodFillCurrentWeightAction()
+{
+	if (!ActiveWeightMap)
+	{
+		return;
+	}
+
+	BeginChange();
+
+	const float SetWeightValue = FilterProperties->AttributeValue;
+	const FDynamicMesh3* Mesh = DynamicMeshComponent->GetMesh();
+	TempROIBuffer.SetNum(0, false);
+	for (int32 vid : Mesh->VertexIndicesItr())
+	{
+		TempROIBuffer.Add(vid);
+	}
+
+	if (bHaveDynamicMeshToWeightConversion)
+	{
+		for (int32 vid : TempROIBuffer)
+		{
+			for (const int32 Idx : WeightToDynamicMesh[DynamicMeshToWeight[vid]])
+			{
+				ActiveWeightEditChangeTracker->SaveVertexOneRingTriangles(Idx, true);
+				ActiveWeightMap->SetValue(Idx, &SetWeightValue);
+			}
+		}
+	}
+	else
+	{
+		for (int32 vid : TempROIBuffer)
+		{
+			ActiveWeightEditChangeTracker->SaveVertexOneRingTriangles(vid, true);
+		}
+		for (int32 vid : TempROIBuffer)
+		{
+			ActiveWeightMap->SetValue(vid, &SetWeightValue);
+		}
+	}
+
+	// update colors
+	UpdateVertexColorOverlay();
+	DynamicMeshComponent->FastNotifyVertexAttributesUpdated(EMeshRenderAttributeFlags::VertexColors);
+	GetToolManager()->PostInvalidation();
+	EndChange();
+}
+
+
+void UDataflowEditorWeightMapPaintTool::ClearAllWeightsAction()
+{
+	if (!ActiveWeightMap)
+	{
+		return;
+	}
+
+	BeginChange();
+
+	float SetWeightValue = 0.0f;
+	const FDynamicMesh3* Mesh = DynamicMeshComponent->GetMesh();
+	TempROIBuffer.SetNum(0, false);
+	for (int32 vid : Mesh->VertexIndicesItr())
+	{
+		TempROIBuffer.Add(vid);
+	}
+
+	if (bHaveDynamicMeshToWeightConversion)
+	{
+		for (int32 vid : TempROIBuffer)
+		{
+			for (const int32 Idx : WeightToDynamicMesh[DynamicMeshToWeight[vid]])
+			{
+				ActiveWeightEditChangeTracker->SaveVertexOneRingTriangles(Idx, true);
+				ActiveWeightMap->SetValue(Idx, &SetWeightValue);
+			}
+		}
+	}
+	else
+	{
+		for (int32 vid : TempROIBuffer)
+		{
+			ActiveWeightEditChangeTracker->SaveVertexOneRingTriangles(vid, true);
+		}
+		for (int32 vid : TempROIBuffer)
+		{
+			ActiveWeightMap->SetValue(vid, &SetWeightValue);
+		}
+	}
+
+	// update colors
+	UpdateVertexColorOverlay();
+	DynamicMeshComponent->FastNotifyVertexAttributesUpdated(EMeshRenderAttributeFlags::VertexColors);
+	GetToolManager()->PostInvalidation();
+	EndChange();
+}
+
+
+void UDataflowEditorWeightMapPaintTool::UpdateSelectedNode()
+{
+	check(ActiveWeightMap);
+	TArray<float> CurrentWeights;
+	GetCurrentWeightMap(CurrentWeights);
+
+	check(WeightMapNodeToUpdate);
+
+	if (bHaveDynamicMeshToWeightConversion)
+	{
+		WeightMapNodeToUpdate->VertexWeights.Init(0.f, WeightToDynamicMesh.Num());
+		for (int32 DynamicMeshIdx = 0; DynamicMeshIdx < CurrentWeights.Num(); ++DynamicMeshIdx)
+		{
+			WeightMapNodeToUpdate->VertexWeights[DynamicMeshToWeight[DynamicMeshIdx]] = CurrentWeights[DynamicMeshIdx];
+		}
+	}
+	else
+	{
+		WeightMapNodeToUpdate->VertexWeights = CurrentWeights;
+	}
+	
+	WeightMapNodeToUpdate->Name = UpdateWeightMapProperties->Name;
+
+	WeightMapNodeToUpdate->Invalidate();
+}
+
+
+//
+// Change Tracking
+//
+
+namespace DataflowWeightPaintLocals
+{
+
+	/**
+	 * A wrapper change that applies a given change to the unwrap canonical mesh of an input, and uses that
+	 * to update the other views. Causes a broadcast of OnCanonicalModified.
+	 */
+	class  FDataflowWeightPaintMeshChange : public FToolCommandChange
+	{
+	public:
+		FDataflowWeightPaintMeshChange(UDynamicMeshComponent* DynamicMeshComponentIn, TUniquePtr<FDynamicMeshChange> DynamicMeshChangeIn)
+			: DynamicMeshComponent(DynamicMeshComponentIn)
+			, DynamicMeshChange(MoveTemp(DynamicMeshChangeIn))
+		{
+			ensure(DynamicMeshComponentIn);
+			ensure(DynamicMeshChange);
+		};
+
+		virtual void Apply(UObject* Object) override
+		{
+			DynamicMeshChange->Apply(DynamicMeshComponent->GetMesh(), false);			
+		}
+
+		virtual void Revert(UObject* Object) override
+		{
+			DynamicMeshChange->Apply(DynamicMeshComponent->GetMesh(), true);			
+		}
+
+		virtual bool HasExpired(UObject* Object) const override
+		{
+			return !(DynamicMeshComponent.IsValid() && DynamicMeshChange);
+		}
+
+
+		virtual FString ToString() const override
+		{
+			return TEXT("FDataflowWeightPaintMeshChange");
+		}
+
+	protected:
+		TWeakObjectPtr<UDynamicMeshComponent> DynamicMeshComponent;
+		TUniquePtr<FDynamicMeshChange> DynamicMeshChange;
+	};
+}
+
+void UDataflowEditorWeightMapPaintTool::BeginChange()
+{
+
+	check(ActiveWeightEditChangeTracker == nullptr);
+	
+	ActiveWeightEditChangeTracker = MakeUnique<FDynamicMeshChangeTracker>(GetSculptMesh());
+	ActiveWeightEditChangeTracker->BeginChange();
+	LongTransactions.Open(LOCTEXT("WeightPaintChange", "Weight Stroke"), GetToolManager());
+}
+
+void UDataflowEditorWeightMapPaintTool::EndChange()
+{
+	check(ActiveWeightEditChangeTracker);
+
+	bAnyChangeMade = true;
+
+	TUniquePtr<FDynamicMeshChange> EditResult = ActiveWeightEditChangeTracker->EndChange();
+
+	TUniquePtr<DataflowWeightPaintLocals::FDataflowWeightPaintMeshChange> DataflowWeightPaintMeshChange =
+		MakeUnique<DataflowWeightPaintLocals::FDataflowWeightPaintMeshChange>(DynamicMeshComponent.Get(), MoveTemp(EditResult));
+	ActiveWeightEditChangeTracker = nullptr;
+
+	TUniquePtr<TWrappedToolCommandChange<DataflowWeightPaintLocals::FDataflowWeightPaintMeshChange>> NewChange = MakeUnique<TWrappedToolCommandChange<DataflowWeightPaintLocals::FDataflowWeightPaintMeshChange>>();
+	NewChange->WrappedChange = MoveTemp(DataflowWeightPaintMeshChange);
+	NewChange->BeforeModify = [this](bool bRevert)
+	{
+		this->WaitForPendingUndoRedo();
+	};
+
+	GetToolManager()->EmitObjectChange(DynamicMeshComponent, MoveTemp(NewChange), LOCTEXT("WeightPaintChange", "Weight Stroke"));
+	LongTransactions.Close(GetToolManager());
+}
+
+
+void UDataflowEditorWeightMapPaintTool::WaitForPendingUndoRedo()
+{
+	if (bUndoUpdatePending)
+	{
+		bUndoUpdatePending = false;
+	}
+}
+
+void UDataflowEditorWeightMapPaintTool::OnDynamicMeshComponentChanged(UDynamicMeshComponent* Component, const FMeshVertexChange* Change, bool bRevert)
+{
+	// update octree
+	FDynamicMesh3* Mesh = GetSculptMesh();
+
+	// make sure any previous async computations are done, and update the undo ROI
+	if (bUndoUpdatePending)
+	{
+		// we should never hit this anymore, because of pre-change calling WaitForPendingUndoRedo()
+		WaitForPendingUndoRedo();
+
+		// this is not right because now we are going to do extra recomputation, but it's very messy otherwise...
+		UE::Geometry::VertexToTriangleOneRing(Mesh, Change->Vertices, AccumulatedTriangleROI);
+	}
+	else
+	{
+		AccumulatedTriangleROI.Reset();
+		UE::Geometry::VertexToTriangleOneRing(Mesh, Change->Vertices, AccumulatedTriangleROI);
+	}
+
+	// note that we have a pending update
+	bUndoUpdatePending = true;
+}
+
+
+void UDataflowEditorWeightMapPaintTool::PrecomputeFilterData()
+{
+	const FDynamicMesh3* Mesh = GetSculptMesh();
+	
+	TriNormals.SetNum(Mesh->MaxTriangleID());
+	ParallelFor(Mesh->MaxTriangleID(), [&](int32 tid)
+	{
+		if (Mesh->IsTriangle(tid))
+		{
+			TriNormals[tid] = Mesh->GetTriNormal(tid);
+		}
+	});
+
+	const FDynamicMeshNormalOverlay* Normals = Mesh->Attributes()->PrimaryNormals();
+	const FDynamicMeshUVOverlay* UVs = Mesh->Attributes()->PrimaryUV();
+	UVSeamEdges.SetNum(Mesh->MaxEdgeID());
+	NormalSeamEdges.SetNum(Mesh->MaxEdgeID());
+	ParallelFor(Mesh->MaxEdgeID(), [&](int32 eid)
+	{
+		if (Mesh->IsEdge(eid))
+		{
+			UVSeamEdges[eid] = UVs->IsSeamEdge(eid);
+			NormalSeamEdges[eid] = Normals->IsSeamEdge(eid);
+		}
+	});
+}
+
+double UDataflowEditorWeightMapPaintTool::GetCurrentWeightValue(int32 VertexId) const
+{
+	float WeightValue = 0.0;
+	if (ActiveWeightMap && VertexId != IndexConstants::InvalidID)
+	{
+		ActiveWeightMap->GetValue(VertexId, &WeightValue);
+	}
+	return WeightValue;
+}
+
+double UDataflowEditorWeightMapPaintTool::GetCurrentWeightValueUnderBrush() const
+{
+	float WeightValue = 0.0;
+	int32 VertexID = GetBrushNearestVertex();
+	if (ActiveWeightMap && VertexID != IndexConstants::InvalidID)
+	{
+		ActiveWeightMap->GetValue(VertexID, &WeightValue);
+	}
+	return WeightValue;
+}
+
+int32 UDataflowEditorWeightMapPaintTool::GetBrushNearestVertex() const
+{
+	int TriangleVertex = 0;
+
+	if (CurrentBaryCentricCoords.X >= CurrentBaryCentricCoords.Y && CurrentBaryCentricCoords.X >= CurrentBaryCentricCoords.Z)
+	{
+		TriangleVertex = 0;
+	}
+	else
+	{
+		if (CurrentBaryCentricCoords.Y >= CurrentBaryCentricCoords.X && CurrentBaryCentricCoords.Y >= CurrentBaryCentricCoords.Z)
+		{
+			TriangleVertex = 1;
+		}
+		else
+		{
+			TriangleVertex = 2;
+		}
+	}
+	const FDynamicMesh3* Mesh = GetSculptMesh();
+
+	int32 Tid = GetBrushTriangleID();
+	if (Tid == IndexConstants::InvalidID)
+	{
+		return IndexConstants::InvalidID;
+	}
+	
+	FIndex3i Vertices = Mesh->GetTriangle(Tid);
+	return Vertices[TriangleVertex];
+}
+
+void UDataflowEditorWeightMapPaintTool::GetCurrentWeightMap(TArray<float>& OutWeights) const
+{
+	if (ActiveWeightMap)
+	{
+		const FDynamicMesh3* Mesh = GetSculptMesh();
+		const int32 NumVertices = Mesh->VertexCount();
+		OutWeights.SetNumUninitialized(NumVertices);
+		for (int32 VertexID = 0; VertexID < NumVertices; ++VertexID)
+		{
+			ActiveWeightMap->GetValue(VertexID, &OutWeights[VertexID]);
+		}
+	}
+}
+
+
+
+void UDataflowEditorWeightMapPaintTool::UpdateSubToolType(EDataflowEditorWeightMapPaintInteractionType NewType)
+{
+	// Currenly we mirror base-brush properties in UDataflowEditorWeightMapPaintBrushFilterProperties, so we never
+	// want to show both
+	//bool bSculptPropsVisible = (NewType == EDataflowEditorWeightMapPaintInteractionType::Brush);
+	//SetToolPropertySourceEnabled(UMeshSculptToolBase::BrushProperties, bSculptPropsVisible);
+	SetToolPropertySourceEnabled(UMeshSculptToolBase::BrushProperties, false);
+
+	SetToolPropertySourceEnabled(FilterProperties, true);
+	SetBrushOpPropsVisibility(false);
+
+	if (NewType != EDataflowEditorWeightMapPaintInteractionType::Gradient)
+	{
+		LowValueGradientVertexSelection.Clear();
+		HighValueGradientVertexSelection.Clear();
+	}
+}
+
+
+void UDataflowEditorWeightMapPaintTool::UpdateBrushType(EDataflowEditorWeightMapPaintBrushType BrushType)
+{
+	static const FText BaseMessage = LOCTEXT("OnStartTool", "Hold Shift to Erase. [/] and S/D change Size (+Shift to small-step)");
+	FTextBuilder Builder;
+	Builder.AppendLine(BaseMessage);
+
+	SetActivePrimaryBrushType((int32)BrushType);
+
+	SetToolPropertySourceEnabled(GizmoProperties, false);
+
+	GetToolManager()->DisplayMessage(Builder.ToText(), EToolMessageLevel::UserNotification);
+}
+
+
+
+
+void UDataflowEditorWeightMapPaintTool::RequestAction(EDataflowEditorWeightMapPaintToolActions ActionType)
+{
+	if (!bHavePendingAction)
+	{
+		PendingAction = ActionType;
+		bHavePendingAction = true;
+	}
+}
+
+
+void UDataflowEditorWeightMapPaintTool::ApplyAction(EDataflowEditorWeightMapPaintToolActions ActionType)
+{
+	switch (ActionType)
+	{
+	case EDataflowEditorWeightMapPaintToolActions::FloodFillCurrent:
+		FloodFillCurrentWeightAction();
+		break;
+
+	case EDataflowEditorWeightMapPaintToolActions::ClearAll:
+		ClearAllWeightsAction();
+		break;
+	}
+}
+
+
+void UDataflowEditorWeightMapPaintTool::UpdateVertexColorOverlay(const TSet<int>* TrianglesToUpdate)
+{
+	FDynamicMesh3* const Mesh = GetSculptMesh();
+	check(Mesh->HasAttributes());
+	check(Mesh->Attributes()->PrimaryColors());
+	check(ActiveWeightMap);
+
+	FDynamicMeshColorOverlay* const ColorOverlay = Mesh->Attributes()->PrimaryColors();
+
+	auto SetColorsFromWeights = [&](int TriangleID)
+	{
+		const FIndex3i Tri = Mesh->GetTriangle(TriangleID);
+		const FIndex3i ColorElementTri = ColorOverlay->GetTriangle(TriangleID);
+
+		for (int TriVertIndex = 0; TriVertIndex < 3; ++TriVertIndex)
+		{
+			float VertexWeight;
+			ActiveWeightMap->GetValue(Tri[TriVertIndex], &VertexWeight);
+
+			const FVector4f NewColor = FVector4f(VertexWeight, VertexWeight, VertexWeight, 1.0f);
+			ColorOverlay->SetElement(ColorElementTri[TriVertIndex], NewColor);
+		}
+	};
+
+	if (TrianglesToUpdate)
+	{
+		for (const int TriangleID : *TrianglesToUpdate)
+		{
+			SetColorsFromWeights(TriangleID);
+		}
+	}
+	else
+	{
+		for (const int TriangleID : Mesh->TriangleIndicesItr())
+		{
+			SetColorsFromWeights(TriangleID);
+		}
+	}
+}
+
+
+#undef LOCTEXT_NAMESPACE
diff --git a/Engine/Plugins/Experimental/Dataflow/Source/DataflowEditor/Public/Dataflow/DataflowEditorCommands.h b/Engine/Plugins/Experimental/Dataflow/Source/DataflowEditor/Public/Dataflow/DataflowEditorCommands.h
--- a/Engine/Plugins/Experimental/Dataflow/Source/DataflowEditor/Public/Dataflow/DataflowEditorCommands.h
+++ b/Engine/Plugins/Experimental/Dataflow/Source/DataflowEditor/Public/Dataflow/DataflowEditorCommands.h
@@ -36,6 +36,15 @@ public:
 	// TInteractiveToolCommands<>
 	 virtual void GetToolDefaultObjectList(TArray<UInteractiveTool*>& ToolCDOs) override;
 
+	/**
+	* Add or remove commands relevant to Tool to the given UICommandList.
+	* Call this when the active tool changes (eg on ToolManager.OnToolStarted / OnToolEnded)
+	* @param bUnbind if true, commands are removed, otherwise added
+	*/
+	 static void UpdateToolCommandBinding(UInteractiveTool* Tool, TSharedPtr<FUICommandList> UICommandList, bool bUnbind = false);
+
+
+
 	TSharedPtr< FUICommandInfo > EvaluateNode;
 	TSharedPtr< FUICommandInfo > CreateComment;
 	TSharedPtr< FUICommandInfo > ToggleEnabledState;
@@ -48,14 +57,20 @@ public:
 
 	TMap< FName, TSharedPtr<FUICommandInfo> > CreateNodesMap;
 
-	// Temporary, just to test the editor, will be replace by the selection tool
-	const static FString BeginAttributeEditorToolIdentifier;
-	TSharedPtr<FUICommandInfo> BeginAttributeEditorTool;
-
-	const static FString BeginMeshSelectionToolIdentifier;
-	TSharedPtr<FUICommandInfo> BeginMeshSelectionTool;
+	const static FString BeginWeightMapPaintToolIdentifier;
+	TSharedPtr<FUICommandInfo> BeginWeightMapPaintTool;
+	const static FString AddWeightMapNodeIdentifier;
+	TSharedPtr<FUICommandInfo> AddWeightMapNode;
+
+	// @todo(brice) Remove Example Tools
+	//const static FString BeginAttributeEditorToolIdentifier;
+	//TSharedPtr<FUICommandInfo> BeginAttributeEditorTool;
+	//
+	//const static FString BeginMeshSelectionToolIdentifier;
+	//TSharedPtr<FUICommandInfo> BeginMeshSelectionTool;
 };
 
+//@todo(brice) Merge this into the above class
 class DATAFLOWEDITOR_API FDataflowEditorCommands
 {
 public:
diff --git a/Engine/Plugins/Experimental/Dataflow/Source/DataflowEditor/Public/Dataflow/DataflowEditorContent.h b/Engine/Plugins/Experimental/Dataflow/Source/DataflowEditor/Public/Dataflow/DataflowEditorContent.h
new file mode 100644
--- /dev/null
+++ b/Engine/Plugins/Experimental/Dataflow/Source/DataflowEditor/Public/Dataflow/DataflowEditorContent.h
@@ -0,0 +1,150 @@
+// Copyright Epic Games, Inc. All Rights Reserved.
+
+#pragma once
+
+#include "Dataflow/DataflowNodeParameters.h"
+#include "Dataflow/DataflowObjectInterface.h"
+#include "Dataflow/DataflowPatternVertexType.h"
+#include "Dataflow/DataflowEdNode.h"
+#include "DynamicMesh/DynamicMesh3.h"
+#include "GeometryCollection/ManagedArrayCollection.h"
+#include "Templates/SharedPointer.h"
+
+#include "DataflowEditorContent.generated.h"
+
+class FDataflowEditorToolkit;
+class UDataflow;
+class USkeletalMesh;
+class USkeleton;
+class UAnimationAsset;
+class SDataflowGraphEditor;
+struct FManagedArrayCollection;
+
+
+UCLASS()
+class UDataflowEditorContextObject : public UObject
+{
+	GENERATED_BODY()
+public:
+
+	// @todo(brice) : Is this even needed?
+	void SetDataflowGraphEditor(TWeakPtr<SDataflowGraphEditor> InEditor) {DataflowGraphEditor = InEditor;}
+
+	/** Selection Collection Access */
+	void SetPrimarySelectedNode(TObjectPtr<UDataflowEdNode> InSelectedNode) { PrimarySelectedNode = InSelectedNode; }
+	const TObjectPtr<UDataflowEdNode> GetPrimarySelectedNode() const { return PrimarySelectedNode; }
+
+	/** Render Collection used to generate the DynamicMesh3D on the PrimarySelection */
+	void SetPrimaryRenderCollection(TSharedPtr<FManagedArrayCollection> InCollection) { PrimaryRenderCollection = InCollection; }
+	const TSharedPtr<const FManagedArrayCollection> GetPrimaryRenderCollection() const { return PrimaryRenderCollection; }
+
+	/** ViewMode Access */
+	void SetConstructionViewMode(Dataflow::EDataflowPatternVertexType InMode) {ConstructionViewMode = InMode;}
+	Dataflow::EDataflowPatternVertexType GetConstructionViewMode() const { return ConstructionViewMode; }
+
+	/** Get a single selected node of the specified type. Return nullptr if the specified node is not selected, or if multiple nodes are selected*/
+	template<typename NodeType>
+	NodeType* GetPrimarySelectedNodeOfType() const 
+	{
+		if (PrimarySelectedNode && PrimarySelectedNode->GetDataflowNode()) 
+		{
+			return PrimarySelectedNode->GetDataflowNode()->AsType<NodeType>();
+		}
+		return nullptr;
+	}
+
+protected:
+
+	TSharedPtr<FManagedArrayCollection> PrimaryRenderCollection = nullptr;
+	TObjectPtr<UDataflowEdNode> PrimarySelectedNode = nullptr;
+	TWeakPtr<SDataflowGraphEditor> DataflowGraphEditor;
+	Dataflow::EDataflowPatternVertexType ConstructionViewMode = Dataflow::EDataflowPatternVertexType::Sim3D;
+};
+
+
+/** 
+ * Dataflow datas that will be used within the editor classes to evaluate the graph
+ */
+UCLASS()
+class UDataflowEditorContent : public UDataflowEditorContextObject
+{
+	GENERATED_BODY()
+
+public:
+	UDataflowEditorContent();
+
+	/** Check if the datas flow datas are valid */
+	bool IsValid() const { return DataflowOwner && DataflowAsset;}
+
+	/** Check if the datas flow datas are dirty and need to be re-renmdered */
+	bool IsDirty() const { return bIsDirty; }
+	void SetIsDirty(bool InDirty) { bIsDirty = InDirty; }
+
+	//bool HasCompleteSkeletalSetup() const { return SkeletalMesh != nullptr && Skeleton!=nullptr; }
+	void SetSkeletalMesh(TObjectPtr<USkeletalMesh> InMesh) { SkeletalMesh = InMesh; bIsDirty = true; }
+	TObjectPtr<USkeletalMesh>& GetSkeletalMesh() { return SkeletalMesh; }
+	const TObjectPtr<USkeletalMesh>& GetSkeletalMesh() const { return SkeletalMesh; }
+
+	/** Data flow object owner */
+	void SetDataflowOwner(TObjectPtr<UObject> InOwner) { DataflowOwner = InOwner; bIsDirty = true;}
+	TObjectPtr<UObject>& GetDataflowOwner() { return DataflowOwner; }
+	const TObjectPtr<UObject>& GetDataflowOwner() const { return DataflowOwner; }
+
+	/** Data flow asset that we will edit */
+	void SetDataflowAsset(TObjectPtr<UDataflow> InAsset) { DataflowAsset = InAsset;  bIsDirty = true;}
+	TObjectPtr<UDataflow>& GetDataflowAsset() { return DataflowAsset; }
+	const TObjectPtr<UDataflow>& GetDataflowAsset() const { return DataflowAsset; }
+
+	/** Data flow terminal path for evaluation */
+	void SetDataflowTerminal(FString InPath) { DataflowTerminal = InPath;  bIsDirty = true;}
+	FString GetDataflowTerminal() const { return DataflowTerminal; }
+
+	/** Data flow skeleton*/
+	void SetSkeleton(TObjectPtr<USkeleton> InSkeleton) { Skeleton = InSkeleton;  bIsDirty = true;}
+	TObjectPtr<USkeleton>& GetSkeleton() { return Skeleton; }
+	const TObjectPtr<USkeleton>& GetSkeleton() const { return Skeleton; }
+
+	/** Animation asset to be used to preview simulation */
+	void SetAnimationAsset(TObjectPtr<UAnimationAsset> InAnimation) { AnimationAsset = InAnimation; bIsDirty = true;}
+	TObjectPtr<UAnimationAsset>& GetAnimationAsset() { return AnimationAsset; }
+	const TObjectPtr<UAnimationAsset>& GetAnimationAsset() const { return AnimationAsset; }
+
+	/**  Engine context to be used for dataflow evaluation */
+	void SetDataflowContext(TSharedPtr<Dataflow::FEngineContext> InContext) { DataflowContext = InContext; bIsDirty = true;}
+	TSharedPtr<Dataflow::FEngineContext>& GetDataflowContext() { return DataflowContext; }
+	const TSharedPtr<Dataflow::FEngineContext>& GetDataflowContext() const { return DataflowContext; }
+
+	/** Last data flow evaluated node time stamp */
+	void SetLastModifiedTimestamp(Dataflow::FTimestamp InTimestamp) { LastModifiedTimestamp = InTimestamp; bIsDirty = true;}
+	Dataflow::FTimestamp& GetLastModifiedTimestamp() { return LastModifiedTimestamp; }
+	const Dataflow::FTimestamp& GetLastModifiedTimestamp() const { return LastModifiedTimestamp; }
+	
+
+private:
+	/** Data flow object owner */
+	TObjectPtr<UObject> DataflowOwner = nullptr;
+
+	/** Data flow asset that we will edit */
+	TObjectPtr<UDataflow> DataflowAsset = nullptr;
+
+	/** Data flow skeletal mesh*/
+	TObjectPtr<USkeletalMesh> SkeletalMesh = nullptr;
+
+	/** Data flow skeleton*/
+	TObjectPtr<USkeleton> Skeleton = nullptr;
+
+	/** Data flow terminal path for evaluation */
+	FString DataflowTerminal = "";
+
+	/** Animation asset to be used to preview simulation */
+	TObjectPtr<UAnimationAsset> AnimationAsset = nullptr;
+
+	/**  Engine context to be used for dataflow evaluation */
+	TSharedPtr<Dataflow::FEngineContext> DataflowContext = nullptr;
+
+	/** Last data flow evaluated node time stamp */
+	Dataflow::FTimestamp LastModifiedTimestamp = Dataflow::FTimestamp::Invalid;
+
+	bool bIsDirty = true;
+
+};
diff --git a/Engine/Plugins/Experimental/Dataflow/Source/DataflowEditor/Public/Dataflow/DataflowEditorMode.h b/Engine/Plugins/Experimental/Dataflow/Source/DataflowEditor/Public/Dataflow/DataflowEditorMode.h
--- a/Engine/Plugins/Experimental/Dataflow/Source/DataflowEditor/Public/Dataflow/DataflowEditorMode.h
+++ b/Engine/Plugins/Experimental/Dataflow/Source/DataflowEditor/Public/Dataflow/DataflowEditorMode.h
@@ -3,11 +3,19 @@
 #pragma once
 
 #include "BaseCharacterFXEditorMode.h"
+#include "Dataflow/DataflowPatternVertexType.h"
+
 #include "DataflowEditorMode.generated.h"
 
 class UDataflowComponent;
 class FDataflowPreviewScene;
 class SDataflowGraphEditor;
+class FDataflowEditorViewportClient;
+class IDataflowEditorToolBuilder;
+struct FManagedArrayCollection;
+class UEdGraphNode;
+class FViewport;
+
 
 /**
  * The dataflow editor mode is the mode used in the cloth asset editor. It holds most of the inter-tool state.
@@ -24,30 +32,60 @@ public:
 	const static FEditorModeID EM_DataflowEditorModeId;
 
 	UDataflowEditorMode();
-	
+
+	/**
+	* Gets the scene bounding box
+	*/
+	virtual FBox SceneBoundingBox() const override;
+
+	// Bounding box for selected rest space mesh components
+	FBox SelectionBoundingBox() const;
+
 	/**
 	* Gets the tool target requirements for the mode.
 	*/
 	static const FToolTargetTypeRequirements& GetToolTargetRequirements();
 
 	/**
-	* Gets the scene bounding box
+	* Construction View Mode
 	*/
-	virtual FBox SceneBoundingBox() const override;
+	void SetConstructionViewMode(Dataflow::EDataflowPatternVertexType InMode);
+	Dataflow::EDataflowPatternVertexType GetConstructionViewMode() const;
+	bool CanChangeConstructionViewModeTo(Dataflow::EDataflowPatternVertexType NewViewMode) const;
+
+	void ToggleConstructionViewWireframe();
+	bool CanSetConstructionViewWireframeActive() const;
+	bool IsConstructionViewWireframeActive() const
+	{
+		return bConstructionViewWireframe;
+	}
 
 	/** Set the data flow preview scene to store the dynamic meshes generated by the tools */
-	void SetDataflowPreviewScene(FDataflowPreviewScene* DataflowPreviewScene) { PreviewScene = DataflowPreviewScene;}
+	void SetDataflowPreviewScene(FDataflowPreviewScene* DataflowPreviewScene);
+	const FDataflowPreviewScene* GetDataflowPreviewScene() const { return PreviewScene; }
 
 	/** Set the data flow graph editor to create nodes once the tools have ended*/
-	void SetDataflowGraphEditor(const TSharedPtr<SDataflowGraphEditor>& DataflowGraphEditor) { GraphEditor = DataflowGraphEditor;}
+	void SetDataflowGraphEditor(TSharedPtr<SDataflowGraphEditor> DataflowGraphEditor);
+
+	TObjectPtr<UEditorInteractiveToolsContext> GetActiveToolsContext()
+	{
+		return ActiveToolsContext;
+	}
 
 private:
+	friend class FDataflowEditorToolkit;
+	friend class FDataflowEditorModeToolkit;
+
 	// UEdMode
-	virtual void CreateToolkit() override;
-	virtual void BindCommands() override;
+	virtual void Enter() final;
 	virtual void Exit() override;
-	virtual void Enter() override;
 	virtual void ModeTick(float DeltaTime) override;
+	virtual bool ShouldToolStartBeAllowed(const FString& ToolIdentifier) const override;
+	virtual void OnToolStarted(UInteractiveToolManager* Manager, UInteractiveTool* Tool) override;
+	virtual void OnToolEnded(UInteractiveToolManager* Manager, UInteractiveTool* Tool) override;
+	virtual void CreateToolkit() override;
+	virtual void BindCommands() override;
+
 
 	// UBaseCharacterFXEditorMode
 	virtual void AddToolTargetFactories() override;
@@ -55,13 +93,103 @@ private:
 	virtual void CreateToolTargets(const TArray<TObjectPtr<UObject>>& AssetsIn) override;
 	virtual void InitializeTargets(const TArray<TObjectPtr<UObject>>& AssetsIn) override;
 
+	// Use this function to register tools rather than UEdMode::RegisterTool() because we need to specify the ToolsContext
+	void RegisterDataflowTool(TSharedPtr<FUICommandInfo> UICommand,
+							  FString ToolIdentifier,
+							  UInteractiveToolBuilder* Builder,
+							  const IDataflowEditorToolBuilder* DataflowToolBuilder,
+							  UEditorInteractiveToolsContext* const ToolsContext,
+							  EToolsContextScope ToolScope = EToolsContextScope::Default);
+
+	void RegisterAddNodeCommand(TSharedPtr<FUICommandInfo> AddNodeCommand, const FName& NewNodeType, TSharedPtr<FUICommandInfo> StartToolCommand);
+
+	void SetPreviewScene(FDataflowPreviewScene* PreviewScene);
+
+	void SetRestSpaceViewportClient(TWeakPtr<FDataflowEditorViewportClient, ESPMode::ThreadSafe> ViewportClient);
+	void RefocusRestSpaceViewportClient();
+	void FirstTimeFocusRestSpaceViewport();
+
+	// @todo(brice) : Can this be in the content?
+	// intended to be called by the toolkit when selected node in the Dataflow graph changes
+	void SetSelectedCollection(TSharedPtr<FManagedArrayCollection> Collection);
+	bool IsComponentSelected(const UPrimitiveComponent* InComponent);
+	TSharedPtr<FManagedArrayCollection> GetSelectedCollection();
+	TSharedPtr<FManagedArrayCollection> SelectedCollection = nullptr;
+
+	void StartToolForSelectedNode(const UObject* SelectedNode);
+	void OnDataflowNodeDeleted(const TSet<UObject*>& DeletedNodes);
+
+	/**
+	* Return the single selected node in the Dataflow Graph Editor only if it has an output of the specified type
+	* If there is not a single node selected, or if it does not have the specified output, return null
+	*/
+	UEdGraphNode* GetSingleSelectedNodeWithOutputType(const FName& SelectedNodeOutputTypeName) const;
+
+	/**
+	 * Create a node with the specified type in the graph
+	*/
+	UEdGraphNode* CreateNewNode(const FName& NewNodeTypeName);
+
+	/** Create a node with the specified type, then connect it to the output of the specified UpstreamNode.
+	* If the specified output of the upstream node is already connected to another node downstream, we first break
+	* that connecttion, then insert the new node along the previous connection.
+	* We want to turn this:
+	*
+	* [UpstreamNode] ----> [DownstreamNode(s)]
+	*
+	* to this:
+	*
+	* [UpstreamNode] ----> [NewNode] ----> [DownstreamNode(s)]
+	*
+	*
+	* @param NewNodeTypeName The type of node to create, by name
+	* @param UpstreamNode Node to connect the new node to
+	* @param ConnectionTypeName The type of output of the upstream node to connect our new node to
+	* @param NewNodeConnectionName The name of the input/output connection on our new node that will be connected
+	* @return The newly-created node
+	*/
+	UEdGraphNode* CreateAndConnectNewNode(const FName& NewNodeTypeName, UEdGraphNode& UpstreamNode, const FName& ConnectionTypeName, const FName& NewNodeConnectionName);
+
+	void InitializeContextObject();
+	void DeleteContextObject();
+
 	/** Dataflow preview scene from the toolkit */
 	FDataflowPreviewScene* PreviewScene = nullptr;
 
 	/** Correspondence between node types and commands to launch tools */
-	TMap<FName, TSharedPtr<const FUICommandInfo>> NodeToolMap;
+	TMap<FName, TSharedPtr<const FUICommandInfo>> NodeTypeToToolCommandMap;
 
 	/** Graph editor is required to add nodes once some tools are complete */
-	TWeakPtr<SDataflowGraphEditor> GraphEditor;
+	TWeakPtr<SDataflowGraphEditor> DataflowGraphEditor;
+
+	Dataflow::EDataflowPatternVertexType ConstructionViewMode = Dataflow::EDataflowPatternVertexType::Sim3D;
+
+	// Whether we should restore the previous view mode when a tool ends
+	bool bShouldRestoreSavedConstructionViewMode = false;
+	// The Construction view mode that was active before starting the current tool. When the tool ends, restore this view mode if bShouldRestoreSavedConstructionViewMode is true.
+	Dataflow::EDataflowPatternVertexType SavedConstructionViewMode;
+
+	bool bConstructionViewWireframe = false;
+	bool bShouldRestoreConstructionViewWireframe = false;
+
+	UPROPERTY()
+	TObjectPtr<UEditorInteractiveToolsContext> ActiveToolsContext = nullptr;
+
+	TWeakPtr<FDataflowEditorViewportClient, ESPMode::ThreadSafe> ConstructionViewportClient;
+
+	// The first time we get a valid mesh, refocus the camera on it
+	bool bFirstValid2DMesh = true;
+	bool bFirstValid3DMesh = true;
+
+	// Whether the rest space viewport should focus on the rest space mesh on the next tick
+	bool bShouldFocusRestSpaceView = true;
+
+	//@todo(brice) : Move to Construction
+	void RestSpaceViewportResized(FViewport* RestspaceViewport, uint32 Unused);
+
+	// Dataflow node type whose corresponding tool should be started on the next Tick
+	FName NodeTypeForPendingToolStart;
+
+
 };
 
diff --git a/Engine/Plugins/Experimental/Dataflow/Source/DataflowEditor/Public/Dataflow/DataflowEditorModeToolkit.h b/Engine/Plugins/Experimental/Dataflow/Source/DataflowEditor/Public/Dataflow/DataflowEditorModeToolkit.h
--- a/Engine/Plugins/Experimental/Dataflow/Source/DataflowEditor/Public/Dataflow/DataflowEditorModeToolkit.h
+++ b/Engine/Plugins/Experimental/Dataflow/Source/DataflowEditor/Public/Dataflow/DataflowEditorModeToolkit.h
@@ -4,6 +4,10 @@
 
 #include "BaseCharacterFXEditorModeToolkit.h"
 
+class UEditorInteractiveToolsContext;
+class SDataflowEditorViewport;
+class SBaseCharacterFXEditorViewport;
+
 /**
  * The dataflow editor mode toolkit is responsible for the panel on the side in the dataflow editor
  * that shows mode and tool properties. Tool buttons would go in Init().
@@ -14,15 +18,34 @@
 class DATAFLOWEDITOR_API FDataflowEditorModeToolkit : public FBaseCharacterFXEditorModeToolkit
 {
 public:
+	void Init(const TSharedPtr<IToolkitHost>& InitToolkitHost, TWeakObjectPtr<UEdMode> InOwningMode);
+
+	// Add buttons to the toolbar with the specified name
+	// Note: Most FModeToolkits would define BuildToolPalette, but we are putting buttons in the top toolbar instead
+	void BuildEditorToolBar(const FName& EditorToolBarName);
+
+	// FBaseCharacterFXEditorModeToolkit interface
+	virtual const FSlateBrush* GetActiveToolIcon(const FString& Identifier) const override;
 
 	// FModeToolkit interface
 	virtual void OnToolStarted(UInteractiveToolManager* Manager, UInteractiveTool* Tool) override;
-	virtual void BuildToolPalette(FName PaletteName, class FToolBarBuilder& ToolbarBuilder) override;
+	virtual void OnToolEnded(UInteractiveToolManager* Manager, UInteractiveTool* Tool) override;
 
 	// IToolkit interface
 	virtual FName GetToolkitFName() const override;
 	virtual FText GetBaseToolkitName() const override;
 
-	// FBaseCharacterFXEditorModeToolkit interface
-	virtual const FSlateBrush* GetActiveToolIcon(const FString& Identifier) const override;
+	void SetRestSpaceViewportWidget(TWeakPtr<SDataflowEditorViewport>);
+	//void SetPreviewViewportWidget(TWeakPtr<SChaosClothAssetEditor3DViewport>);
+
+private:
+
+	// Get the viewport widget associated with the given manager
+	// TODO: This should not be necessary any more as we do not run tools in the Cloth Preview Viewport (JIRA UE-201248)
+	SBaseCharacterFXEditorViewport* GetViewportWidgetForManager(UInteractiveToolManager* Manager);
+
+	UEditorInteractiveToolsContext* GetCurrentToolsContext();
+
+	TWeakPtr<SDataflowEditorViewport> RestSpaceViewportWidget;
+	//TWeakPtr<SChaosClothAssetEditor3DViewport> PreviewViewportWidget;
 };
diff --git a/Engine/Plugins/Experimental/Dataflow/Source/DataflowEditor/Public/Dataflow/DataflowEditorModeUILayer.h b/Engine/Plugins/Experimental/Dataflow/Source/DataflowEditor/Public/Dataflow/DataflowEditorModeUILayer.h
new file mode 100644
--- /dev/null
+++ b/Engine/Plugins/Experimental/Dataflow/Source/DataflowEditor/Public/Dataflow/DataflowEditorModeUILayer.h
@@ -0,0 +1,30 @@
+// Copyright Epic Games, Inc. All Rights Reserved.
+
+#pragma once
+
+#include "BaseCharacterFXEditorModeUILayer.h"
+#include "DataflowEditorModeUILayer.generated.h"
+
+/** Interchange layer to manage built in tab locations within the editor's layout. **/
+UCLASS()
+class DATAFLOWEDITOR_API UDataflowEditorUISubsystem : public UAssetEditorUISubsystem
+{
+	GENERATED_BODY()
+public:
+
+	virtual void Initialize(FSubsystemCollectionBase& Collection) override;
+	virtual void Deinitialize() override;
+	virtual void RegisterLayoutExtensions(FLayoutExtender& Extender) override;
+
+	// Identifier for the Layout extension in DataflowEditorToolkit's StandaloneDefaultLayout
+	static const FName EditorSidePanelAreaName;
+
+};
+
+/** Handles the hosting of additional toolkits, such as the mode toolkit, within the DataflowEditor's toolkit. **/
+class FDataflowEditorModeUILayer : public FBaseCharacterFXEditorModeUILayer
+{
+public:
+	FDataflowEditorModeUILayer(const IToolkitHost* InToolkitHost) : FBaseCharacterFXEditorModeUILayer(InToolkitHost) {}
+};
+
diff --git a/Engine/Plugins/Experimental/Dataflow/Source/DataflowEditor/Public/Dataflow/DataflowEditorToolkit.h b/Engine/Plugins/Experimental/Dataflow/Source/DataflowEditor/Public/Dataflow/DataflowEditorToolkit.h
--- a/Engine/Plugins/Experimental/Dataflow/Source/DataflowEditor/Public/Dataflow/DataflowEditorToolkit.h
+++ b/Engine/Plugins/Experimental/Dataflow/Source/DataflowEditor/Public/Dataflow/DataflowEditorToolkit.h
@@ -23,6 +23,7 @@ class UDataflow;
 class USkeletalMesh;
 class SDataflowGraphEditor;
 class UDataflowEditorContent;
+class FDataflowPreviewScene;
 
 namespace Dataflow
 {
@@ -74,6 +75,7 @@ protected:
 	void OnNodeTitleCommitted(const FText& InNewText, ETextCommit::Type InCommitType, UEdGraphNode* GraphNode);
 	void OnNodeSelectionChanged(const TSet<UObject*>& NewSelection);
 	void OnNodeDeleted(const TSet<UObject*>& NewSelection);
+	void OnNodeSingleClicked(UObject* ClickedNode) const;
 	void OnAssetPropertyValueChanged(const FPropertyChangedEvent& PropertyChangedEvent);
 	
 	// Callback to remove the closed one from the listener views
@@ -156,6 +158,12 @@ private:
     FDelegateHandle OnFinishedChangingPropertiesDelegateHandle;
     FDelegateHandle OnFinishedChangingAssetPropertiesDelegateHandle;
 
-	// Previous node selection to compare values in delegates
-	TSet<UObject*> PrevNodeSelection;
+	// The currently selected set of dataflow nodes. 
+	TSet<UObject*> SelectedDataflowNodes;
+
+	// The most recently selected dataflow node. (Top
+	TObjectPtr<UDataflowEdNode> PrimarySelection;
+
+	/** Scene in which the 3D sim space preview meshes live. Ownership shared with AdvancedPreviewSettingsWidget*/
+	TSharedPtr<FDataflowPreviewScene> DataflowPreviewScene;
 };
diff --git a/Engine/Plugins/Experimental/Dataflow/Source/DataflowEditor/Public/Dataflow/DataflowEditorViewportClient.h b/Engine/Plugins/Experimental/Dataflow/Source/DataflowEditor/Public/Dataflow/DataflowEditorViewportClient.h
--- a/Engine/Plugins/Experimental/Dataflow/Source/DataflowEditor/Public/Dataflow/DataflowEditorViewportClient.h
+++ b/Engine/Plugins/Experimental/Dataflow/Source/DataflowEditor/Public/Dataflow/DataflowEditorViewportClient.h
@@ -5,6 +5,8 @@
 #include "EditorViewportClient.h"
 #include "Dataflow/DataflowNodeParameters.h"
 #include "Dataflow/DataflowComponentSelectionState.h"
+#include "Dataflow/DataflowPatternVertexType.h"
+#include "InputBehaviorSet.h"
 
 class FDataflowEditorToolkit;
 class ADataflowActor;
@@ -12,31 +14,48 @@ class UTransformProxy;
 class UCombinedTransformGizmo;
 class FTransformGizmoDataBinder;
 class FDataflowPreviewScene;
+class UInputBehaviorSet;
 
-class DATAFLOWEDITOR_API FDataflowEditorViewportClient : public FEditorViewportClient
+class DATAFLOWEDITOR_API FDataflowEditorViewportClient : public FEditorViewportClient //, public IInputBehaviorSource
 {
+//@todo(brice) : Add BehaviorUI support
+
 public:
 	using Super = FEditorViewportClient;
 
 	FDataflowEditorViewportClient(FEditorModeTools* InModeTools, FPreviewScene* InPreviewScene,
 								  const TWeakPtr<SEditorViewport> InEditorViewportWidget = nullptr);
 
-	// FEditorViewportClient interface
-	virtual void Tick(float DeltaSeconds) override;
+	void SetConstructionViewMode(Dataflow::EDataflowPatternVertexType InViewMode);
+	Dataflow::EDataflowPatternVertexType GetConstructionViewMode() const;
+
+	// IInputBehaviorSource
+	// virtual const UInputBehaviorSet* GetInputBehaviors() const override;
 
 	/** Set the data flow toolkit used to create the client*/
 	void SetDataflowEditorToolkit(TWeakPtr<FDataflowEditorToolkit> DataflowToolkit);
 	
+	void SetToolCommandList(TWeakPtr<FUICommandList> ToolCommandList);
+
 	// FGCObject Interface
 	virtual void AddReferencedObjects(FReferenceCollector& Collector) override;
 	virtual FString GetReferencerName() const override { return TEXT("FDataflowEditorViewportClient"); }
 
 private:
 
+	virtual void Tick(float DeltaSeconds) override;
+
 	/** Toolkit used to create the viewport client */
 	TWeakPtr<FDataflowEditorToolkit> DataflowEditorToolkitPtr = nullptr;
 
 	/** Dataflow preview scene from the toolkit */
 	FDataflowPreviewScene* PreviewScene = nullptr;
 
+	// @todo(brice) : Is this needed?
+	TWeakPtr<FUICommandList> ToolCommandList;
+
+	Dataflow::EDataflowPatternVertexType ConstructionViewMode = Dataflow::EDataflowPatternVertexType::Sim3D;
+
+	TObjectPtr<UInputBehaviorSet> BehaviorSet;
+
 };
diff --git a/Engine/Plugins/Experimental/Dataflow/Source/DataflowEditor/Public/Dataflow/DataflowPreviewScene.h b/Engine/Plugins/Experimental/Dataflow/Source/DataflowEditor/Public/Dataflow/DataflowPreviewScene.h
--- a/Engine/Plugins/Experimental/Dataflow/Source/DataflowEditor/Public/Dataflow/DataflowPreviewScene.h
+++ b/Engine/Plugins/Experimental/Dataflow/Source/DataflowEditor/Public/Dataflow/DataflowPreviewScene.h
@@ -13,6 +13,8 @@ class UAnimSingleNodeInstance;
 class UMaterialInterface;
 class FAssetEditorModeManager;
 class ADataflowActor;
+class UMeshElementsVisualizer;
+class UInteractiveToolPropertySet;
 
 namespace UE
 {
@@ -35,41 +37,109 @@ public:
 	virtual ~FDataflowPreviewScene();
 
 
+	//
+	// Functional API
+	//
+public:
+	/** Update the state of the scene.*/
 	void Update();
 
-	/** Update Skeletal Mesh Component */
-	void UpdateSkeletalMeshComponent();
+	/** Tick */
+	// @todo(brice) : depricate
+	virtual void Tick(float DeltaTime) override;
 
-	/** FGCObject interface */
-	virtual void AddReferencedObjects(FReferenceCollector& Collector) override;
+	/** Cleam up the scene on exit. */
+	void Exit();
 
-	/** Modify/Get the data flow datas */
-	TObjectPtr<UDataflowEditorContent> GetDataflowEditorContent() { return EditorContent; }
-	const TObjectPtr<UDataflowEditorContent> GetDataflowEditorContent() const { return EditorContent; }
 
-	/** Modify/Get the data flow mode manager */
-	TSharedPtr<FAssetEditorModeManager>& GetDataflowModeManager() { return DataflowModeManager; }
-	const TSharedPtr<FAssetEditorModeManager>& GetDataflowModeManager() const { return DataflowModeManager; }
-	
-	/** Modify/Get the preview anim instance */
-	UAnimSingleNodeInstance* GetPreviewAnimInstance() { return PreviewAnimInstance; }
-	const UAnimSingleNodeInstance* const GetPreviewAnimInstance() const { return PreviewAnimInstance; }
+	//
+	// SkeletalMesh Rendering
+	//
+public:
+	/** Update Skeletal Mesh Component */
+	void UpdateSkeletalMeshComponent();
 
 	/** Modify/Get the skeletal mesh component */
 	USkeletalMeshComponent* GetSkeletalMeshComponent() { return SkeletalMeshComponent; }
 	const USkeletalMeshComponent* GetSkeletalMeshComponent() const { return SkeletalMeshComponent; }
 
+
+
+
+	//
+	// ManagedArrayCollectionRendering
+	//
+public:
+	void UpdateDynamicMeshComponents();
+
+	/** Modify/Get the skeletal mesh component */
+	TArray<TObjectPtr<UDynamicMeshComponent>> GetDynamicMeshComponents() { return DynamicMeshComponents; }
+	const TArray<TObjectPtr<UDynamicMeshComponent>> GetDynamicMeshComponents() const { return DynamicMeshComponents; }
+
+	/** Trigger a Update when the render geometry has changed.*/
+	//@todo(brice) : deprecate
+	void ReinitializeDynamicMeshComponents();
+
+private:
 	/** Reset all the dynamic mesh components */
 	void ResetDynamicMeshComponents();
 	
 	/** Add a dynamic mesh component to the scene */
 	TObjectPtr<UDynamicMeshComponent>& AddDynamicMeshComponent(UE::Geometry::FDynamicMesh3&& DynamicMesh, const TArray<UMaterialInterface*>& MaterialSet);
 
-	TArray<TObjectPtr<UDynamicMeshComponent>> GetDynamicMeshComponents() {return DynamicMeshComponents;}
+
+
+
+	//
+	// Wireframe
+	//
+public:
+	void UpdateWireframeMeshElementsVisualizer();
+
+private:
+	void ResetWireframeMeshElementsVisualizer();
+
+	void AddWireframeMeshElementsVisualizer();
+
+
+
+
+
+	//
+	// Visibility
+	//
+public:
+	/** Is there ANYTHING to render */
+	bool HasRenderableGeometry();
 
 	/** Build the scene bounding box */
 	FBox GetBoundingBox() const;
-	
+
+	/** Build the selected objects bounding box */
+	FBox SelectedComponentBounds() const;
+
+
+	//
+	// General Access
+	//
+public:
+
+	/** Modify/Get the data flow datas */
+	TObjectPtr<UDataflowEditorContent> GetDataflowEditorContent() { return EditorContent; }
+	const TObjectPtr<UDataflowEditorContent> GetDataflowEditorContent() const { return EditorContent; }
+
+	/** Modify/Get the data flow mode manager */
+	TSharedPtr<FAssetEditorModeManager>& GetDataflowModeManager() { return DataflowModeManager; }
+	const TSharedPtr<FAssetEditorModeManager>& GetDataflowModeManager() const { return DataflowModeManager; }
+
+	/** Modify/Get the preview anim instance */
+	UAnimSingleNodeInstance* GetPreviewAnimInstance() { return PreviewAnimInstance; }
+	const UAnimSingleNodeInstance* const GetPreviewAnimInstance() const { return PreviewAnimInstance; }
+
+	/** FGCObject interface */
+	virtual void AddReferencedObjects(FReferenceCollector& Collector) override;
+
+
 private:
 	
 	/** Check if a primitive component is selected */
@@ -95,5 +165,16 @@ private:
 
 	/** Mode Manager for selection */
 	TSharedPtr<FAssetEditorModeManager> DataflowModeManager;
+
+	// Property objects (visible or not) that get ticked. [ticked]
+	TArray<TObjectPtr<UInteractiveToolPropertySet>> PropertyObjectsToTick;
+
+	// Rest space wireframe. They have to get ticked to be able to respond to setting changes. [ticked]
+	TObjectPtr<UMeshElementsVisualizer> WireframeDraw = nullptr;
+
+	// Show the wireframe on render. 
+	bool bConstructionViewWireframe = true;
+
+	
 };
 
diff --git a/Engine/Plugins/Experimental/Dataflow/Source/DataflowEditor/Public/DataflowEditorTools/DataflowEditorToolBuilder.h b/Engine/Plugins/Experimental/Dataflow/Source/DataflowEditor/Public/DataflowEditorTools/DataflowEditorToolBuilder.h
new file mode 100644
--- /dev/null
+++ b/Engine/Plugins/Experimental/Dataflow/Source/DataflowEditor/Public/DataflowEditorTools/DataflowEditorToolBuilder.h
@@ -0,0 +1,35 @@
+// Copyright Epic Games, Inc. All Rights Reserved.
+
+#pragma once
+
+#include "Containers/Array.h"
+#include "UObject/Interface.h"
+
+#include "DataflowEditorToolBuilder.generated.h"
+
+namespace Dataflow
+{
+	enum class EDataflowPatternVertexType : uint8;
+}
+
+
+UINTERFACE(MinimalAPI)
+class UDataflowEditorToolBuilder : public UInterface
+{
+	GENERATED_BODY()
+};
+
+
+class IDataflowEditorToolBuilder
+{
+	GENERATED_BODY()
+
+public:
+
+	/** Returns all Construction View modes that this tool can operate in. The first element should be the preferred mode to switch to if necessary. */
+	virtual void GetSupportedViewModes(TArray<Dataflow::EDataflowPatternVertexType>& Modes) const = 0;
+
+	/** Returns whether or not view can be set to wireframe when this tool is active.. */
+	virtual bool CanSetConstructionViewWireframeActive() const { return true; }
+};
+
diff --git a/Engine/Plugins/Experimental/Dataflow/Source/DataflowEditor/Public/DataflowEditorTools/DataflowEditorWeightMapPaintTool.h b/Engine/Plugins/Experimental/Dataflow/Source/DataflowEditor/Public/DataflowEditorTools/DataflowEditorWeightMapPaintTool.h
new file mode 100644
--- /dev/null
+++ b/Engine/Plugins/Experimental/Dataflow/Source/DataflowEditor/Public/DataflowEditorTools/DataflowEditorWeightMapPaintTool.h
@@ -0,0 +1,491 @@
+// Copyright Epic Games, Inc. All Rights Reserved.
+
+#pragma once
+
+#include "CoreMinimal.h"
+#include "HAL/Platform.h"
+#include "UObject/NoExportTypes.h"
+#include "BaseTools/BaseBrushTool.h"
+#include "BaseTools/MeshSurfacePointMeshEditingTool.h"
+#include "Components/DynamicMeshComponent.h"
+#include "PropertySets/WeightMapSetProperties.h"
+#include "Mechanics/PolyLassoMarqueeMechanic.h"
+
+#include "Sculpting/MeshSculptToolBase.h"
+#include "Sculpting/MeshBrushOpBase.h"
+
+#include "DynamicMesh/DynamicMeshAABBTree3.h"
+#include "DynamicMesh/DynamicMeshOctree3.h"
+#include "DynamicMesh/DynamicMeshAttributeSet.h"
+#include "DynamicMesh/DynamicMeshChangeTracker.h"
+#include "DynamicMesh/MeshNormals.h"
+#include "TransformTypes.h"
+#include "ToolDataVisualizer.h"
+#include "GroupTopology.h"
+#include "DataflowEditorTools/DataflowEditorToolBuilder.h"
+
+#include "DataflowEditorWeightMapPaintTool.generated.h"
+
+class UMeshElementsVisualizer;
+class UDataflowWeightMapEraseBrushOpProps;
+class UDataflowWeightMapPaintBrushOpProps;
+class UDataflowWeightMapSmoothBrushOpProps;
+class UPolygonSelectionMechanic;
+class UDataflowEditorContextObject;
+struct FDataflowCollectionAddScalarVertexPropertyNode;
+
+DECLARE_STATS_GROUP(TEXT("WeightMapPaintTool"), STATGROUP_WeightMapPaintTool, STATCAT_Advanced);
+DECLARE_CYCLE_STAT(TEXT("WeightMapPaintTool_UpdateROI"), WeightMapPaintTool_UpdateROI, STATGROUP_WeightMapPaintTool);
+DECLARE_CYCLE_STAT(TEXT("WeightMapPaintTool_ApplyStamp"), WeightMapPaintToolApplyStamp, STATGROUP_WeightMapPaintTool );
+DECLARE_CYCLE_STAT(TEXT("WeightMapPaintTool_Tick"), WeightMapPaintToolTick, STATGROUP_WeightMapPaintTool);
+DECLARE_CYCLE_STAT(TEXT("WeightMapPaintTool_Tick_ApplyStampBlock"), WeightMapPaintTool_Tick_ApplyStampBlock, STATGROUP_WeightMapPaintTool);
+DECLARE_CYCLE_STAT(TEXT("WeightMapPaintTool_Tick_ApplyStamp_Remove"), WeightMapPaintTool_Tick_ApplyStamp_Remove, STATGROUP_WeightMapPaintTool);
+DECLARE_CYCLE_STAT(TEXT("WeightMapPaintTool_Tick_ApplyStamp_Insert"), WeightMapPaintTool_Tick_ApplyStamp_Insert, STATGROUP_WeightMapPaintTool);
+DECLARE_CYCLE_STAT(TEXT("WeightMapPaintTool_Tick_NormalsBlock"), WeightMapPaintTool_Tick_NormalsBlock, STATGROUP_WeightMapPaintTool);
+DECLARE_CYCLE_STAT(TEXT("WeightMapPaintTool_Tick_UpdateMeshBlock"), WeightMapPaintTool_Tick_UpdateMeshBlock, STATGROUP_WeightMapPaintTool);
+DECLARE_CYCLE_STAT(TEXT("WeightMapPaintTool_Tick_UpdateTargetBlock"), WeightMapPaintTool_Tick_UpdateTargetBlock, STATGROUP_WeightMapPaintTool);
+DECLARE_CYCLE_STAT(TEXT("WeightMapPaintTool_Normals_Collect"), WeightMapPaintTool_Normals_Collect, STATGROUP_WeightMapPaintTool);
+DECLARE_CYCLE_STAT(TEXT("WeightMapPaintTool_Normals_Compute"), WeightMapPaintTool_Normals_Compute, STATGROUP_WeightMapPaintTool);
+
+
+
+
+
+/**
+ * Tool Builder
+ */
+UCLASS()
+class DATAFLOWEDITOR_API UDataflowEditorWeightMapPaintToolBuilder : public UMeshSurfacePointMeshEditingToolBuilder, public IDataflowEditorToolBuilder
+{
+	GENERATED_BODY()
+
+private:
+	virtual void GetSupportedViewModes(TArray<Dataflow::EDataflowPatternVertexType>& Modes) const override;
+	virtual UMeshSurfacePointTool* CreateNewTool(const FToolBuilderState& SceneState) const override;
+	virtual bool CanSetConstructionViewWireframeActive() const { return false; }
+};
+
+
+
+
+/** Mesh Sculpting Brush Types */
+UENUM()
+enum class EDataflowEditorWeightMapPaintInteractionType : uint8
+{
+	Brush,
+	Fill,
+	PolyLasso,
+	Gradient,
+
+	LastValue UMETA(Hidden)
+};
+
+
+
+
+
+
+
+/** Mesh Sculpting Brush Types */
+UENUM()
+enum class EDataflowEditorWeightMapPaintBrushType : uint8
+{
+	/** Paint weights */
+	Paint UMETA(DisplayName = "Paint"),
+
+	/** Smooth existing weights */
+	Smooth UMETA(DisplayName = "Smooth"),
+
+	/** Erase weights */
+	Erase UMETA(Hidden, DisplayName = "Erase"),
+
+	LastValue UMETA(Hidden)
+};
+
+
+/** Mesh Sculpting Brush Area Types */
+UENUM()
+enum class EDataflowEditorWeightMapPaintBrushAreaType : uint8
+{
+	Connected,
+	Volumetric
+};
+
+/** Mesh Sculpting Brush Types */
+UENUM()
+enum class EDataflowEditorWeightMapPaintVisibilityType : uint8
+{
+	None,
+	Unoccluded
+};
+
+
+
+// TODO: Look at EditConditions for all these properties. Which ones make sense for which SubToolType?
+
+UCLASS()
+class DATAFLOWEDITOR_API UDataflowEditorWeightMapPaintBrushFilterProperties : public UInteractiveToolPropertySet
+{
+	GENERATED_BODY()
+
+public:
+	
+	UPROPERTY(EditAnywhere, Category = ActionType, meta = (DisplayName = "Action"))
+	EDataflowEditorWeightMapPaintInteractionType SubToolType = EDataflowEditorWeightMapPaintInteractionType::Brush;
+
+	UPROPERTY(EditAnywhere, Category = ActionType, meta = (DisplayName = "Brush Mode", 
+		HideEditConditionToggle, EditConditionHides, EditCondition = "SubToolType == EDataflowEditorWeightMapPaintInteractionType::Brush"))
+	EDataflowEditorWeightMapPaintBrushType PrimaryBrushType = EDataflowEditorWeightMapPaintBrushType::Paint;
+
+	/** Relative size of brush */
+	UPROPERTY(EditAnywhere, Category = Brush, meta = (DisplayName = "Brush Size", UIMin = "0.0", UIMax = "1.0", ClampMin = "0.0", ClampMax = "10.0", 
+		HideEditConditionToggle, EditConditionHides, EditCondition = "SubToolType == EDataflowEditorWeightMapPaintInteractionType::Brush"))
+	float BrushSize = 0.25f;
+
+	/** The new value to paint on the mesh */
+	UPROPERTY(EditAnywhere, Category = Brush, meta = (UIMin = 0, ClampMin = 0, UIMax = 1, ClampMax = 1,
+		HideEditConditionToggle, EditConditionHides, EditCondition = 
+		"(SubToolType == EDataflowEditorWeightMapPaintInteractionType::Brush && PrimaryBrushType == EDataflowEditorWeightMapPaintBrushType::Paint) || SubToolType == EDataflowEditorWeightMapPaintInteractionType::Fill || SubToolType == EDataflowEditorWeightMapPaintInteractionType::PolyLasso"))
+	double AttributeValue = 1;
+
+	/** How quickly each brush stroke will drive mesh values towards the desired value */
+	UPROPERTY(EditAnywhere, Category = Brush, meta = (UIMin = 0, ClampMin = 0, UIMax = 1, ClampMax = 1,
+		HideEditConditionToggle, EditConditionHides, EditCondition = 
+		"SubToolType == EDataflowEditorWeightMapPaintInteractionType::Brush || SubToolType == EDataflowEditorWeightMapPaintInteractionType::Fill"))
+	double Strength = 0.5;
+
+	/** The Gradient upper limit value */
+	UPROPERTY(EditAnywhere, Category = Gradient, meta = (UIMin = 0, ClampMin = 0, UIMax = 1, ClampMax = 1,
+		HideEditConditionToggle, EditConditionHides, EditCondition = "SubToolType == EDataflowEditorWeightMapPaintInteractionType::Gradient"))
+	double GradientHighValue = 1.0;
+
+	/** The Gradient lower limit value */
+	UPROPERTY(EditAnywhere, Category = Gradient, meta = (UIMin = 0, ClampMin = 0, UIMax = 1, ClampMax = 1,
+		HideEditConditionToggle, EditConditionHides, EditCondition = "SubToolType == EDataflowEditorWeightMapPaintInteractionType::Gradient"))
+	double GradientLowValue = 0.0;
+
+
+	/** The Region affected by the current operation will be bounded by edge angles larger than this threshold */
+	UPROPERTY(EditAnywhere, Category = Filters, meta = (UIMin = "0.0", ClampMin = "0.0", UIMax = "180.0", ClampMax = "180.0",
+		HideEditConditionToggle, EditConditionHides, EditCondition = 
+		"SubToolType == EDataflowEditorWeightMapPaintInteractionType::Brush || SubToolType == EDataflowEditorWeightMapPaintInteractionType::Fill"))
+	float AngleThreshold = 180.0f;
+
+	/** The Region affected by the current operation will be bounded by UV borders/seams */
+	UPROPERTY(EditAnywhere, Category = Filters, meta = (HideEditConditionToggle, EditConditionHides, EditCondition = 
+		"SubToolType == EDataflowEditorWeightMapPaintInteractionType::Brush || SubToolType == EDataflowEditorWeightMapPaintInteractionType::Fill"))
+	bool bUVSeams = false;
+
+	/** The Region affected by the current operation will be bounded by Hard Normal edges/seams */
+	UPROPERTY(EditAnywhere, Category = Filters, meta = (HideEditConditionToggle, EditConditionHides, EditCondition = 
+		"SubToolType == EDataflowEditorWeightMapPaintInteractionType::Brush || SubToolType == EDataflowEditorWeightMapPaintInteractionType::Fill"))
+	bool bNormalSeams = false;
+
+	/** Control which triangles can be affected by the current operation based on visibility. Applied after all other filters. */
+	UPROPERTY(EditAnywhere, Category = Filters, meta = (HideEditConditionToggle, EditConditionHides, EditCondition =
+		"SubToolType == EDataflowEditorWeightMapPaintInteractionType::Brush || SubToolType == EDataflowEditorWeightMapPaintInteractionType::Fill || SubToolType == EDataflowEditorWeightMapPaintInteractionType::PolyLasso"))
+	EDataflowEditorWeightMapPaintVisibilityType VisibilityFilter = EDataflowEditorWeightMapPaintVisibilityType::None;
+
+	/** The weight value at the brush indicator */
+	UPROPERTY(VisibleAnywhere, Transient, Category = Query, meta = (NoResetToDefault, 
+		HideEditConditionToggle, EditConditionHides, EditCondition = "SubToolType == EDataflowEditorWeightMapPaintInteractionType::Brush || SubToolType == EDataflowEditorWeightMapPaintInteractionType::Fill"))
+	double ValueAtBrush = 0;
+
+};
+
+
+
+
+
+UENUM()
+enum class EDataflowEditorWeightMapPaintToolActions
+{
+	NoAction,
+
+	FloodFillCurrent,
+	ClearAll,
+};
+
+
+UCLASS()
+class DATAFLOWEDITOR_API UDataflowEditorMeshWeightMapPaintToolActions : public UInteractiveToolPropertySet
+{
+	GENERATED_BODY()
+
+public:
+
+	TWeakObjectPtr<UDataflowEditorWeightMapPaintTool> ParentTool;
+
+	void Initialize(UDataflowEditorWeightMapPaintTool* ParentToolIn) { ParentTool = ParentToolIn; }
+
+	void PostAction(EDataflowEditorWeightMapPaintToolActions Action);
+
+	UFUNCTION(CallInEditor, Category = Operations, meta = (DisplayPriority = 10))
+	void ClearAll()
+	{
+		PostAction(EDataflowEditorWeightMapPaintToolActions::ClearAll);
+	}
+
+	UFUNCTION(CallInEditor, Category = Operations, meta = (DisplayPriority = 12))
+	void FloodFillCurrent()
+	{
+		PostAction(EDataflowEditorWeightMapPaintToolActions::FloodFillCurrent);
+	}
+
+};
+
+UCLASS()
+class DATAFLOWEDITOR_API UDataflowEditorUpdateWeightMapProperties : public UInteractiveToolPropertySet
+{
+	GENERATED_BODY()
+
+public:
+
+	UPROPERTY(EditAnywhere, Category = UpdateNode, meta = (DisplayName = "Name"))
+	FString Name;
+
+private:
+
+	virtual void PostEditChangeProperty(struct FPropertyChangedEvent& PropertyChangedEvent) override;
+};
+
+/**
+ * Mesh Element Paint Tool Class
+ */
+UCLASS()
+class DATAFLOWEDITOR_API UDataflowEditorWeightMapPaintTool : public UMeshSculptToolBase
+{
+	GENERATED_BODY()
+
+public:
+	virtual void RegisterActions(FInteractiveToolActionSet& ActionSet) override;
+
+	virtual void Setup() override;
+	virtual void Shutdown(EToolShutdownType ShutdownType) override;
+	virtual void DrawHUD(FCanvas* Canvas, IToolsContextRenderAPI* RenderAPI) override;
+	virtual void Render(IToolsContextRenderAPI* RenderAPI) override;
+	virtual void UpdateMaterialMode(EMeshEditingMaterialModes MaterialMode) override;
+
+	virtual void UpdateStampPendingState() override;
+
+	virtual void OnTick(float DeltaTime) override;
+
+	virtual bool HasCancel() const override { return true; }
+	virtual bool HasAccept() const override { return true; }
+	virtual bool CanAccept() const override;
+
+	virtual bool OnUpdateHover(const FInputDeviceRay& DevicePos) override;
+
+	virtual void OnPropertyModified(UObject* PropertySet, FProperty* Property) override;
+
+	bool IsInBrushSubMode() const;
+
+	virtual void CommitResult(UBaseDynamicMeshComponent* Component, bool bModifiedTopology) override;
+
+	void SetDataflowEditorContextObject(TObjectPtr<UDataflowEditorContextObject> InDataflowEditorContextObject);
+
+public:
+
+
+	/** Filters on paint brush */
+	UPROPERTY()
+	TObjectPtr<UDataflowEditorWeightMapPaintBrushFilterProperties> FilterProperties;
+
+
+private:
+	UPROPERTY()
+	TObjectPtr<UDataflowWeightMapPaintBrushOpProps> PaintBrushOpProperties;
+
+	UPROPERTY()
+	TObjectPtr<UDataflowWeightMapSmoothBrushOpProps> SmoothBrushOpProperties;
+
+	UPROPERTY()
+	TObjectPtr<UDataflowWeightMapEraseBrushOpProps> EraseBrushOpProperties;
+
+public:
+	void FloodFillCurrentWeightAction();
+	void ClearAllWeightsAction();
+
+	void SetVerticesToWeightMap(const TSet<int32>& Vertices, double WeightValue, bool bIsErase);
+
+	bool HaveVisibilityFilter() const;
+	void ApplyVisibilityFilter(const TArray<int32>& Vertices, TArray<int32>& VisibleVertices);
+	void ApplyVisibilityFilter(TSet<int32>& Triangles, TArray<int32>& ROIBuffer, TArray<int32>& OutputBuffer);
+
+	// we override these so we can update the separate BrushSize property added for this tool
+	virtual void IncreaseBrushRadiusAction();
+	virtual void DecreaseBrushRadiusAction();
+	virtual void IncreaseBrushRadiusSmallStepAction();
+	virtual void DecreaseBrushRadiusSmallStepAction();
+
+protected:
+	// UMeshSculptToolBase API
+	virtual UBaseDynamicMeshComponent* GetSculptMeshComponent() { return DynamicMeshComponent; }
+	virtual FDynamicMesh3* GetBaseMesh() { check(false); return nullptr; }
+	virtual const FDynamicMesh3* GetBaseMesh() const { check(false); return nullptr; }
+
+	virtual int32 FindHitSculptMeshTriangle(const FRay3d& LocalRay) override;
+	virtual int32 FindHitTargetMeshTriangle(const FRay3d& LocalRay) override;
+
+	virtual void OnBeginStroke(const FRay& WorldRay) override;
+	virtual void OnEndStroke() override;
+
+	virtual TUniquePtr<FMeshSculptBrushOp>& GetActiveBrushOp();
+
+	virtual bool SharesBrushPropertiesChanges() const override { return false; }
+	// end UMeshSculptToolBase API
+
+
+
+	//
+	// Action support
+	//
+
+public:
+	virtual void RequestAction(EDataflowEditorWeightMapPaintToolActions ActionType);
+	
+	UPROPERTY()
+	TObjectPtr<UDataflowEditorMeshWeightMapPaintToolActions> ActionsProps;
+
+	UPROPERTY()
+	TObjectPtr<UDataflowEditorUpdateWeightMapProperties> UpdateWeightMapProperties;
+
+protected:
+	bool bHavePendingAction = false;
+	EDataflowEditorWeightMapPaintToolActions PendingAction;
+	virtual void ApplyAction(EDataflowEditorWeightMapPaintToolActions ActionType);
+
+
+
+	//
+	// Marquee Support
+	//
+public:
+	UPROPERTY()
+	TObjectPtr<UPolyLassoMarqueeMechanic> PolyLassoMechanic;
+
+protected:
+	void OnPolyLassoFinished(const FCameraPolyLasso& Lasso, bool bCanceled);
+
+
+	// 
+	// Gradient Support
+	//
+	UPROPERTY()
+	TObjectPtr<UPolygonSelectionMechanic> PolygonSelectionMechanic;
+
+	TUniquePtr<UE::Geometry::FDynamicMeshAABBTree3> MeshSpatial = nullptr;
+
+	TUniquePtr<UE::Geometry::FTriangleGroupTopology> GradientSelectionTopology = nullptr;
+
+	FToolDataVisualizer GradientSelectionRenderer;
+
+	UE::Geometry::FGroupTopologySelection LowValueGradientVertexSelection;
+	UE::Geometry::FGroupTopologySelection HighValueGradientVertexSelection;
+
+	void ComputeGradient();
+	void OnSelectionModified();
+
+	//
+	// Internals
+	//
+
+protected:
+
+	UPROPERTY()
+	TObjectPtr<AInternalToolFrameworkActor> PreviewMeshActor = nullptr;
+
+	UPROPERTY()
+	TObjectPtr<UDynamicMeshComponent> DynamicMeshComponent;
+
+	UPROPERTY()
+	TObjectPtr<UMeshElementsVisualizer> MeshElementsDisplay;
+
+	UPROPERTY()
+	TObjectPtr<UDataflowEditorContextObject> DataflowEditorContextObject = nullptr;
+
+	// realtime visualization
+	void OnDynamicMeshComponentChanged(UDynamicMeshComponent* Component, const FMeshVertexChange* Change, bool bRevert);
+	FDelegateHandle OnDynamicMeshComponentChangedHandle;
+
+	UE::Geometry::FDynamicMeshWeightAttribute* ActiveWeightMap;
+	double GetCurrentWeightValue(int32 VertexId) const;
+	double GetCurrentWeightValueUnderBrush() const;
+	FVector3d CurrentBaryCentricCoords;
+	int32 GetBrushNearestVertex() const;
+
+	void GetCurrentWeightMap(TArray<float>& OutWeights) const;
+
+	void UpdateSubToolType(EDataflowEditorWeightMapPaintInteractionType NewType);
+
+	void UpdateBrushType(EDataflowEditorWeightMapPaintBrushType BrushType);
+
+	TSet<int32> AccumulatedTriangleROI;
+	bool bUndoUpdatePending = false;
+	TArray<int> NormalsBuffer;
+	void WaitForPendingUndoRedo();
+
+	TArray<int> TempROIBuffer;
+	TArray<int> VertexROI;
+	TArray<bool> VisibilityFilterBuffer;
+	TSet<int> VertexSetBuffer;
+	TSet<int> TriangleROI;
+	void UpdateROI(const FSculptBrushStamp& CurrentStamp);
+
+	EDataflowEditorWeightMapPaintBrushType PendingStampType = EDataflowEditorWeightMapPaintBrushType::Paint;
+
+	bool UpdateStampPosition(const FRay& WorldRay);
+	bool ApplyStamp();
+
+	UE::Geometry::FDynamicMeshOctree3 Octree;
+
+	bool UpdateBrushPosition(const FRay& WorldRay);
+
+	bool GetInEraseStroke()
+	{
+		// Re-use the smoothing stroke key (shift) for erase stroke in the weight paint tool
+		return GetInSmoothingStroke();
+	}
+
+
+	bool bPendingPickWeight = false;
+
+	TArray<int32> ROITriangleBuffer;
+	TArray<double> ROIWeightValueBuffer;
+	bool SyncMeshWithWeightBuffer(FDynamicMesh3* Mesh);
+	bool SyncWeightBufferWithMesh(const FDynamicMesh3* Mesh);
+
+	TUniquePtr<UE::Geometry::FDynamicMeshChangeTracker> ActiveWeightEditChangeTracker;
+	void BeginChange();
+	void EndChange();
+
+	FColor GetColorForWeightValue(double WeightValue);
+
+	TArray<FVector3d> TriNormals;
+	TArray<int32> UVSeamEdges;
+	TArray<int32> NormalSeamEdges;
+	void PrecomputeFilterData();
+
+	// DynamicMesh might be unwelded mesh, but weights are on the welded mesh.
+	bool bHaveDynamicMeshToWeightConversion = false;
+	TArray<int32> DynamicMeshToWeight; 
+	TArray<TArray<int32>> WeightToDynamicMesh;
+
+protected:
+	virtual bool ShowWorkPlane() const override { return false; }
+
+	friend class UDataflowEditorWeightMapPaintToolBuilder;
+
+	bool bAnyChangeMade = false;
+
+	// Node graph editor support
+
+	FDataflowCollectionAddScalarVertexPropertyNode* WeightMapNodeToUpdate = nullptr;
+
+	void UpdateSelectedNode();
+
+	void UpdateVertexColorOverlay(const TSet<int>* TrianglesToUpdate = nullptr);
+};
+
+
+
diff --git a/Engine/Plugins/Experimental/Dataflow/Source/DataflowNodes/Private/Dataflow/DataflowCollectionAddScalarVertexPropertyNode.cpp b/Engine/Plugins/Experimental/Dataflow/Source/DataflowNodes/Private/Dataflow/DataflowCollectionAddScalarVertexPropertyNode.cpp
new file mode 100644
--- /dev/null
+++ b/Engine/Plugins/Experimental/Dataflow/Source/DataflowNodes/Private/Dataflow/DataflowCollectionAddScalarVertexPropertyNode.cpp
@@ -0,0 +1,58 @@
+// Copyright Epic Games, Inc. All Rights Reserved.
+
+#include "Dataflow/DataflowCollectionAddScalarVertexPropertyNode.h"
+
+#include "Dataflow/DataflowInputOutput.h"
+#include "Dataflow/DataflowTools.h"
+
+#include UE_INLINE_GENERATED_CPP_BY_NAME(DataflowCollectionAddScalarVertexPropertyNode)
+
+#define LOCTEXT_NAMESPACE "DataflowCollectionAddScalarVertexProperty"
+
+FDataflowCollectionAddScalarVertexPropertyNode::FDataflowCollectionAddScalarVertexPropertyNode(const Dataflow::FNodeParameters& InParam, FGuid InGuid)
+	: FDataflowNode(InParam, InGuid)
+{
+	RegisterInputConnection(&Collection);
+	RegisterOutputConnection(&Collection, &Collection);
+	RegisterOutputConnection(&Name);
+}
+
+void FDataflowCollectionAddScalarVertexPropertyNode::Evaluate(Dataflow::FContext& Context, const FDataflowOutput* Out) const
+{
+	using namespace Dataflow;
+
+	if (Out->IsA<FManagedArrayCollection>(&Collection))
+	{
+		// Evaluate in collection
+		FManagedArrayCollection InCollection = GetValue<FManagedArrayCollection>(Context, &Collection);
+
+		if (!Name.IsEmpty())
+		{
+			FName nName(Name), nGroup("Vertices");
+			TManagedArray<float>& Scalar = InCollection.AddAttribute<float>(nName,nGroup);
+
+			const int32 MaxWeightIndex = FMath::Min(VertexWeights.Num(), Scalar.Num());
+			if (VertexWeights.Num() > 0 && VertexWeights.Num() != Scalar.Num())
+			{
+				FDataflowTools::LogAndToastWarning(*this,
+					LOCTEXT("VertexCountMismatchHeadline", "Vertex count mismatch."),
+					FText::Format(LOCTEXT("VertexCountMismatchDetails", "Vertex weights in the node: {0}\n3D vertices in the cloth: {1}"),
+						VertexWeights.Num(),
+						Scalar.Num()));
+			}
+
+			for (int32 VertexID = 0; VertexID < MaxWeightIndex; ++VertexID)
+			{
+				Scalar[VertexID] = VertexWeights[VertexID];
+			}
+		}
+
+		SetValue(Context, MoveTemp(InCollection), &Collection);
+	}
+	else if (Out->IsA<FString>(&Name))
+	{
+		SetValue(Context, Name, &Name);
+	}
+}
+
+#undef LOCTEXT_NAMESPACE
diff --git a/Engine/Plugins/Experimental/Dataflow/Source/DataflowNodes/Private/Dataflow/DataflowNodesPlugin.cpp b/Engine/Plugins/Experimental/Dataflow/Source/DataflowNodes/Private/Dataflow/DataflowNodesPlugin.cpp
--- a/Engine/Plugins/Experimental/Dataflow/Source/DataflowNodes/Private/Dataflow/DataflowNodesPlugin.cpp
+++ b/Engine/Plugins/Experimental/Dataflow/Source/DataflowNodes/Private/Dataflow/DataflowNodesPlugin.cpp
@@ -8,6 +8,7 @@
 #include "Dataflow/DataflowNodeFactory.h"
 #include "Dataflow/DataflowSelectionNodes.h"
 #include "Dataflow/DataflowContextOverridesNodes.h"
+#include "Dataflow/DataflowCollectionAddScalarVertexPropertyNode.h"
 
 #define LOCTEXT_NAMESPACE "DataflowNodes"
 
@@ -18,6 +19,7 @@ void IDataflowNodesPlugin::StartupModule()
 	Dataflow::RegisterStaticMeshNodes();
 	Dataflow::RegisterSelectionNodes();
 	Dataflow::RegisterContextOverridesNodes();
+	DATAFLOW_NODE_REGISTER_CREATION_FACTORY(FDataflowCollectionAddScalarVertexPropertyNode);
 }
 
 void IDataflowNodesPlugin::ShutdownModule()
diff --git a/Engine/Plugins/Experimental/Dataflow/Source/DataflowNodes/Private/Dataflow/DataflowTools.cpp b/Engine/Plugins/Experimental/Dataflow/Source/DataflowNodes/Private/Dataflow/DataflowTools.cpp
new file mode 100644
--- /dev/null
+++ b/Engine/Plugins/Experimental/Dataflow/Source/DataflowNodes/Private/Dataflow/DataflowTools.cpp
@@ -0,0 +1,33 @@
+// Copyright Epic Games, Inc. All Rights Reserved.
+
+#include "Dataflow/DataflowTools.h"
+#include "Dataflow/DataflowNode.h"
+//#include "Widgets/Notifications/SNotificationList.h"
+//#include "Framework/Notifications/NotificationManager.h"
+
+DEFINE_LOG_CATEGORY(LogDataflowNodes);
+
+namespace Dataflow
+{
+	void FDataflowTools::LogAndToastWarning(const FDataflowNode& DataflowNode, const FText& Headline, const FText& Details)
+	{
+		static const FTextFormat TextFormat = FTextFormat::FromString(TEXT("{0}: {1}\n{2}"));
+		const FText NodeName = FText::FromName(DataflowNode.GetName());
+		const FText Text = FText::Format(TextFormat, NodeName, Headline, Details);
+
+		// @todo(Dataflow) : Add error suppurt during evaluation
+		//FNotificationInfo NotificationInfo(Text);
+		//NotificationInfo.ExpireDuration = 5.0f;
+		//FSlateNotificationManager::Get().AddNotification(NotificationInfo);
+
+		UE_LOG(LogDataflowNodes, Display, TEXT("%s"), *Text.ToString());
+	}
+
+	void FDataflowTools::MakeCollectionName(FString& InOutString)
+	{
+		InOutString = SlugStringForValidName(InOutString, TEXT("_")).Replace(TEXT("\\"), TEXT("_"));
+		bool bCharsWereRemoved;
+		do { InOutString.TrimCharInline(TEXT('_'), &bCharsWereRemoved); } while (bCharsWereRemoved);
+	}
+
+}  // End namespace Dataflow
diff --git a/Engine/Plugins/Experimental/Dataflow/Source/DataflowNodes/Public/Dataflow/DataflowCollectionAddScalarVertexPropertyNode.h b/Engine/Plugins/Experimental/Dataflow/Source/DataflowNodes/Public/Dataflow/DataflowCollectionAddScalarVertexPropertyNode.h
new file mode 100644
--- /dev/null
+++ b/Engine/Plugins/Experimental/Dataflow/Source/DataflowNodes/Public/Dataflow/DataflowCollectionAddScalarVertexPropertyNode.h
@@ -0,0 +1,33 @@
+// Copyright Epic Games, Inc. All Rights Reserved.
+
+#pragma once 
+
+#include "Dataflow/DataflowNode.h"
+#include "GeometryCollection/ManagedArrayCollection.h"
+
+#include "DataflowCollectionAddScalarVertexPropertyNode.generated.h"
+
+/** Scalar vertex properties. */
+USTRUCT(Meta = (DataflowCollection))
+struct DATAFLOWNODES_API FDataflowCollectionAddScalarVertexPropertyNode : public FDataflowNode
+{
+	GENERATED_USTRUCT_BODY()
+	DATAFLOW_NODE_DEFINE_INTERNAL(FDataflowCollectionAddScalarVertexPropertyNode, "AddScalarVertexProperty", "Collection", "Add a saved scalar property to a collection")
+
+public:
+
+	UPROPERTY(Meta = (Dataflowinput, DataflowOutput, DataflowPassthrough = "Collection"))
+	FManagedArrayCollection Collection;
+
+	/** The name to be set as a weight map attribute. */
+	UPROPERTY(EditAnywhere, Category = "Vertex Attribute", Meta = (DataflowOutput))
+	FString Name;
+
+	UPROPERTY()
+	TArray<float> VertexWeights;
+
+	FDataflowCollectionAddScalarVertexPropertyNode(const Dataflow::FNodeParameters& InParam, FGuid InGuid = FGuid::NewGuid());
+
+private:
+	virtual void Evaluate(Dataflow::FContext& Context, const FDataflowOutput* Out) const override;
+};
diff --git a/Engine/Plugins/Experimental/Dataflow/Source/DataflowNodes/Public/Dataflow/DataflowTools.h b/Engine/Plugins/Experimental/Dataflow/Source/DataflowNodes/Public/Dataflow/DataflowTools.h
new file mode 100644
--- /dev/null
+++ b/Engine/Plugins/Experimental/Dataflow/Source/DataflowNodes/Public/Dataflow/DataflowTools.h
@@ -0,0 +1,31 @@
+// Copyright Epic Games, Inc. All Rights Reserved.
+
+#pragma once
+
+#include "Templates/SharedPointer.h"
+#include "Logging/LogMacros.h"
+
+struct FManagedArrayCollection;
+struct FDataflowNode;
+class FString;
+
+namespace Dataflow
+{
+	/**
+	* Dataflow Node Tools
+	*/
+	struct FDataflowTools
+	{
+		static void LogAndToastWarning(const FDataflowNode& DataflowNode, const FText& Headline, const FText& Details);
+
+		/**
+		 * Turn a string into a valid collection group or attribute name.
+		 * The resulting name won't contains spaces and any other special characters as listed in
+		 * INVALID_OBJECTNAME_CHARACTERS (currently "',/.:|&!~\n\r\t@#(){}[]=;^%$`).
+		 * It will also have all leading underscore removed, as these names are reserved for internal use.
+		 */
+		static void MakeCollectionName(FString& InOutString);
+	};
+}  // End namespace Dataflow
+
+DECLARE_LOG_CATEGORY_EXTERN(LogDataflowNodes, Log, All);
