Dataflow : Fixed duplicate function definition between Cloth and Dataflow.

diff --git a/Engine/Plugins/Experimental/Dataflow/Source/DataflowEditor/Private/DataflowEditorTools/DataflowEditorWeightMapPaintTool.cpp b/Engine/Plugins/Experimental/Dataflow/Source/DataflowEditor/Private/DataflowEditorTools/DataflowEditorWeightMapPaintTool.cpp
--- a/Engine/Plugins/Experimental/Dataflow/Source/DataflowEditor/Private/DataflowEditorTools/DataflowEditorWeightMapPaintTool.cpp
+++ b/Engine/Plugins/Experimental/Dataflow/Source/DataflowEditor/Private/DataflowEditorTools/DataflowEditorWeightMapPaintTool.cpp
@@ -774,22 +774,24 @@ bool UDataflowEditorWeightMapPaintTool::SyncWeightBufferWithMesh(const FDynamicM
 	return (NumModified > 0);
 }
 
-template<typename RealType>
-static bool FindPolylineSelfIntersection(
-	const TArray<UE::Math::TVector2<RealType>>& Polyline, 
-	UE::Math::TVector2<RealType>& IntersectionPointOut, 
-	FIndex2i& IntersectionIndexOut,
-	bool bParallel = true)
+namespace Dataflow
 {
-	int32 N = Polyline.Num();
-	std::atomic<bool> bSelfIntersects(false);
-	ParallelFor(N - 1, [&](int32 i)
-	{
-		TSegment2<RealType> SegA(Polyline[i], Polyline[i + 1]);
+	template<typename RealType>
+	static bool FindPolylineSelfIntersection(
+		const TArray<UE::Math::TVector2<RealType>>& Polyline,
+		UE::Math::TVector2<RealType>& IntersectionPointOut,
+		FIndex2i& IntersectionIndexOut,
+		bool bParallel = true)
+	{
+		int32 N = Polyline.Num();
+		std::atomic<bool> bSelfIntersects(false);
+		ParallelFor(N - 1, [&](int32 i)
+			{
+				TSegment2<RealType> SegA(Polyline[i], Polyline[i + 1]);
 		for (int32 j = i + 2; j < N - 1 && bSelfIntersects == false; ++j)
 		{
 			TSegment2<RealType> SegB(Polyline[j], Polyline[j + 1]);
-			if (SegA.Intersects(SegB) && bSelfIntersects == false)		
+			if (SegA.Intersects(SegB) && bSelfIntersects == false)
 			{
 				bool ExpectedValue = false;
 				if (std::atomic_compare_exchange_strong(&bSelfIntersects, &ExpectedValue, true))
@@ -802,106 +804,105 @@ static bool FindPolylineSelfIntersection(
 				}
 			}
 		}
-	}, (bParallel) ? EParallelForFlags::None : EParallelForFlags::ForceSingleThread );
-
-	return bSelfIntersects;
-}
+			}, (bParallel) ? EParallelForFlags::None : EParallelForFlags::ForceSingleThread);
 
+		return bSelfIntersects;
+	}
 
 
-template<typename RealType>
-static bool FindPolylineSegmentIntersection(
-	const TArray<UE::Math::TVector2<RealType>>& Polyline,
-	const TSegment2<RealType>& Segment,
-	UE::Math::TVector2<RealType>& IntersectionPointOut,
-	int& IntersectionIndexOut)
-{
 
-	int32 N = Polyline.Num();
-	for (int32 i = 0; i < N-1; ++i)
+	template<typename RealType>
+	static bool FindPolylineSegmentIntersection(
+		const TArray<UE::Math::TVector2<RealType>>& Polyline,
+		const TSegment2<RealType>& Segment,
+		UE::Math::TVector2<RealType>& IntersectionPointOut,
+		int& IntersectionIndexOut)
 	{
-		TSegment2<RealType> PolySeg(Polyline[i], Polyline[i + 1]);
-		if (Segment.Intersects(PolySeg))
+
+		int32 N = Polyline.Num();
+		for (int32 i = 0; i < N - 1; ++i)
 		{
-			TIntrSegment2Segment2<RealType> Intersection(Segment, PolySeg);
-			Intersection.Find();
-			IntersectionPointOut = Intersection.Point0;
-			IntersectionIndexOut = i;
-			return true;
+			TSegment2<RealType> PolySeg(Polyline[i], Polyline[i + 1]);
+			if (Segment.Intersects(PolySeg))
+			{
+				TIntrSegment2Segment2<RealType> Intersection(Segment, PolySeg);
+				Intersection.Find();
+				IntersectionPointOut = Intersection.Point0;
+				IntersectionIndexOut = i;
+				return true;
+			}
 		}
+		return false;
 	}
-	return false;
-}
-
 
 
-bool ApproxSelfClipPolyline(TArray<FVector2f>& Polyline)
-{
-	int32 N = Polyline.Num();
-
-	// handle already-closed polylines
-	if (Distance(Polyline[0], Polyline[N-1]) < 0.0001f)
+	bool ApproxSelfClipPolyline(TArray<FVector2f>& Polyline)
 	{
-		return true;
-	}
+		int32 N = Polyline.Num();
 
-	FVector2f IntersectPoint;
-	FIndex2i IntersectionIndex(-1, -1);
-	bool bSelfIntersects = FindPolylineSelfIntersection(Polyline, IntersectPoint, IntersectionIndex);
-	if (bSelfIntersects)
-	{
-		TArray<FVector2f> NewPolyline;
-		NewPolyline.Add(IntersectPoint);
-		for (int32 i = IntersectionIndex.A; i <= IntersectionIndex.B; ++i)
+		// handle already-closed polylines
+		if (Distance(Polyline[0], Polyline[N - 1]) < 0.0001f)
 		{
-			NewPolyline.Add(Polyline[i]);
+			return true;
 		}
-		NewPolyline.Add(IntersectPoint);
-		Polyline = MoveTemp(NewPolyline);
-		return true;
-	}
-
 
-	FVector2f StartDirOut = UE::Geometry::Normalized(Polyline[0] - Polyline[1]);
-	FLine2f StartLine(Polyline[0], StartDirOut);
-	FVector2f EndDirOut = UE::Geometry::Normalized(Polyline[N - 1] - Polyline[N - 2]);
-	FLine2f EndLine(Polyline[N - 1], EndDirOut);
-	FIntrLine2Line2f LineIntr(StartLine, EndLine);
-	bool bIntersects = false;
-	if (LineIntr.Find())
-	{
-		bIntersects = LineIntr.IsSimpleIntersection() && (LineIntr.Segment1Parameter > 0) && (LineIntr.Segment2Parameter > 0);
-		if (bIntersects)
+		FVector2f IntersectPoint;
+		FIndex2i IntersectionIndex(-1, -1);
+		bool bSelfIntersects = FindPolylineSelfIntersection(Polyline, IntersectPoint, IntersectionIndex);
+		if (bSelfIntersects)
 		{
-			Polyline.Add(StartLine.PointAt(LineIntr.Segment1Parameter));
-			Polyline.Add(StartLine.Origin);
+			TArray<FVector2f> NewPolyline;
+			NewPolyline.Add(IntersectPoint);
+			for (int32 i = IntersectionIndex.A; i <= IntersectionIndex.B; ++i)
+			{
+				NewPolyline.Add(Polyline[i]);
+			}
+			NewPolyline.Add(IntersectPoint);
+			Polyline = MoveTemp(NewPolyline);
 			return true;
 		}
-	}
 
 
-	FAxisAlignedBox2f Bounds;
-	for (const FVector2f& P : Polyline)
-	{
-		Bounds.Contain(P);
-	}
-	float Size = Bounds.DiagonalLength();
+		FVector2f StartDirOut = UE::Geometry::Normalized(Polyline[0] - Polyline[1]);
+		FLine2f StartLine(Polyline[0], StartDirOut);
+		FVector2f EndDirOut = UE::Geometry::Normalized(Polyline[N - 1] - Polyline[N - 2]);
+		FLine2f EndLine(Polyline[N - 1], EndDirOut);
+		FIntrLine2Line2f LineIntr(StartLine, EndLine);
+		bool bIntersects = false;
+		if (LineIntr.Find())
+		{
+			bIntersects = LineIntr.IsSimpleIntersection() && (LineIntr.Segment1Parameter > 0) && (LineIntr.Segment2Parameter > 0);
+			if (bIntersects)
+			{
+				Polyline.Add(StartLine.PointAt(LineIntr.Segment1Parameter));
+				Polyline.Add(StartLine.Origin);
+				return true;
+			}
+		}
 
-	FVector2f StartPos = Polyline[0] + 0.001f * StartDirOut;
-	if (FindPolylineSegmentIntersection(Polyline, FSegment2f(StartPos, StartPos + 2*Size*StartDirOut), IntersectPoint, IntersectionIndex.A))
-	{
-		return true;
-	}
 
-	FVector2f EndPos = Polyline[N-1] + 0.001f * EndDirOut;
-	if (FindPolylineSegmentIntersection(Polyline, FSegment2f(EndPos, EndPos + 2*Size*EndDirOut), IntersectPoint, IntersectionIndex.A))
-	{
-		return true;
-	}
+		FAxisAlignedBox2f Bounds;
+		for (const FVector2f& P : Polyline)
+		{
+			Bounds.Contain(P);
+		}
+		float Size = Bounds.DiagonalLength();
 
-	return false;
-}
+		FVector2f StartPos = Polyline[0] + 0.001f * StartDirOut;
+		if (FindPolylineSegmentIntersection(Polyline, FSegment2f(StartPos, StartPos + 2 * Size * StartDirOut), IntersectPoint, IntersectionIndex.A))
+		{
+			return true;
+		}
 
+		FVector2f EndPos = Polyline[N - 1] + 0.001f * EndDirOut;
+		if (FindPolylineSegmentIntersection(Polyline, FSegment2f(EndPos, EndPos + 2 * Size * EndDirOut), IntersectPoint, IntersectionIndex.A))
+		{
+			return true;
+		}
+
+		return false;
+	}
+}
 
 
 void UDataflowEditorWeightMapPaintTool::OnPolyLassoFinished(const FCameraPolyLasso& Lasso, bool bCanceled)
@@ -921,7 +922,7 @@ void UDataflowEditorWeightMapPaintTool::OnPolyLassoFinished(const FCameraPolyLas
 	// Try to clip polyline to be closed, or closed-enough for winding evaluation to work.
 	// If that returns false, the polyline is "too open". In that case we will extend
 	// outwards from the endpoints and then try to create a closed very large polygon
-	if (ApproxSelfClipPolyline(Polyline) == false)
+	if (Dataflow::ApproxSelfClipPolyline(Polyline) == false)
 	{
 		FVector2f StartDirOut = UE::Geometry::Normalized(Polyline[0] - Polyline[1]);
 		FLine2f StartLine(Polyline[0], StartDirOut);
