Dataflow : FGCObject includes for garbage collection

diff --git a/Engine/Plugins/Experimental/Dataflow/Source/DataflowEditor/Private/Dataflow/DataflowEditorToolkit.cpp b/Engine/Plugins/Experimental/Dataflow/Source/DataflowEditor/Private/Dataflow/DataflowEditorToolkit.cpp
--- a/Engine/Plugins/Experimental/Dataflow/Source/DataflowEditor/Private/Dataflow/DataflowEditorToolkit.cpp
+++ b/Engine/Plugins/Experimental/Dataflow/Source/DataflowEditor/Private/Dataflow/DataflowEditorToolkit.cpp
@@ -503,6 +503,7 @@ void FDataflowEditorToolkit::OnNodeSelectionChanged(const TSet<UObject*>& InNewS
 			// The selection is empty. 
 			ResetListeners();
 			SelectedDataflowNodes = TSet<UObject*>();
+			if (PrimarySelection) bPrimarySelectionChanged = true;
 			PrimarySelection = nullptr;
 		}
 		else
diff --git a/Engine/Plugins/Experimental/Dataflow/Source/DataflowEditor/Private/Dataflow/DataflowEditorViewportClient.cpp b/Engine/Plugins/Experimental/Dataflow/Source/DataflowEditor/Private/Dataflow/DataflowEditorViewportClient.cpp
--- a/Engine/Plugins/Experimental/Dataflow/Source/DataflowEditor/Private/Dataflow/DataflowEditorViewportClient.cpp
+++ b/Engine/Plugins/Experimental/Dataflow/Source/DataflowEditor/Private/Dataflow/DataflowEditorViewportClient.cpp
@@ -129,5 +129,6 @@ Dataflow::EDataflowPatternVertexType FDataflowEditorViewportClient::GetConstruct
 void FDataflowEditorViewportClient::AddReferencedObjects(FReferenceCollector& Collector)
 {
 	Super::AddReferencedObjects(Collector);
+	Collector.AddReferencedObject(BehaviorSet);
 }
 
diff --git a/Engine/Plugins/Experimental/Dataflow/Source/DataflowEditor/Private/Dataflow/DataflowGraphEditor.cpp b/Engine/Plugins/Experimental/Dataflow/Source/DataflowEditor/Private/Dataflow/DataflowGraphEditor.cpp
--- a/Engine/Plugins/Experimental/Dataflow/Source/DataflowEditor/Private/Dataflow/DataflowGraphEditor.cpp
+++ b/Engine/Plugins/Experimental/Dataflow/Source/DataflowEditor/Private/Dataflow/DataflowGraphEditor.cpp
@@ -529,4 +529,11 @@ bool SDataflowGraphEditor::GetPinVisibility(SGraphEditor::EPinVisibility PinVisi
 	return false;
 }
 
+void SDataflowGraphEditor::AddReferencedObjects(FReferenceCollector& Collector)
+{
+	Collector.AddReferencedObject(DataflowAsset);
+	Collector.AddReferencedObject(AssetOwner);
+}
+
+
 #undef LOCTEXT_NAMESPACE
diff --git a/Engine/Plugins/Experimental/Dataflow/Source/DataflowEditor/Private/Dataflow/DataflowSCommentNode.cpp b/Engine/Plugins/Experimental/Dataflow/Source/DataflowEditor/Private/Dataflow/DataflowSCommentNode.cpp
--- a/Engine/Plugins/Experimental/Dataflow/Source/DataflowEditor/Private/Dataflow/DataflowSCommentNode.cpp
+++ b/Engine/Plugins/Experimental/Dataflow/Source/DataflowEditor/Private/Dataflow/DataflowSCommentNode.cpp
@@ -119,4 +119,11 @@ UEdGraphNode* FAssetSchemaAction_Dataflow_DuplicateCommentNode_DataflowEdNode::P
 	return CommentTemplate;
 }
 
+
+void FAssetSchemaAction_Dataflow_DuplicateCommentNode_DataflowEdNode::AddReferencedObjects(FReferenceCollector& Collector)
+{
+	Collector.AddReferencedObject(CommentNodeToDuplicate);
+}
+
+
 #undef LOCTEXT_NAMESPACE
diff --git a/Engine/Plugins/Experimental/Dataflow/Source/DataflowEditor/Private/Dataflow/DataflowSNode.cpp b/Engine/Plugins/Experimental/Dataflow/Source/DataflowEditor/Private/Dataflow/DataflowSNode.cpp
--- a/Engine/Plugins/Experimental/Dataflow/Source/DataflowEditor/Private/Dataflow/DataflowSNode.cpp
+++ b/Engine/Plugins/Experimental/Dataflow/Source/DataflowEditor/Private/Dataflow/DataflowSNode.cpp
@@ -156,6 +156,7 @@ void SDataflowEdNode::AddReferencedObjects(FReferenceCollector& Collector)
 {
 	if (DataflowGraphNode)
 	{
+		Collector.AddReferencedObject(DataflowGraphNode);
 		if (TSharedPtr<FDataflowNode> DataflowNode = DataflowGraphNode->GetDataflowNode())
 		{
 			Collector.AddPropertyReferences(DataflowNode->TypedScriptStruct(), DataflowNode.Get());
diff --git a/Engine/Plugins/Experimental/Dataflow/Source/DataflowEditor/Private/Dataflow/DataflowSchema.cpp b/Engine/Plugins/Experimental/Dataflow/Source/DataflowEditor/Private/Dataflow/DataflowSchema.cpp
--- a/Engine/Plugins/Experimental/Dataflow/Source/DataflowEditor/Private/Dataflow/DataflowSchema.cpp
+++ b/Engine/Plugins/Experimental/Dataflow/Source/DataflowEditor/Private/Dataflow/DataflowSchema.cpp
@@ -304,5 +304,11 @@ void FDataflowConnectionDrawingPolicy::DetermineWiringStyle(UEdGraphPin* OutputP
 	}
 }
 
+void FDataflowConnectionDrawingPolicy::AddReferencedObjects(FReferenceCollector& Collector)
+{
+	Collector.AddReferencedObject(Schema);
+}
+
+
 #undef LOCTEXT_NAMESPACE
 
diff --git a/Engine/Plugins/Experimental/Dataflow/Source/DataflowEditor/Private/Dataflow/DataflowView.cpp b/Engine/Plugins/Experimental/Dataflow/Source/DataflowEditor/Private/Dataflow/DataflowView.cpp
--- a/Engine/Plugins/Experimental/Dataflow/Source/DataflowEditor/Private/Dataflow/DataflowView.cpp
+++ b/Engine/Plugins/Experimental/Dataflow/Source/DataflowEditor/Private/Dataflow/DataflowView.cpp
@@ -111,5 +111,10 @@ void FDataflowNodeView::OnNodeInvalidated(FDataflowNode* InvalidatedNode)
 	}
 }
 
+void FDataflowNodeView::AddReferencedObjects(FReferenceCollector& Collector)
+{
+	Collector.AddReferencedObject(SelectedNode);
+}
+
 
 #undef LOCTEXT_NAMESPACE
diff --git a/Engine/Plugins/Experimental/Dataflow/Source/DataflowEditor/Public/Dataflow/DataflowEditor.h b/Engine/Plugins/Experimental/Dataflow/Source/DataflowEditor/Public/Dataflow/DataflowEditor.h
--- a/Engine/Plugins/Experimental/Dataflow/Source/DataflowEditor/Public/Dataflow/DataflowEditor.h
+++ b/Engine/Plugins/Experimental/Dataflow/Source/DataflowEditor/Public/Dataflow/DataflowEditor.h
@@ -52,6 +52,7 @@ private :
 	// is the one holding the dynamic mesh components to be rendered in the viewport
 	// It is why the data flow asset/owner/skelmesh have been added here. Could be added
 	// in the subsystem if necessary
+	UPROPERTY()
 	TObjectPtr<UDataflowEditorContent> Content;
 };
 
diff --git a/Engine/Plugins/Experimental/Dataflow/Source/DataflowEditor/Public/Dataflow/DataflowEditorContent.h b/Engine/Plugins/Experimental/Dataflow/Source/DataflowEditor/Public/Dataflow/DataflowEditorContent.h
--- a/Engine/Plugins/Experimental/Dataflow/Source/DataflowEditor/Public/Dataflow/DataflowEditorContent.h
+++ b/Engine/Plugins/Experimental/Dataflow/Source/DataflowEditor/Public/Dataflow/DataflowEditorContent.h
@@ -55,8 +55,10 @@ public:
 
 protected:
 
-	TSharedPtr<FManagedArrayCollection> PrimaryRenderCollection = nullptr;
+	UPROPERTY()
 	TObjectPtr<UDataflowEdNode> PrimarySelectedNode = nullptr;
+
+	TSharedPtr<FManagedArrayCollection> PrimaryRenderCollection = nullptr;
 	TWeakPtr<SDataflowGraphEditor> DataflowGraphEditor;
 	Dataflow::EDataflowPatternVertexType ConstructionViewMode = Dataflow::EDataflowPatternVertexType::Sim3D;
 };
@@ -122,21 +124,26 @@ public:
 
 private:
 	/** Data flow object owner */
+	UPROPERTY()
 	TObjectPtr<UObject> DataflowOwner = nullptr;
 
 	/** Data flow asset that we will edit */
+	UPROPERTY()
 	TObjectPtr<UDataflow> DataflowAsset = nullptr;
 
 	/** Data flow skeletal mesh*/
+	UPROPERTY()
 	TObjectPtr<USkeletalMesh> SkeletalMesh = nullptr;
 
 	/** Data flow skeleton*/
+	UPROPERTY()
 	TObjectPtr<USkeleton> Skeleton = nullptr;
 
 	/** Data flow terminal path for evaluation */
 	FString DataflowTerminal = "";
 
 	/** Animation asset to be used to preview simulation */
+	UPROPERTY()
 	TObjectPtr<UAnimationAsset> AnimationAsset = nullptr;
 
 	/**  Engine context to be used for dataflow evaluation */
diff --git a/Engine/Plugins/Experimental/Dataflow/Source/DataflowEditor/Public/Dataflow/DataflowEditorToolkit.h b/Engine/Plugins/Experimental/Dataflow/Source/DataflowEditor/Public/Dataflow/DataflowEditorToolkit.h
--- a/Engine/Plugins/Experimental/Dataflow/Source/DataflowEditor/Public/Dataflow/DataflowEditorToolkit.h
+++ b/Engine/Plugins/Experimental/Dataflow/Source/DataflowEditor/Public/Dataflow/DataflowEditorToolkit.h
@@ -159,9 +159,11 @@ private:
     FDelegateHandle OnFinishedChangingAssetPropertiesDelegateHandle;
 
 	// The currently selected set of dataflow nodes. 
+	UPROPERTY()
 	TSet<UObject*> SelectedDataflowNodes;
 
-	// The most recently selected dataflow node. (Top
+	// The most recently selected dataflow node.
+	UPROPERTY()
 	TObjectPtr<UDataflowEdNode> PrimarySelection;
 
 	/** Scene in which the 3D sim space preview meshes live. Ownership shared with AdvancedPreviewSettingsWidget*/
diff --git a/Engine/Plugins/Experimental/Dataflow/Source/DataflowEditor/Public/Dataflow/DataflowGraphEditor.h b/Engine/Plugins/Experimental/Dataflow/Source/DataflowEditor/Public/Dataflow/DataflowGraphEditor.h
--- a/Engine/Plugins/Experimental/Dataflow/Source/DataflowEditor/Public/Dataflow/DataflowGraphEditor.h
+++ b/Engine/Plugins/Experimental/Dataflow/Source/DataflowEditor/Public/Dataflow/DataflowGraphEditor.h
@@ -31,7 +31,7 @@ namespace Dataflow {
  * 
  * see(SDataprepGraphEditor for reference)
  */
-class DATAFLOWEDITOR_API SDataflowGraphEditor : public SGraphEditor
+class DATAFLOWEDITOR_API SDataflowGraphEditor : public SGraphEditor, public FGCObject
 {
 public:
 
@@ -122,6 +122,10 @@ public:
 
 	SGraphEditor* GetGraphEditor() { return (SGraphEditor*)this; }
 
+	/** FGCObject interface */
+	virtual void AddReferencedObjects(FReferenceCollector& Collector) override;
+	virtual FString GetReferencerName() const override { return TEXT("SDataflowGraphEditor"); }
+
 private:
 	/** Add an additional option pin to all selected Dataflow nodes for those that overrides the AddPin function. */
 	void OnAddOptionPin();
diff --git a/Engine/Plugins/Experimental/Dataflow/Source/DataflowEditor/Public/Dataflow/DataflowSCommentNode.h b/Engine/Plugins/Experimental/Dataflow/Source/DataflowEditor/Public/Dataflow/DataflowSCommentNode.h
--- a/Engine/Plugins/Experimental/Dataflow/Source/DataflowEditor/Public/Dataflow/DataflowSCommentNode.h
+++ b/Engine/Plugins/Experimental/Dataflow/Source/DataflowEditor/Public/Dataflow/DataflowSCommentNode.h
@@ -50,7 +50,7 @@ public:
 // Action to duplicate a comment node in the graph
 //
 USTRUCT()
-struct DATAFLOWEDITOR_API FAssetSchemaAction_Dataflow_DuplicateCommentNode_DataflowEdNode : public FEdGraphSchemaAction
+struct DATAFLOWEDITOR_API FAssetSchemaAction_Dataflow_DuplicateCommentNode_DataflowEdNode : public FEdGraphSchemaAction, public FGCObject
 {
 	GENERATED_USTRUCT_BODY();
 
@@ -65,6 +65,11 @@ public:
 
 	virtual UEdGraphNode* PerformAction(class UEdGraph* ParentGraph, UEdGraphPin* FromPin, const FVector2D Location, bool bSelectNewNode = true) override;
 
+	/** FGCObject interface */
+	virtual void AddReferencedObjects(FReferenceCollector& Collector) override;
+	virtual FString GetReferencerName() const override { return TEXT("FDataflowConnectionDrawingPolicy"); }
+
 	TSharedPtr<SGraphEditor> GraphEditor;
-	UEdGraphNode_Comment* CommentNodeToDuplicate;
+
+	TObjectPtr<UEdGraphNode_Comment> CommentNodeToDuplicate;
 };
diff --git a/Engine/Plugins/Experimental/Dataflow/Source/DataflowEditor/Public/Dataflow/DataflowSNode.h b/Engine/Plugins/Experimental/Dataflow/Source/DataflowEditor/Public/Dataflow/DataflowSNode.h
--- a/Engine/Plugins/Experimental/Dataflow/Source/DataflowEditor/Public/Dataflow/DataflowSNode.h
+++ b/Engine/Plugins/Experimental/Dataflow/Source/DataflowEditor/Public/Dataflow/DataflowSNode.h
@@ -47,7 +47,7 @@ public:
 	//~ End FGCObject interface
 
 private:
-	UDataflowEdNode* DataflowGraphNode = nullptr;	
+	TObjectPtr<UDataflowEdNode> DataflowGraphNode = nullptr;	
 
 	FCheckBoxStyle CheckBoxStyle;
 	TSharedPtr<SCheckBox> RenderCheckBoxWidget;
diff --git a/Engine/Plugins/Experimental/Dataflow/Source/DataflowEditor/Public/Dataflow/DataflowSchema.h b/Engine/Plugins/Experimental/Dataflow/Source/DataflowEditor/Public/Dataflow/DataflowSchema.h
--- a/Engine/Plugins/Experimental/Dataflow/Source/DataflowEditor/Public/Dataflow/DataflowSchema.h
+++ b/Engine/Plugins/Experimental/Dataflow/Source/DataflowEditor/Public/Dataflow/DataflowSchema.h
@@ -31,7 +31,7 @@ public:
 	static FLinearColor GetTypeColor(const FName& Type);
 };
 
-class FDataflowConnectionDrawingPolicy : public FConnectionDrawingPolicy
+class FDataflowConnectionDrawingPolicy : public FConnectionDrawingPolicy, public FGCObject
 {
 public:
 	FDataflowConnectionDrawingPolicy(int32 InBackLayerID, int32 InFrontLayerID, float InZoomFactor, const FSlateRect& InClippingRect, FSlateWindowElementList& InDrawElements, UEdGraph* InGraph);
@@ -39,7 +39,11 @@ public:
 
 	const UDataflowSchema* GetSchema() { return Schema; }
 
+	/** FGCObject interface */
+	virtual void AddReferencedObjects(FReferenceCollector& Collector) override;
+	virtual FString GetReferencerName() const override { return TEXT("FDataflowConnectionDrawingPolicy"); }
+
 private:
-	class UDataflowSchema* Schema;
+	TObjectPtr<UDataflowSchema> Schema = nullptr;
 };
 
diff --git a/Engine/Plugins/Experimental/Dataflow/Source/DataflowEditor/Public/Dataflow/DataflowToolTarget.h b/Engine/Plugins/Experimental/Dataflow/Source/DataflowEditor/Public/Dataflow/DataflowToolTarget.h
--- a/Engine/Plugins/Experimental/Dataflow/Source/DataflowEditor/Public/Dataflow/DataflowToolTarget.h
+++ b/Engine/Plugins/Experimental/Dataflow/Source/DataflowEditor/Public/Dataflow/DataflowToolTarget.h
@@ -56,9 +56,11 @@ protected:
 	TUniquePtr<FMeshDescription> DataflowMeshDescription;
 
 	// Internal dataflow pointer
+	UPROPERTY()
 	TObjectPtr<UDataflow> Dataflow = nullptr;
 
 	// Internal asset pointer
+	UPROPERTY()
 	TObjectPtr<UObject> Asset = nullptr;
 
 	/**  Engine context to be used for dataflow evaluation */
diff --git a/Engine/Plugins/Experimental/Dataflow/Source/DataflowEditor/Public/Dataflow/DataflowView.h b/Engine/Plugins/Experimental/Dataflow/Source/DataflowEditor/Public/Dataflow/DataflowView.h
--- a/Engine/Plugins/Experimental/Dataflow/Source/DataflowEditor/Public/Dataflow/DataflowView.h
+++ b/Engine/Plugins/Experimental/Dataflow/Source/DataflowEditor/Public/Dataflow/DataflowView.h
@@ -27,7 +27,7 @@ public:
 * FDataflowNodeView class implements common functions for single node based Dataflow views
 *
 */
-class FDataflowNodeView : public IDataflowViewListener
+class FDataflowNodeView : public IDataflowViewListener, public FGCObject
 {
 public:
 	virtual ~FDataflowNodeView();
@@ -62,8 +62,13 @@ public:
 	virtual void OnSelectedNodeChanged(UDataflowEdNode* InNode) override;  // nullptr is valid
 	virtual void OnNodeInvalidated(FDataflowNode* InvalidatedNode) override;
 
+
+	/** FGCObject interface */
+	virtual void AddReferencedObjects(FReferenceCollector& Collector) override;
+	virtual FString GetReferencerName() const override { return TEXT("FDataflowNodeView"); }
+
 private:
-	UDataflowEdNode* SelectedNode = nullptr;
+	TObjectPtr<UDataflowEdNode> SelectedNode = nullptr;
 	TSharedPtr<Dataflow::FContext> Context;
 	bool bIsPinnedDown = false;
 	bool bIsRefreshLocked = false;
