Dataflow : Per vertex shading

diff --git a/Engine/Plugins/Experimental/ChaosFlesh/Source/ChaosFlesh/Private/ChaosFlesh/TetrahedralCollection.cpp b/Engine/Plugins/Experimental/ChaosFlesh/Source/ChaosFlesh/Private/ChaosFlesh/TetrahedralCollection.cpp
--- a/Engine/Plugins/Experimental/ChaosFlesh/Source/ChaosFlesh/Private/ChaosFlesh/TetrahedralCollection.cpp
+++ b/Engine/Plugins/Experimental/ChaosFlesh/Source/ChaosFlesh/Private/ChaosFlesh/TetrahedralCollection.cpp
@@ -17,7 +17,7 @@ const FName FTetrahedralCollection::IncidentElementsLocalIndexAttribute("Inciden
 const FName FTetrahedralCollection::GuidAttribute("Guid");
 
 FTetrahedralCollection::FTetrahedralCollection()
-	: Super::FGeometryCollection()
+	: Super::FGeometryCollection({ FLinearColor(0.6, 0.6, 0.6).ToRGBE() })
 {
 	Construct();
 }
diff --git a/Engine/Plugins/Experimental/Dataflow/Source/DataflowAssetTools/Private/Dataflow/CollectionRenderingPatternUtility.cpp b/Engine/Plugins/Experimental/Dataflow/Source/DataflowAssetTools/Private/Dataflow/CollectionRenderingPatternUtility.cpp
--- a/Engine/Plugins/Experimental/Dataflow/Source/DataflowAssetTools/Private/Dataflow/CollectionRenderingPatternUtility.cpp
+++ b/Engine/Plugins/Experimental/Dataflow/Source/DataflowAssetTools/Private/Dataflow/CollectionRenderingPatternUtility.cpp
@@ -30,18 +30,38 @@ namespace Dataflow
 				const TManagedArray<FVector3f>& Normals = Facade.GetNormals();
 				const TManagedArray<FLinearColor>& Colors = Facade.GetVertexColor();
 
-				DynamicMesh.Clear();
 				for (int32 VertexIndex = 0; VertexIndex < NumVertices; ++VertexIndex)
 				{
 					DynamicMesh.AppendVertex(FVertexInfo(FVector3d(Positions[VertexIndex]), Normals[VertexIndex],
-											 FVector3f(Colors[VertexIndex].R, Colors[VertexIndex].G, Colors[VertexIndex].B)));
+						FVector3f(Colors[VertexIndex].R, Colors[VertexIndex].G, Colors[VertexIndex].B)));
 				}
 				for (int32 TriangleIndex = 0; TriangleIndex < NumTriangles; ++TriangleIndex)
 				{
 					DynamicMesh.AppendTriangle(FIndex3i(Indices[TriangleIndex].X, Indices[TriangleIndex].Y, Indices[TriangleIndex].Z));
 				}
-
 				FMeshNormals::QuickComputeVertexNormals(DynamicMesh);
+
+				DynamicMesh.EnableAttributes();
+				DynamicMesh.Attributes()->EnablePrimaryColors();
+				DynamicMesh.Attributes()->PrimaryColors()->CreateFromPredicate([](int ParentVID, int TriIDA, int TriIDB) {return true; }, 0.f);
+				DynamicMesh.EnableVertexColors(FVector3f::Zero());
+
+				FDynamicMeshColorOverlay* const ColorOverlay = DynamicMesh.Attributes()->PrimaryColors();
+
+				auto SetColorsFromWeights = [&](int TriangleID)
+				{
+					const FIndex3i Tri = DynamicMesh.GetTriangle(TriangleID);
+					const FIndex3i ColorElementTri = ColorOverlay->GetTriangle(TriangleID);
+					for (int TriVertIndex = 0; TriVertIndex < 3; ++TriVertIndex)
+					{
+						FVector4f Color(Colors[Tri[TriVertIndex]]); Color.W = 1.0f;
+						ColorOverlay->SetElement(ColorElementTri[TriVertIndex], Color);
+					}
+				};
+				for (const int TriangleID : DynamicMesh.TriangleIndicesItr())
+				{
+					SetColorsFromWeights(TriangleID);
+				}
 			}
 		}
 
diff --git a/Engine/Plugins/Experimental/Dataflow/Source/DataflowEditor/Private/Dataflow/DataflowEditorStyle.h b/Engine/Plugins/Experimental/Dataflow/Source/DataflowEditor/Private/Dataflow/DataflowEditorStyle.h
--- a/Engine/Plugins/Experimental/Dataflow/Source/DataflowEditor/Private/Dataflow/DataflowEditorStyle.h
+++ b/Engine/Plugins/Experimental/Dataflow/Source/DataflowEditor/Private/Dataflow/DataflowEditorStyle.h
@@ -48,7 +48,7 @@ public:
 		Set(*PropertyNameString, new IMAGE_BRUSH_SVG("Icons/PaintMaps", ToolbarIconSize));
 
 		DefaultMaterial = Cast<UMaterial>(StaticLoadObject(UMaterial::StaticClass(), NULL, TEXT("/Engine/BasicShapes/BasicShapeMaterial")));
-
+		VertexMaterial = Cast<UMaterial>(StaticLoadObject(UMaterial::StaticClass(), NULL, TEXT("/Engine/EditorMaterials/Dataflow/DataflowVertexMaterial")));
 		FSlateStyleRegistry::RegisterSlateStyle(*this);
 	}
 
@@ -60,6 +60,8 @@ public:
 
 	/** Default Rendering Material for Mesh surfaces */
 	UMaterial* DefaultMaterial = nullptr;
+	UMaterial* VertexMaterial = nullptr;
+
 
 public:
 
diff --git a/Engine/Plugins/Experimental/Dataflow/Source/DataflowEditor/Private/Dataflow/DataflowPreviewScene.cpp b/Engine/Plugins/Experimental/Dataflow/Source/DataflowEditor/Private/Dataflow/DataflowPreviewScene.cpp
--- a/Engine/Plugins/Experimental/Dataflow/Source/DataflowEditor/Private/Dataflow/DataflowPreviewScene.cpp
+++ b/Engine/Plugins/Experimental/Dataflow/Source/DataflowEditor/Private/Dataflow/DataflowPreviewScene.cpp
@@ -166,15 +166,18 @@ TObjectPtr<UDynamicMeshComponent>& FDataflowPreviewScene::AddDynamicMeshComponen
 	DynamicMeshComponent->SetMesh(MoveTemp(DynamicMesh));
 
 	// @todo(Dataflow) : Material support
-	// This is just to have a material, we should transfer the materials from the assets if they have them. 
-	if (FDataflowEditorStyle::Get().DefaultMaterial)
-	{
-		DynamicMeshComponent->ConfigureMaterialSet({ FDataflowEditorStyle::Get().DefaultMaterial });
-	}
-	else
-	{
-		DynamicMeshComponent->ValidateMaterialSlots(true, false);
-	}
+	//if (FDataflowEditorStyle::Get().DefaultMaterial)
+	//{
+	//	DynamicMeshComponent->ConfigureMaterialSet({ FDataflowEditorStyle::Get().DefaultMaterial });
+	//}
+	//else
+	//{
+	//	DynamicMeshComponent->ValidateMaterialSlots(true, false);
+	//}
+	// just vertex colors atm
+	DynamicMeshComponent->SetOverrideRenderMaterial(FDataflowEditorStyle::Get().VertexMaterial);
+	DynamicMeshComponent->SetShadowsEnabled(false);
+
 
 	DynamicMeshComponent->SelectionOverrideDelegate = UPrimitiveComponent::FSelectionOverride::CreateRaw(this, &FDataflowPreviewScene::IsComponentSelected);
 	DynamicMeshComponent->RegisterComponentWithWorld(GetWorld());
