Simulation viewport + Core dataflow content

diff --git a/Engine/Plugins/Experimental/Dataflow/Source/DataflowAssetTools/Public/Dataflow/CollectionPatternToDynamicMesh.h b/Engine/Plugins/Experimental/Dataflow/Source/DataflowAssetTools/Public/Dataflow/CollectionPatternToDynamicMesh.h
--- a/Engine/Plugins/Experimental/Dataflow/Source/DataflowAssetTools/Public/Dataflow/CollectionPatternToDynamicMesh.h
+++ b/Engine/Plugins/Experimental/Dataflow/Source/DataflowAssetTools/Public/Dataflow/CollectionPatternToDynamicMesh.h
@@ -2,7 +2,7 @@
 
 #pragma once
 
-#include "DataflowPatternVertexType.h"
+#include "Dataflow/DataflowContent.h"
 #include "Templates/SharedPointer.h"
 
 
diff --git a/Engine/Plugins/Experimental/Dataflow/Source/DataflowAssetTools/Public/Dataflow/DataflowPatternVertexType.h b/Engine/Plugins/Experimental/Dataflow/Source/DataflowAssetTools/Public/Dataflow/DataflowPatternVertexType.h
deleted file mode 100644
diff --git a/Engine/Plugins/Experimental/Dataflow/Source/DataflowEditor/DataflowEditor.Build.cs b/Engine/Plugins/Experimental/Dataflow/Source/DataflowEditor/DataflowEditor.Build.cs
--- a/Engine/Plugins/Experimental/Dataflow/Source/DataflowEditor/DataflowEditor.Build.cs
+++ b/Engine/Plugins/Experimental/Dataflow/Source/DataflowEditor/DataflowEditor.Build.cs
@@ -60,7 +60,10 @@ namespace UnrealBuildTool.Rules
 					"TypedElementRuntime",
 					"UnrealEd",
 					"WorkspaceMenuStructure",
-					"XmlParser"
+					"XmlParser",
+					"EditorWidgets",
+					"KismetWidgets",      // SScrubControlPanel
+					"AnimGraph"           // UAnimSingleNodeInstance
 				}
 			);
 		}
diff --git a/Engine/Plugins/Experimental/Dataflow/Source/DataflowEditor/Private/Dataflow/AssetDefinition_DataflowAsset.cpp b/Engine/Plugins/Experimental/Dataflow/Source/DataflowEditor/Private/Dataflow/AssetDefinition_DataflowAsset.cpp
--- a/Engine/Plugins/Experimental/Dataflow/Source/DataflowEditor/Private/Dataflow/AssetDefinition_DataflowAsset.cpp
+++ b/Engine/Plugins/Experimental/Dataflow/Source/DataflowEditor/Private/Dataflow/AssetDefinition_DataflowAsset.cpp
@@ -184,7 +184,7 @@ EAssetCommandResult UAssetDefinition_DataflowAsset::OpenAssets(const FAssetOpenA
 		// Validate the asset
 		if (UDataflow* const DataflowAsset = CastChecked<UDataflow>(DataflowObjects[0]))
 		{
-			AssetEditor->Initialize({ DataflowAsset });
+			AssetEditor->InitializeContent(nullptr, DataflowAsset);
 			return EAssetCommandResult::Handled;
 		}
 	}
diff --git a/Engine/Plugins/Experimental/Dataflow/Source/DataflowEditor/Private/Dataflow/DataflowEditor.cpp b/Engine/Plugins/Experimental/Dataflow/Source/DataflowEditor/Private/Dataflow/DataflowEditor.cpp
--- a/Engine/Plugins/Experimental/Dataflow/Source/DataflowEditor/Private/Dataflow/DataflowEditor.cpp
+++ b/Engine/Plugins/Experimental/Dataflow/Source/DataflowEditor/Private/Dataflow/DataflowEditor.cpp
@@ -3,7 +3,7 @@
 #include "Dataflow/DataflowEditor.h"
 
 #include "Animation/Skeleton.h"
-#include "Dataflow/DataflowEditorContent.h"
+#include "Dataflow/DataflowContent.h"
 #include "Dataflow/DataflowEditorToolkit.h"
 #include "Dataflow/DataflowEditorUtil.h"
 #include "Engine/SkeletalMesh.h"
@@ -12,12 +12,8 @@
 
 DEFINE_LOG_CATEGORY(LogDataflowEditor);
 
-
-
 UDataflowEditor::UDataflowEditor() : Super()
-{
-	Content = NewObject<UDataflowEditorContent>(this, MakeUniqueObjectName(this, UDataflowEditor::StaticClass(), "EditorData"));
-}
+{}
 
 TSharedPtr<FBaseAssetToolkit> UDataflowEditor::CreateToolkit()
 {
@@ -27,48 +23,82 @@ TSharedPtr<FBaseAssetToolkit> UDataflowEditor::CreateToolkit()
 
 void UDataflowEditor::Initialize(const TArray<TObjectPtr<UObject>>& InObjects)
 {
-	check(Content);
+	if(!InObjects.IsEmpty())
+	{
+		InitializeContent(nullptr, InObjects[0]);
+	}
+}
 
-	TArray<TObjectPtr<UObject>> ObjectsToEdit;
-	if (ensure(InObjects.Num() == 1))
+void UDataflowEditor::InitializeContent(TObjectPtr<UDataflowBaseContent> BaseContent, const TObjectPtr<UObject>& ContentOwner)
+{
+	DataflowContent = BaseContent;
+	if(!DataflowContent)
 	{
-		TObjectPtr<UObject> RootObject = InObjects[0];
-		Content->SetDataflowAsset(Cast<UDataflow>(RootObject));
-		if(!Content->GetDataflowAsset())
+		if(UDataflow* DataflowAsset = Cast<UDataflow>(ContentOwner))
 		{
-			Content->SetDataflowAsset(Private::GetDataflowAssetFrom(RootObject));
-			Content->SetDataflowTerminal(Private::GetDataflowTerminalFrom(RootObject));
-			Content->SetSkeletalMesh(Private::GetSkeletalMeshFrom(RootObject));
-			Content->SetSkeleton(Private::GetSkeletonFrom(RootObject));
-			Content->SetAnimationAsset(Private::GetAnimationAssetFrom(RootObject));
+			DataflowContent = NewObject<UDataflowBaseContent>();
+			
+			DataflowContent->SetDataflowAsset(DataflowAsset);
+			DataflowContent->SetDataflowTerminal(FString());
 		}
-		if(Content->GetDataflowAsset())
+		else
 		{
-			Content->SetDataflowOwner(RootObject);
-			Content->GetDataflowAsset()->Schema = UDataflowSchema::StaticClass();
-			Content->SetDataflowContext(MakeShared<Dataflow::FEngineContext>(RootObject, Content->GetDataflowAsset(), FPlatformTime::Cycles64()));
-			Content->SetLastModifiedTimestamp(Content->GetDataflowContext()->GetTimestamp());
-
-			ObjectsToEdit.Add(Content->GetDataflowOwner());
-
-			if(!Content->GetSkeletalMesh())
+			if(Private::HasDataflowAsset(ContentOwner))
 			{
-				const FName SkeletonName = MakeUniqueObjectName(Content->GetDataflowAsset(), UDataflow::StaticClass(), FName("USkeleton"));
-				const FName SkeletalMeshName = MakeUniqueObjectName(Content->GetDataflowAsset(), UDataflow::StaticClass(), FName("USkeletalMesh"));
-
-				Content->SetSkeleton(NewObject<USkeleton>(Content->GetDataflowAsset(), SkeletonName));
-				Content->SetSkeletalMesh(NewObject<USkeletalMesh>(Content->GetDataflowAsset(), SkeletalMeshName));
-				Content->GetSkeletalMesh()->SetSkeleton(Content->GetSkeleton());
+				if(Private::HasSkeletalMesh(ContentOwner))
+				{
+					DataflowContent = NewObject<UDataflowSkeletalContent>();
+					const TObjectPtr<UDataflowSkeletalContent> SkeletalContent = Cast<UDataflowSkeletalContent>(DataflowContent);
+					
+					SkeletalContent->SetSkeletalMesh(Private::GetSkeletalMeshFrom(ContentOwner));
+                    SkeletalContent->SetSkeleton(Private::GetSkeletonFrom(ContentOwner));
+                    SkeletalContent->SetAnimationAsset(Private::GetAnimationAssetFrom(ContentOwner));
+				}
+				else
+				{
+					DataflowContent = NewObject<UDataflowBaseContent>();
+				}
+				DataflowContent->SetDataflowAsset(Private::GetDataflowAssetFrom(ContentOwner));
+				DataflowContent->SetDataflowTerminal(Private::GetDataflowTerminalFrom(ContentOwner));
 			}
-			else if(!Content->GetSkeleton())
+		}
+	}
+
+	if(const TObjectPtr<UDataflowSkeletalContent> SkeletalContent = Cast<UDataflowSkeletalContent>(DataflowContent))
+	{
+		if(!SkeletalContent->GetSkeletalMesh())
+		{
+			const FName SkeletalMeshName = MakeUniqueObjectName(SkeletalContent->GetDataflowAsset(), UDataflow::StaticClass(), FName("USkeletalMesh"));
+			USkeletalMesh* SkeletalMesh = NewObject<USkeletalMesh>(SkeletalContent->GetDataflowAsset(), SkeletalMeshName);
+
+			USkeleton* Skeleton = SkeletalContent->GetSkeleton();
+			if(!Skeleton)
 			{
-				Content->SetSkeleton(Content->GetSkeletalMesh()->GetSkeleton());
+				const FName SkeletonName = MakeUniqueObjectName(SkeletalContent->GetDataflowAsset(), UDataflow::StaticClass(), FName("USkeleton"));
+                Skeleton = NewObject<USkeleton>(SkeletalContent->GetDataflowAsset(), SkeletonName);
 			}
+			SkeletalMesh->SetSkeleton(Skeleton);
+			SkeletalContent->SetSkeletalMesh(SkeletalMesh);
 		}
+		else if(!SkeletalContent->GetSkeleton())
+		{
+			SkeletalContent->SetSkeleton(SkeletalContent->GetSkeletalMesh()->GetSkeleton());
+		}
+	}
+
+	if(DataflowContent && DataflowContent->GetDataflowAsset())
+	{
+		DataflowContent->GetDataflowAsset()->Schema = UDataflowSchema::StaticClass();
+		DataflowContent->BuildBaseContent(ContentOwner);
+		
+		// the owner is either a base content derived class or the dataflow itself
+		DataflowContent->SetDataflowOwner(ContentOwner);
 	}
+	
 	// Potentially we could add additional objects to edit here (fields, meshes....)
 	// If these objects have a matching factory we would be able to use geometry tools
-	UBaseCharacterFXEditor::Initialize(ObjectsToEdit);
+	const TArray<TObjectPtr<UObject>> ArrayObjects = {ContentOwner};
+	UBaseCharacterFXEditor::Initialize(ArrayObjects);
 }
 
 
diff --git a/Engine/Plugins/Experimental/Dataflow/Source/DataflowEditor/Private/Dataflow/DataflowEditorCommands.cpp b/Engine/Plugins/Experimental/Dataflow/Source/DataflowEditor/Private/Dataflow/DataflowEditorCommands.cpp
--- a/Engine/Plugins/Experimental/Dataflow/Source/DataflowEditor/Private/Dataflow/DataflowEditorCommands.cpp
+++ b/Engine/Plugins/Experimental/Dataflow/Source/DataflowEditor/Private/Dataflow/DataflowEditorCommands.cpp
@@ -3,7 +3,7 @@
 #include "Dataflow/DataflowEditorCommands.h"
 
 #include "Dataflow/DataflowEdNode.h"
-#include "Dataflow/DataflowEditorContent.h"
+#include "Dataflow/DataflowContent.h"
 #include "Dataflow/DataflowEditorStyle.h"
 #include "Dataflow/DataflowEditorUtil.h"
 #include "Dataflow/DataflowGraphEditor.h"
@@ -294,11 +294,11 @@ void FDataflowEditorCommands::OnNodeTitleCommitted(const FText& InNewText, EText
 	}
 }
 
-void FDataflowEditorCommands::OnAssetPropertyValueChanged(TObjectPtr<UDataflowEditorContent> Content, const FPropertyChangedEvent& InPropertyChangedEvent)
+void FDataflowEditorCommands::OnAssetPropertyValueChanged(TObjectPtr<UDataflowBaseContent> Content, const FPropertyChangedEvent& InPropertyChangedEvent)
 {
 	if (Content)
 	{
-		TObjectPtr<UDataflow>& DataflowAsset = Content->GetDataflowAsset();
+		const TObjectPtr<UDataflow>& DataflowAsset = Content->GetDataflowAsset();
 		if (DataflowAsset)
 		{
 			if (InPropertyChangedEvent.ChangeType == EPropertyChangeType::ValueSet ||
@@ -322,14 +322,6 @@ void FDataflowEditorCommands::OnAssetPropertyValueChanged(TObjectPtr<UDataflowEd
 						}
 					}
 				}
-				else if (InPropertyChangedEvent.GetPropertyName() == FName("SkeletalMesh"))
-				{
-					Content->SetSkeletalMesh(Private::GetSkeletalMeshFrom(Content->GetDataflowOwner()));
-				}
-				else if (InPropertyChangedEvent.GetPropertyName() == FName("Skeleton"))
-				{
-					Content->SetSkeleton(Private::GetSkeletonFrom(Content->GetDataflowOwner()));
-				}
 			}
 		}
 	}
diff --git a/Engine/Plugins/Experimental/Dataflow/Source/DataflowEditor/Private/Dataflow/DataflowEditorContent.cpp b/Engine/Plugins/Experimental/Dataflow/Source/DataflowEditor/Private/Dataflow/DataflowEditorContent.cpp
deleted file mode 100644
diff --git a/Engine/Plugins/Experimental/Dataflow/Source/DataflowEditor/Private/Dataflow/DataflowEditorMode.cpp b/Engine/Plugins/Experimental/Dataflow/Source/DataflowEditor/Private/Dataflow/DataflowEditorMode.cpp
--- a/Engine/Plugins/Experimental/Dataflow/Source/DataflowEditor/Private/Dataflow/DataflowEditorMode.cpp
+++ b/Engine/Plugins/Experimental/Dataflow/Source/DataflowEditor/Private/Dataflow/DataflowEditorMode.cpp
@@ -11,7 +11,7 @@
 #include "Dataflow/DataflowComponentToolTarget.h"
 #include "Dataflow/DataflowCollectionAddScalarVertexPropertyNode.h"
 #include "Dataflow/DataflowEditor.h"
-#include "Dataflow/DataflowEditorContent.h"
+#include "Dataflow/DataflowContent.h"
 #include "Dataflow/DataflowEditorCommands.h"
 #include "Dataflow/DataflowEditorModeToolkit.h"
 #include "Dataflow/DataflowEditorViewportClient.h"
@@ -251,7 +251,7 @@ void UDataflowEditorMode::OnToolEnded(UInteractiveToolManager* Manager, UInterac
 	}
 	else
 	{
-		PreviewScene->ReinitializeDynamicMeshComponents();
+		PreviewScene->ResetConstructionScene();
 	}
 
 	if (TSharedPtr<SDataflowGraphEditor> GraphEditor = DataflowGraphEditor.Pin())
@@ -291,13 +291,13 @@ void UDataflowEditorMode::Exit()
 {
 	UActorComponent::MarkRenderStateDirtyEvent.RemoveAll(this);
 
-	PreviewScene->Exit();
+	PreviewScene->ResetConstructionScene();
 	PreviewScene = nullptr;
 
 	Super::Exit();
 }
 
-void UDataflowEditorMode::SetDataflowPreviewScene(FDataflowPreviewScene* InPreviewScene)
+void UDataflowEditorMode::SetDataflowConstructionScene(FDataflowConstructionScene* InPreviewScene)
 {
 	PreviewScene = InPreviewScene;
 
@@ -322,9 +322,9 @@ void UDataflowEditorMode::SetDataflowPreviewScene(FDataflowPreviewScene* InPrevi
 void UDataflowEditorMode::CreateToolTargets(const TArray<TObjectPtr<UObject>>& AssetsIn)
 {
 	ToolTargets.Reset();
-	if (TObjectPtr<UDataflowEditorContent> EditorContent = PreviewScene->GetDataflowEditorContent())
+	if (TObjectPtr<UDataflowBaseContent> EditorContent = PreviewScene->GetDataflowContent())
 	{
-		if (UToolTarget* Target = GetInteractiveToolsContext()->TargetManager->BuildTarget(EditorContent.Get(), GetToolTargetRequirements()))
+		if (UToolTarget* Target = GetInteractiveToolsContext()->TargetManager->BuildTarget(EditorContent, GetToolTargetRequirements()))
 		{
 			ToolTargets.Add(Target);
 		}
@@ -351,7 +351,7 @@ bool UDataflowEditorMode::IsComponentSelected(const UPrimitiveComponent* InCompo
 void UDataflowEditorMode::SetSelectedCollection(TSharedPtr<FManagedArrayCollection> Collection)
 {
 	SelectedCollection = Collection;
-	PreviewScene->ReinitializeDynamicMeshComponents();
+	PreviewScene->UpdateConstructionScene();
 
 	// The first time we get a valid mesh, refocus the camera on it
 	FirstTimeFocusRestSpaceViewport();
@@ -418,11 +418,6 @@ void UDataflowEditorMode::InitializeTargets(const TArray<TObjectPtr<UObject>>& O
 {
 	UBaseCharacterFXEditorMode::InitializeTargets(ObjectsToEdit);
 
-	USelection* SelectedComponents = GetModeManager()->GetSelectedComponents();
-	SelectedComponents->Modify();
-	SelectedComponents->BeginBatchSelectOperation();
-	SelectedComponents->DeselectAll();
-
 	// @todo(brice) : Consider initializing the Content here from the ObjectsToEdit
 
 	// @todo(brice) : What are the ToolTargets storing?
@@ -430,19 +425,11 @@ void UDataflowEditorMode::InitializeTargets(const TArray<TObjectPtr<UObject>>& O
 	// ... UE::ToolTarget::GetDynamicMeshCopy(Target)
 	// ... UE::ToolTarget::GetMaterialSet(Target).Materials for PreviewScene->AddDynamicMeshComponent
 	// ... }
-	PreviewScene->Update();
 
-	for (TObjectPtr<UDynamicMeshComponent>& Component : PreviewScene->GetDynamicMeshComponents())
-	{
-		if (Component)
-		{
-			SelectedComponents->Select(Component);
-		}
-	}
-	SelectedComponents->EndBatchSelectOperation();
+	// @todo(michael) : do we need to update the construction scene?
+	PreviewScene->UpdateConstructionScene();
 }
 
-
 void UDataflowEditorMode::ModeTick(float DeltaTime)
 {
 	Super::ModeTick(DeltaTime);
@@ -473,12 +460,6 @@ void UDataflowEditorMode::ModeTick(float DeltaTime)
 
 		NodeTypeForPendingToolStart = FName();
 	}
-
-	if (PreviewScene)
-	{
-		PreviewScene->GetWorld()->Tick(ELevelTick::LEVELTICK_All, DeltaTime);
-	}
-
 }
 
 void UDataflowEditorMode::RestSpaceViewportResized(FViewport* RestspaceViewport, uint32 /*Unused*/)
@@ -499,7 +480,8 @@ FBox UDataflowEditorMode::SceneBoundingBox() const
 
 FBox UDataflowEditorMode::SelectionBoundingBox() const
 {
-	FBox Bounds = PreviewScene->SelectedComponentBounds();
+	// If the selection is on the GetBoundingBox is automatically computing the selection one
+	FBox Bounds = PreviewScene->GetBoundingBox();
 	if (Bounds.IsValid)
 	{
 		return Bounds;
@@ -529,7 +511,7 @@ void UDataflowEditorMode::SetConstructionViewMode(Dataflow::EDataflowPatternVert
 	}
 
 	ConstructionViewMode = InMode;
-	PreviewScene->ReinitializeDynamicMeshComponents();
+	PreviewScene->UpdateConstructionScene();
 
 	const TSharedPtr<FDataflowEditorViewportClient> VC = ConstructionViewportClient.Pin();
 	if (VC.IsValid())
@@ -587,7 +569,7 @@ void UDataflowEditorMode::ToggleConstructionViewWireframe()
 {
 	check(false);
 	bConstructionViewWireframe = !bConstructionViewWireframe;
-	PreviewScene->ReinitializeDynamicMeshComponents();
+	PreviewScene->UpdateConstructionScene();
 }
 
 bool UDataflowEditorMode::CanSetConstructionViewWireframeActive() const
@@ -626,29 +608,29 @@ void UDataflowEditorMode::InitializeContextObject()
 {
 	check(PreviewScene);
 
-	if (TObjectPtr<UDataflowEditorContent> DataflowEditorContent = PreviewScene->GetDataflowEditorContent())
+	if (TObjectPtr<UDataflowBaseContent> DataflowContent = PreviewScene->GetDataflowContent())
 	{
 		UEditorInteractiveToolsContext* const RestSpaceToolsContext = GetInteractiveToolsContext();
 
-		UDataflowEditorContent* EditorContextObject = RestSpaceToolsContext->ContextObjectStore->FindContext<UDataflowEditorContent>();
-		if (!EditorContextObject)
+		UDataflowContextObject* ContextObject = RestSpaceToolsContext->ContextObjectStore->FindContext<UDataflowContextObject>();
+		if (!ContextObject)
 		{
-			EditorContextObject = PreviewScene->GetDataflowEditorContent();
-			RestSpaceToolsContext->ContextObjectStore->AddContextObject(PreviewScene->GetDataflowEditorContent());
+			ContextObject = DataflowContent;
+			RestSpaceToolsContext->ContextObjectStore->AddContextObject(ContextObject);
 		}
 
-		check(EditorContextObject);
+		check(ContextObject);
 
-		EditorContextObject->SetConstructionViewMode(ConstructionViewMode);
+		ContextObject->SetConstructionViewMode(ConstructionViewMode);
 	}
 }
 
 void UDataflowEditorMode::DeleteContextObject()
 {
 	UEditorInteractiveToolsContext* const RestSpaceToolsContext = GetInteractiveToolsContext();
-	if (UDataflowEditorContent* EditorContextObject = RestSpaceToolsContext->ContextObjectStore->FindContext<UDataflowEditorContent>())
+	if (UDataflowContextObject* ContextObject = RestSpaceToolsContext->ContextObjectStore->FindContext<UDataflowContextObject>())
 	{
-		RestSpaceToolsContext->ContextObjectStore->RemoveContextObject(EditorContextObject);
+		RestSpaceToolsContext->ContextObjectStore->RemoveContextObject(ContextObject);
 	}
 }
 
@@ -677,7 +659,6 @@ void UDataflowEditorMode::StartToolForSelectedNode(const UObject* SelectedNode)
 	}
 }
 
-
 void UDataflowEditorMode::OnDataflowNodeDeleted(const TSet<UObject*>& DeletedNodes)
 {
 	UEditorInteractiveToolsContext* const ToolsContext = GetInteractiveToolsContext();
@@ -728,7 +709,7 @@ UEdGraphNode* UDataflowEditorMode::CreateNewNode(const FName& NewNodeTypeName)
 		return nullptr;
 	}
 
-	if (TObjectPtr<UDataflowEditorContent> EditorContent = PreviewScene->GetDataflowEditorContent())
+	if (TObjectPtr<UDataflowBaseContent> EditorContent = PreviewScene->GetDataflowContent())
 	{
 		if (TObjectPtr<UDataflow> DataflowGraph = EditorContent->GetDataflowAsset())
 		{
@@ -745,7 +726,7 @@ UEdGraphNode* UDataflowEditorMode::CreateNewNode(const FName& NewNodeTypeName)
 
 UEdGraphNode* UDataflowEditorMode::CreateAndConnectNewNode(const FName& NewNodeTypeName, UEdGraphNode& UpstreamNode, const FName& ConnectionTypeName, const FName& NewNodeConnectionName)
 {
-	if (TObjectPtr<UDataflowEditorContent> EditorContent = PreviewScene->GetDataflowEditorContent())
+	if (TObjectPtr<UDataflowBaseContent> EditorContent = PreviewScene->GetDataflowContent())
 	{
 		if (TObjectPtr<UDataflow> DataflowGraph = EditorContent->GetDataflowAsset())
 		{
diff --git a/Engine/Plugins/Experimental/Dataflow/Source/DataflowEditor/Private/Dataflow/DataflowEditorModeToolkit.cpp b/Engine/Plugins/Experimental/Dataflow/Source/DataflowEditor/Private/Dataflow/DataflowEditorModeToolkit.cpp
--- a/Engine/Plugins/Experimental/Dataflow/Source/DataflowEditor/Private/Dataflow/DataflowEditorModeToolkit.cpp
+++ b/Engine/Plugins/Experimental/Dataflow/Source/DataflowEditor/Private/Dataflow/DataflowEditorModeToolkit.cpp
@@ -158,7 +158,7 @@ SBaseCharacterFXEditorViewport* FDataflowEditorModeToolkit::GetViewportWidgetFor
 
 		if (const UDataflowEditorMode* const DataflowEdMode = Cast<UDataflowEditorMode>(Mode))
 		{
-			if (const FDataflowPreviewScene* const PreviewScene = DataflowEdMode->GetDataflowPreviewScene())
+			if (const FDataflowPreviewScene* const PreviewScene = DataflowEdMode->GetDataflowConstructionScene())
 			{
 				if (const UEditorInteractiveToolsContext* const PreviewToolsContext = PreviewScene->GetDataflowModeManager()->GetInteractiveToolsContext())
 				{
diff --git a/Engine/Plugins/Experimental/Dataflow/Source/DataflowEditor/Private/Dataflow/DataflowEditorToolkit.cpp b/Engine/Plugins/Experimental/Dataflow/Source/DataflowEditor/Private/Dataflow/DataflowEditorToolkit.cpp
--- a/Engine/Plugins/Experimental/Dataflow/Source/DataflowEditor/Private/Dataflow/DataflowEditorToolkit.cpp
+++ b/Engine/Plugins/Experimental/Dataflow/Source/DataflowEditor/Private/Dataflow/DataflowEditorToolkit.cpp
@@ -6,7 +6,7 @@
 #include "Animation/Skeleton.h"
 #include "Dataflow/DataflowCore.h"
 #include "Dataflow/DataflowEditor.h"
-#include "Dataflow/DataflowEditorContent.h"
+#include "Dataflow/DataflowContent.h"
 #include "Dataflow/DataflowEditorCommands.h"
 #include "Dataflow/DataflowEditorMode.h"
 #include "Dataflow/DataflowEditorModeToolkit.h"
@@ -55,13 +55,14 @@ const FName FDataflowEditorToolkit::CollectionSpreadSheetTabId_1(TEXT("DataflowE
 const FName FDataflowEditorToolkit::CollectionSpreadSheetTabId_2(TEXT("DataflowEditor_CollectionSpreadSheet_2"));
 const FName FDataflowEditorToolkit::CollectionSpreadSheetTabId_3(TEXT("DataflowEditor_CollectionSpreadSheet_3"));
 const FName FDataflowEditorToolkit::CollectionSpreadSheetTabId_4(TEXT("DataflowEditor_CollectionSpreadSheet_4"));
+const FName FDataflowEditorToolkit::SimulationViewportTabId(TEXT("DataflowEditor_SimulationViewport"));
 
 FDataflowEditorToolkit::FDataflowEditorToolkit(UAssetEditor* InOwningAssetEditor)
 	: FBaseCharacterFXEditorToolkit(InOwningAssetEditor, FName("DataflowEditor")) 
 {
 	check(Cast<UDataflowEditor>(InOwningAssetEditor));
 
-	StandaloneDefaultLayout = FTabManager::NewLayout(FName("DataflowEditorLayout.V9"))
+	StandaloneDefaultLayout = FTabManager::NewLayout(FName("DataflowEditorLayout.V10"))
 		->AddArea
 		(
 			FTabManager::NewPrimaryArea()->SetOrientation(Orient_Horizontal)
@@ -111,10 +112,11 @@ FDataflowEditorToolkit::FDataflowEditorToolkit(UAssetEditor* InOwningAssetEditor
 	FAdvancedPreviewScene::ConstructionValues PreviewSceneArgs;
 	PreviewSceneArgs.bShouldSimulatePhysics = 1;
 	PreviewSceneArgs.bCreatePhysicsScene = 1;
-	ObjectScene = MakeUnique<FDataflowPreviewScene>(PreviewSceneArgs, GetDataflowEditorContent());
+	
+	ObjectScene = MakeUnique<FDataflowConstructionScene>(PreviewSceneArgs, GetDataflowContent());
+	SimulationScene = MakeShared<FDataflowSimulationScene>(PreviewSceneArgs, GetDataflowContent());
 }
 
-
 FDataflowEditorToolkit::~FDataflowEditorToolkit()
 {
 	if (GraphEditor)
@@ -137,21 +139,30 @@ FDataflowEditorToolkit::~FDataflowEditorToolkit()
 	// will end up getting destroyed before the mode's Exit() function gets to run, and we'll get some
 	// warnings when we destroy any mode actors.
 	EditorModeManager->DestroyMode(UDataflowEditorMode::EM_DataflowEditorModeId);
-	
+	SimulationModeManager->DestroyMode(UDataflowEditorMode::EM_DataflowEditorModeId);
 }
 
 void FDataflowEditorToolkit::CreateEditorModeManager()
 {
+	// Setup the construction manager / scene
 	FBaseCharacterFXEditorToolkit::CreateEditorModeManager();
 	static_cast<FDataflowPreviewScene*>(ObjectScene.Get())->GetDataflowModeManager()
 		= StaticCastSharedPtr<FAssetEditorModeManager>(EditorModeManager);
+
+	// Setup the simulation manager / scene
+	SimulationModeManager = MakeShared<FAssetEditorModeManager>();
+	StaticCastSharedPtr<FAssetEditorModeManager>(SimulationModeManager)->SetPreviewScene(
+		SimulationScene.Get());
+
+	static_cast<FDataflowPreviewScene*>(SimulationScene.Get())->GetDataflowModeManager()
+		= StaticCastSharedPtr<FAssetEditorModeManager>(SimulationModeManager);
 }
 
 bool FDataflowEditorToolkit::CanOpenDataflowEditor(UObject* ObjectToEdit)
 {
 	if (const UClass* Class = ObjectToEdit->GetClass())
 	{
-		return Class->FindPropertyByName(FName("DataflowAsset")) != nullptr;
+		return (Class->FindPropertyByName(FName("DataflowAsset")) != nullptr) ;
 	}
 	return false;
 }
@@ -209,17 +220,16 @@ FEditorModeID FDataflowEditorToolkit::GetEditorModeId() const
 	return UDataflowEditorMode::EM_DataflowEditorModeId;
 }
 
-TObjectPtr<UDataflowEditorContent> FDataflowEditorToolkit::GetDataflowEditorContent()
+TObjectPtr<UDataflowBaseContent> FDataflowEditorToolkit::GetDataflowContent()
 {
-	return Cast<UDataflowEditor>(OwningAssetEditor)->Content;
+	return Cast<UDataflowEditor>(OwningAssetEditor)->DataflowContent;
 }
 
-TObjectPtr<const UDataflowEditorContent> FDataflowEditorToolkit::GetDataflowEditorContent() const
+const TObjectPtr<UDataflowBaseContent> FDataflowEditorToolkit::GetDataflowContent() const
 {
-	return Cast<UDataflowEditor>(OwningAssetEditor)->Content;
+	return Cast<UDataflowEditor>(OwningAssetEditor)->DataflowContent;
 }
 
-
 bool FDataflowEditorToolkit::OnRequestClose(EAssetEditorCloseReason InCloseReason)
 {
 	// Note: This needs a bit of adjusting, because currently OnRequestClose seems to be 
@@ -284,37 +294,30 @@ void FDataflowEditorToolkit::PostInitAssetEditor()
 
 		// We need the viewport client to start out focused, or else it won't get ticked until
 		// we click inside it.
-		Client->ReceivedFocus(Client->Viewport);
+		if(Client->Viewport)
+		{
+			Client->ReceivedFocus(Client->Viewport);
+		}
 	};
 	SetCommonViewportClientOptions(ViewportClient.Get());
+	SetCommonViewportClientOptions(SimulationViewportClient.Get());
 
 	//@todo(brice) : FChaosClothAssetEditorToolkit::PostInitAssetEditor() has alot more going on. 
 
-	// Set up 3D viewport
-	TSharedPtr<FDataflowEditorViewportClient> DataflowViewportClient = StaticCastSharedPtr<FDataflowEditorViewportClient>(ViewportClient);
-	DataflowViewportClient->SetDataflowEditorToolkit(StaticCastSharedRef<FDataflowEditorToolkit>(this->AsShared()));
-
 	//@todo(brice) : can we remove this coupling of the viewport and mode
 	UDataflowEditorMode* const DataflowMode = CastChecked<UDataflowEditorMode>(EditorModeManager->GetActiveScriptableMode(UDataflowEditorMode::EM_DataflowEditorModeId));
 	const TWeakPtr<FViewportClient> WeakViewportClient(ViewportClient);
 	DataflowMode->SetRestSpaceViewportClient(StaticCastWeakPtr<FDataflowEditorViewportClient>(WeakViewportClient));
-
-	check(GetDataflowEditorContent());
-
-
-	//SetEditingObject(GetDataflowEditorContent().DataflowOwner);
-
 }
 
 void FDataflowEditorToolkit::InitializeEdMode(UBaseCharacterFXEditorMode* EdMode)
 {
-
 	UDataflowEditorMode* DataflowMode = Cast<UDataflowEditorMode>(EdMode);
 	check(DataflowMode);
 
 	// We first set the preview scene in order to store the dynamic mesh elements
 	// generated by the tools
-	DataflowMode->SetDataflowPreviewScene(static_cast<FDataflowPreviewScene*>(ObjectScene.Get()));
+	DataflowMode->SetDataflowConstructionScene(static_cast<FDataflowConstructionScene*>(ObjectScene.Get()));
 
 	// Set of the graph editor to be able to add nodes
 	DataflowMode->SetDataflowGraphEditor(GraphEditor);
@@ -327,7 +330,6 @@ void FDataflowEditorToolkit::InitializeEdMode(UBaseCharacterFXEditorMode* EdMode
 	{
 		FDataflowEditorModeToolkit* DataflowModeToolkit = static_cast<FDataflowEditorModeToolkit*>(ModeToolkit.Get());
 		DataflowModeToolkit->SetRestSpaceViewportWidget(DataflowEditorViewport);
-		//ClothModeToolkit->SetPreviewViewportWidget(PreviewViewportWidget);
 
 		FName ParentToolbarName;
 		const FName ToolBarName = GetToolMenuToolbarName(ParentToolbarName);
@@ -335,11 +337,9 @@ void FDataflowEditorToolkit::InitializeEdMode(UBaseCharacterFXEditorMode* EdMode
 	}
 
 	// @todo(brice) : This used to crash when comnmented out. 
-	//FBaseCharacterFXEditorToolkit::InitializeEdMode(EdMode);
+	FBaseCharacterFXEditorToolkit::InitializeEdMode(EdMode);
 }
 
-
-
 void FDataflowEditorToolkit::CreateEditorModeUILayer()
 {
 	FBaseCharacterFXEditorToolkit::CreateEditorModeUILayer();
@@ -349,9 +349,9 @@ void FDataflowEditorToolkit::GetSaveableObjects(TArray<UObject*>& OutObjects) co
 {
 	FBaseCharacterFXEditorToolkit::GetSaveableObjects(OutObjects);
 
-	if (ensure(GetDataflowEditorContent()))
+	if (ensure(GetDataflowContent()))
 	{
-		if (UDataflow* DataflowAsset = GetDataflowEditorContent()->GetDataflowAsset())
+		if (UDataflow* DataflowAsset = GetDataflowContent()->DataflowAsset)
 		{
 			check(DataflowAsset->IsAsset());
 			OutObjects.Add(DataflowAsset);
@@ -368,14 +368,16 @@ void FDataflowEditorToolkit::CreateWidgets()
 {
 	FBaseCharacterFXEditorToolkit::CreateWidgets();
 
-	if (TObjectPtr<UDataflowEditorContent> Content = GetDataflowEditorContent())
+	if (TObjectPtr<UDataflowBaseContent> DataflowContent = GetDataflowContent())
 	{
-		if (UDataflow* DataflowAsset = Content->GetDataflowAsset())
+		if (UDataflow* DataflowAsset = DataflowContent->GetDataflowAsset())
 		{
-			NodeDetailsEditor = CreateNodeDetailsEditorWidget(Content->GetDataflowOwner());
-			AssetDetailsEditor = CreateAssetDetailsEditorWidget(Content->GetDataflowOwner());
+			NodeDetailsEditor = CreateNodeDetailsEditorWidget(DataflowContent->GetDataflowOwner());
+			AssetDetailsEditor = CreateAssetDetailsEditorWidget(DataflowContent->GetDataflowOwner());
 			GraphEditor = CreateGraphEditorWidget(DataflowAsset, NodeDetailsEditor);
-			SkeletalEditor = CreateSkeletalEditorWidget(Content->GetDataflowOwner());
+			SkeletalEditor = CreateSkeletalEditorWidget();
+
+			CreateSimulationViewportClient();
 		}
 	}
 }
@@ -403,28 +405,51 @@ TSharedPtr<FEditorViewportClient> FDataflowEditorToolkit::CreateEditorViewportCl
 	// namely ViewportType.
 	// Instead, we do viewport client adjustment in PostInitAssetEditor().
 	check(EditorModeManager.IsValid());
-	return MakeShared<FDataflowEditorViewportClient>(EditorModeManager.Get(), ObjectScene.Get());
+	TSharedPtr<FDataflowEditorViewportClient> LocalConstructionClient = MakeShared<FDataflowEditorViewportClient>(
+	EditorModeManager.Get(), ObjectScene.Get(), true);
+	LocalConstructionClient->SetDataflowEditorToolkit(StaticCastSharedRef<FDataflowEditorToolkit>(
+		const_cast<FDataflowEditorToolkit*>(this)->AsShared()));
+	return LocalConstructionClient;
+}
+
+void FDataflowEditorToolkit::CreateSimulationViewportClient()
+{
+	SimulationTabContent = MakeShareable(new FEditorViewportTabContent());
+	SimulationViewportClient = MakeShared<FDataflowEditorViewportClient>(SimulationModeManager.Get(),
+		SimulationScene.Get(), false);
+	
+	SimulationViewportClient->SetDataflowEditorToolkit(StaticCastSharedRef<FDataflowEditorToolkit>(this->AsShared()));
+
+	SimulationViewportDelegate = [this](FAssetEditorViewportConstructionArgs InArgs)
+	{
+		return SAssignNew(DataflowSimulationViewport, SDataflowEditorViewport, InArgs)
+			.ViewportClient(StaticCastSharedPtr<FDataflowEditorViewportClient>(SimulationViewportClient));
+	};
 }
 
 //~ End FBaseAssetToolkit overrides
 
 void FDataflowEditorToolkit::OnPropertyValueChanged(const FPropertyChangedEvent& PropertyChangedEvent)
 {
-	if (TObjectPtr<UDataflowEditorContent> EditorContent = GetDataflowEditorContent())
+	if (TObjectPtr<UDataflowBaseContent> EditorContent = GetDataflowContent())
 	{
 		ensure(EditorContent);
 		if (UDataflow* DataflowAsset = EditorContent->GetDataflowAsset())
 		{
-			TSharedPtr<Dataflow::FEngineContext>& DataflowContext = EditorContent->GetDataflowContext();
-			Dataflow::FTimestamp& LastNodeTimestamp = EditorContent->GetLastModifiedTimestamp();
+			TSharedPtr<Dataflow::FEngineContext> DataflowContext = EditorContent->GetDataflowContext();
+			Dataflow::FTimestamp LastNodeTimestamp = EditorContent->GetLastModifiedTimestamp();
+			
 			FDataflowEditorCommands::OnPropertyValueChanged(DataflowAsset, DataflowContext, LastNodeTimestamp, PropertyChangedEvent, SelectedDataflowNodes);
+
+			EditorContent->SetDataflowContext(DataflowContext);
+			EditorContent->SetLastModifiedTimestamp(LastNodeTimestamp);
 		}
 	}
 }
 
 void FDataflowEditorToolkit::OnAssetPropertyValueChanged(const FPropertyChangedEvent& PropertyChangedEvent)
 {
-	if (TObjectPtr<UDataflowEditorContent> EditorContent = GetDataflowEditorContent())
+	if (TObjectPtr<UDataflowBaseContent> EditorContent = GetDataflowContent())
 	{
 		ensure(EditorContent);	FDataflowEditorCommands::OnAssetPropertyValueChanged(EditorContent, PropertyChangedEvent);
 	}
@@ -442,7 +467,6 @@ void FDataflowEditorToolkit::OnNodeTitleCommitted(const FText& InNewText, ETextC
 
 void FDataflowEditorToolkit::OnNodeSelectionChanged(const TSet<UObject*>& InNewSelection)
 {
-
 	auto FindDataflowNodesInSet = [](const TSet<UObject*>& InSet) {
 		TSet<UObject*> Results;
 		for (UObject* Item : InSet)
@@ -493,7 +517,7 @@ void FDataflowEditorToolkit::OnNodeSelectionChanged(const TSet<UObject*>& InNewS
 	// Despite this function's name, we might not have actually changed which node is selected
 	bool bPrimarySelectionChanged = false;
 
-	if (TObjectPtr<UDataflowEditorContent> EditorContent = GetDataflowEditorContent(); EditorContent->GetDataflowAsset())
+	if (TObjectPtr<UDataflowBaseContent> EditorContent = GetDataflowContent(); EditorContent->DataflowAsset)
 	{
 		// Only keep UDataflowEdNode from NewSelection
 		TSet<UObject*> NodeSelection = FindDataflowNodesInSet(InNewSelection);
@@ -598,17 +622,19 @@ void FDataflowEditorToolkit::OnNodeDeleted(const TSet<UObject*>& NewSelection)
 
 void FDataflowEditorToolkit::Tick(float DeltaTime)
 {
-	if (TObjectPtr<UDataflowEditorContent> EditorContent = GetDataflowEditorContent())
+	if (TObjectPtr<UDataflowBaseContent> EditorContent = GetDataflowContent())
 	{
-		if (EditorContent->IsValid())
+		if (EditorContent->GetDataflowAsset())
 		{
+			Dataflow::FTimestamp TimeStamp = EditorContent->GetLastModifiedTimestamp();
 			if (!EditorContent->GetDataflowContext())
 			{
-				EditorContent->GetDataflowContext() = MakeShared<Dataflow::FEngineContext>(EditorContent->GetDataflowOwner(), EditorContent->GetDataflowAsset(), Dataflow::FTimestamp::Invalid);
-				EditorContent->GetLastModifiedTimestamp() = Dataflow::FTimestamp::Invalid;
+				EditorContent->SetDataflowContext(MakeShared<Dataflow::FEngineContext>(EditorContent->GetDataflowOwner(), EditorContent->GetDataflowAsset(), Dataflow::FTimestamp::Invalid));
+				TimeStamp = Dataflow::FTimestamp::Invalid;
 			}
-			FDataflowEditorCommands::EvaluateTerminalNode(*EditorContent->GetDataflowContext().Get(), EditorContent->GetLastModifiedTimestamp(),
-														  EditorContent->GetDataflowAsset(), nullptr, nullptr, EditorContent->GetDataflowOwner(), EditorContent->GetDataflowTerminal());
+			FDataflowEditorCommands::EvaluateTerminalNode(*EditorContent->GetDataflowContext().Get(), TimeStamp, EditorContent->GetDataflowAsset(), 
+			nullptr, nullptr, EditorContent->GetDataflowOwner(), EditorContent->GetDataflowTerminal());
+			EditorContent->SetLastModifiedTimestamp(TimeStamp);
 		}
 	}
 }
@@ -625,19 +651,20 @@ TSharedRef<SDataflowGraphEditor> FDataflowEditorToolkit::CreateGraphEditorWidget
 
 	FDataflowEditorCommands::FGraphEvaluationCallback Evaluate = [&](FDataflowNode* Node, FDataflowOutput* Out)
 	{
-		if (TObjectPtr<UDataflowEditorContent> EditorContent = GetDataflowEditorContent())
+		if (TObjectPtr<UDataflowBaseContent> EditorContent = GetDataflowContent())
 		{
-			if (EditorContent->IsValid())
+			if (EditorContent->GetDataflowAsset())
 			{
 				if (!EditorContent->GetDataflowContext())
 				{
-					EditorContent->GetDataflowContext() = MakeShared<Dataflow::FEngineContext>(EditorContent->GetDataflowOwner(), EditorContent->GetDataflowAsset(), Dataflow::FTimestamp::Invalid);
+					EditorContent->SetDataflowContext(MakeShared<Dataflow::FEngineContext>(EditorContent->GetDataflowOwner(), EditorContent->GetDataflowAsset(), Dataflow::FTimestamp::Invalid));
 				}
 				Node->Invalidate();
-				EditorContent->GetLastModifiedTimestamp() = Dataflow::FTimestamp::Invalid;
-
-				FDataflowEditorCommands::EvaluateTerminalNode(*EditorContent->GetDataflowContext().Get(), EditorContent->GetLastModifiedTimestamp(),
+				Dataflow::FTimestamp TimeStamp = Dataflow::FTimestamp::Invalid;
+				
+				FDataflowEditorCommands::EvaluateTerminalNode(*EditorContent->GetDataflowContext().Get(), TimeStamp,
 															  EditorContent->GetDataflowAsset(), Node, Out, EditorContent->GetDataflowOwner(), EditorContent->GetDataflowTerminal());
+				EditorContent->SetLastModifiedTimestamp(TimeStamp);
 			}
 		}
 	};
@@ -655,7 +682,6 @@ TSharedRef<SDataflowGraphEditor> FDataflowEditorToolkit::CreateGraphEditorWidget
 
 	OnSelectionChangedMulticastDelegateHandle = NewGraphEditor->OnSelectionChangedMulticast.AddSP(this, &FDataflowEditorToolkit::OnNodeSelectionChanged);
 	OnNodeDeletedMulticastDelegateHandle = NewGraphEditor->OnNodeDeletedMulticast.AddSP(this, &FDataflowEditorToolkit::OnNodeDeleted);
-	GetDataflowEditorContent()->SetDataflowGraphEditor(NewGraphEditor);
 
 	return NewGraphEditor;
 }
@@ -715,15 +741,15 @@ TSharedPtr<IDetailsView> FDataflowEditorToolkit::CreateAssetDetailsEditorWidget(
 
 }
 
-TSharedPtr<ISkeletonTree> FDataflowEditorToolkit::CreateSkeletalEditorWidget(UObject* ObjectToEdit)
+TSharedPtr<ISkeletonTree> FDataflowEditorToolkit::CreateSkeletalEditorWidget()
 {
-	if (TObjectPtr<UDataflowEditorContent> EditorContent = GetDataflowEditorContent())
+	if (const TObjectPtr<UDataflowSkeletalContent> SkeletalContent = Cast<UDataflowSkeletalContent>(GetDataflowContent()))
 	{
-		if (EditorContent->GetDataflowAsset())
+		if (SkeletalContent->GetDataflowAsset())
 		{
-			FSkeletonTreeArgs SkeletonTreeArgs;
+			const FSkeletonTreeArgs SkeletonTreeArgs;
 			ISkeletonEditorModule& SkeletonEditorModule = FModuleManager::LoadModuleChecked<ISkeletonEditorModule>("SkeletonEditor");
-			TSharedPtr<ISkeletonTree> SkeletonTree = SkeletonEditorModule.CreateSkeletonTree(EditorContent->GetSkeleton(), SkeletonTreeArgs);
+			TSharedPtr<ISkeletonTree> SkeletonTree = SkeletonEditorModule.CreateSkeletonTree(SkeletalContent->GetSkeleton(), SkeletonTreeArgs);
 			return SkeletonTree;
 		}
 	}
@@ -741,6 +767,15 @@ TSharedRef<SDockTab> FDataflowEditorToolkit::SpawnTab_AssetDetails(const FSpawnT
 		];
 }
 
+TSharedRef<SDockTab> FDataflowEditorToolkit::SpawnTab_SimulationViewport(const FSpawnTabArgs& Args)
+{
+	TSharedRef<SDockTab> DockableTab = SNew(SDockTab);
+	if(SimulationTabContent)
+	{
+		SimulationTabContent->Initialize(SimulationViewportDelegate, DockableTab, SimulationViewportTabId.ToString());
+	}
+	return DockableTab;
+}
 
 TSharedRef<SDockTab> FDataflowEditorToolkit::SpawnTab_GraphCanvas(const FSpawnTabArgs& Args)
 {
@@ -769,16 +804,21 @@ TSharedRef<SDockTab> FDataflowEditorToolkit::SpawnTab_NodeDetails(const FSpawnTa
 TSharedRef<SDockTab> FDataflowEditorToolkit::SpawnTab_Skeletal(const FSpawnTabArgs& Args)
 {
 	check(Args.GetTabId() == SkeletalTabId);
-	if (const TObjectPtr<UDataflowEditorContent> EditorContent = GetDataflowEditorContent())
+	if(SkeletalEditor)
 	{
-		SkeletalEditor->SetSkeletalMesh(EditorContent->GetSkeletalMesh());
+		if (const TObjectPtr<UDataflowSkeletalContent> EditorContent = Cast<UDataflowSkeletalContent>(GetDataflowContent()))
+		{
+			SkeletalEditor->SetSkeletalMesh(EditorContent->GetSkeletalMesh());
+		}
+		return SNew(SDockTab)
+			.Label(LOCTEXT("DataflowEditor_Skeletal_TabTitle", "Skeletal Hierarchy"))
+			[
+				SkeletalEditor.ToSharedRef()
+			];
 	}
 
 	return SNew(SDockTab)
-		.Label(LOCTEXT("DataflowEditor_Skeletal_TabTitle", "Skeletal Hierarchy"))
-		[
-			SkeletalEditor.ToSharedRef()
-		];
+		.Label(LOCTEXT("DataflowEditor_Skeletal_TabTitle", "Skeletal Hierarchy"));
 }
 
 TSharedRef<SDockTab> FDataflowEditorToolkit::SpawnTab_SelectionView(const FSpawnTabArgs& Args)
@@ -827,7 +867,7 @@ TSharedRef<SDockTab> FDataflowEditorToolkit::SpawnTab_SelectionView(const FSpawn
 
 	if (SelectionViewWidget)
 	{
-		if (const TObjectPtr<UDataflowEditorContent> EditorContent = GetDataflowEditorContent())
+		if (const TObjectPtr<UDataflowBaseContent> EditorContent = GetDataflowContent())
 		{
 			if (Args.GetTabId() == SelectionViewTabId_1)
 			{
@@ -921,7 +961,7 @@ TSharedRef<SDockTab> FDataflowEditorToolkit::SpawnTab_CollectionSpreadSheet(cons
 
 	if (CollectionSpreadSheetWidget)
 	{
-		if (const TObjectPtr<UDataflowEditorContent> EditorContent = GetDataflowEditorContent())
+		if (const TObjectPtr<UDataflowBaseContent> EditorContent = GetDataflowContent())
 		{
 			if (Args.GetTabId() == CollectionSpreadSheetTabId_1)
 			{
@@ -980,7 +1020,12 @@ void FDataflowEditorToolkit::RegisterTabSpawners(const TSharedRef<FTabManager>&
 	const TSharedRef<FWorkspaceItem> CollectionSpreadSheetWorkspaceMenuCategoryRef = EditorMenuCategory->AddGroup(LOCTEXT("WorkspaceMenu_CollectionSpreadSheet", "Collection SpreadSheet"), FSlateIcon(FAppStyle::GetAppStyleSetName(), "LevelEditor.Tabs.Outliner"));
 	
 	InTabManager->RegisterTabSpawner(ViewportTabID, FOnSpawnTab::CreateSP(this, &FDataflowEditorToolkit::SpawnTab_Viewport))
-		.SetDisplayName(LOCTEXT("DataflowViewportTab", "Dataflow Viewport"))
+		.SetDisplayName(LOCTEXT("DataflowViewportTab", "Construction Viewport"))
+		.SetGroup(EditorMenuCategory.ToSharedRef())
+		.SetIcon(FSlateIcon(FAppStyle::GetAppStyleSetName(), "LevelEditor.Tabs.Viewports"));
+
+	InTabManager->RegisterTabSpawner(SimulationViewportTabId, FOnSpawnTab::CreateSP(this, &FDataflowEditorToolkit::SpawnTab_SimulationViewport))
+		.SetDisplayName(LOCTEXT("SimulationViewportTab", "Simulation Viewport"))
 		.SetGroup(EditorMenuCategory.ToSharedRef())
 		.SetIcon(FSlateIcon(FAppStyle::GetAppStyleSetName(), "LevelEditor.Tabs.Viewports"));
 
@@ -1060,6 +1105,7 @@ void FDataflowEditorToolkit::UnregisterTabSpawners(const TSharedRef<class FTabMa
 	InTabManager->UnregisterTabSpawner(CollectionSpreadSheetTabId_2);
 	InTabManager->UnregisterTabSpawner(CollectionSpreadSheetTabId_3);
 	InTabManager->UnregisterTabSpawner(CollectionSpreadSheetTabId_4);
+	InTabManager->UnregisterTabSpawner(SimulationViewportTabId);
 }
 
 void FDataflowEditorToolkit::OnTabClosed(TSharedRef<SDockTab> Tab)
@@ -1105,7 +1151,7 @@ FName FDataflowEditorToolkit::GetToolkitFName() const
 
 FText FDataflowEditorToolkit::GetToolkitName() const
 {
-	if (TObjectPtr<const UDataflowEditorContent> EditorContent = GetDataflowEditorContent())
+	if (const TObjectPtr<UDataflowBaseContent> EditorContent = GetDataflowContent())
 	{
 		if (EditorContent->GetDataflowOwner())
 		{
diff --git a/Engine/Plugins/Experimental/Dataflow/Source/DataflowEditor/Private/Dataflow/DataflowEditorUtil.cpp b/Engine/Plugins/Experimental/Dataflow/Source/DataflowEditor/Private/Dataflow/DataflowEditorUtil.cpp
--- a/Engine/Plugins/Experimental/Dataflow/Source/DataflowEditor/Private/Dataflow/DataflowEditorUtil.cpp
+++ b/Engine/Plugins/Experimental/Dataflow/Source/DataflowEditor/Private/Dataflow/DataflowEditorUtil.cpp
@@ -4,7 +4,7 @@
 
 #include "Animation/Skeleton.h"
 #include "Dataflow/DataflowEditor.h"
-#include "Dataflow/DataflowEditorContent.h"
+#include "Dataflow/DataflowContent.h"
 #include "Dataflow/DataflowEditorToolkit.h"
 #include "Dataflow/DataflowEdNode.h"
 #include "Dataflow/DataflowObject.h"
@@ -18,6 +18,26 @@ using namespace UE::Geometry;
 
 namespace Private
 {
+	bool HasSkeletalMesh(UObject* InObject)
+	{
+		if (UClass* Class = InObject->GetClass())
+		{
+			return Class->FindPropertyByName(FName("SkeletalMesh")) &&
+				   Class->FindPropertyByName(FName("Skeleton"));
+		}
+		return false;
+	}
+	
+	bool HasDataflowAsset(UObject* InObject)
+	{
+		if (UClass* Class = InObject->GetClass())
+		{
+			return Class->FindPropertyByName(FName("DataflowAsset")) &&
+				   Class->FindPropertyByName(FName("DataflowTerminal"));
+		}
+		return false;
+	}
+	
 	UDataflow* GetDataflowAssetFrom(UObject* InObject)
 	{
 		if (UClass* Class = InObject->GetClass())
@@ -96,7 +116,7 @@ namespace UE
 
 namespace Dataflow
 {
-	TSharedPtr<FEngineContext> GetContext(UDataflowEditorContent* Content)
+	TSharedPtr<FEngineContext> GetContext(TObjectPtr<UDataflowBaseContent> Content)
 	{
 		if (Content)
 		{
diff --git a/Engine/Plugins/Experimental/Dataflow/Source/DataflowEditor/Private/Dataflow/DataflowEditorUtil.h b/Engine/Plugins/Experimental/Dataflow/Source/DataflowEditor/Private/Dataflow/DataflowEditorUtil.h
--- a/Engine/Plugins/Experimental/Dataflow/Source/DataflowEditor/Private/Dataflow/DataflowEditorUtil.h
+++ b/Engine/Plugins/Experimental/Dataflow/Source/DataflowEditor/Private/Dataflow/DataflowEditorUtil.h
@@ -11,7 +11,7 @@
 
 class UDataflow;
 class UObject;
-class UDataflowEditorContent;
+class UDataflowBaseContent;
 class USkeletalMesh;
 class USkeleton;
 class UAnimationAsset;
@@ -19,6 +19,10 @@ class UMaterial;
 
 namespace Private
 {
+	bool HasSkeletalMesh(UObject* InObject);
+
+	bool HasDataflowAsset(UObject* InObject);
+	
 	UDataflow* GetDataflowAssetFrom(UObject* InObject);
 
 	USkeletalMesh* GetSkeletalMeshFrom(UObject* InObject);
@@ -40,5 +44,5 @@ namespace UE
 
 namespace Dataflow
 {
-	TSharedPtr<::Dataflow::FEngineContext> GetContext(UDataflowEditorContent* Content);
+	TSharedPtr<::Dataflow::FEngineContext> GetContext(TObjectPtr<UDataflowBaseContent> Content);
 }
\ No newline at end of file
diff --git a/Engine/Plugins/Experimental/Dataflow/Source/DataflowEditor/Private/Dataflow/DataflowEditorViewport.cpp b/Engine/Plugins/Experimental/Dataflow/Source/DataflowEditor/Private/Dataflow/DataflowEditorViewport.cpp
--- a/Engine/Plugins/Experimental/Dataflow/Source/DataflowEditor/Private/Dataflow/DataflowEditorViewport.cpp
+++ b/Engine/Plugins/Experimental/Dataflow/Source/DataflowEditor/Private/Dataflow/DataflowEditorViewport.cpp
@@ -2,13 +2,15 @@
 
 #include "Dataflow/DataflowEditorViewport.h"
 
-#include "AdvancedPreviewScene.h"
 #include "Dataflow/DataflowActor.h"
-#include "Dataflow/DataflowEditorCommands.h"
 #include "Dataflow/DataflowEditorMode.h"
 #include "Dataflow/DataflowEditorViewportClient.h"
+#include "Dataflow/DataflowEditorToolkit.h"
 #include "EditorModeManager.h"
+#include "Dataflow/DataflowContent.h"
 #include "Dataflow/DataflowEditorViewportToolbar.h"
+#include "Dataflow/DataflowPreviewScene.h"
+#include "Dataflow/DataflowSimulationPanel.h"
 
 #define LOCTEXT_NAMESPACE "SDataflowEditorViewport"
 
@@ -23,6 +25,33 @@ void SDataflowEditorViewport::Construct(const FArguments& InArgs, const FAssetEd
 	ParentArgs._EditorViewportClient = InArgs._ViewportClient;
 	SAssetEditorViewport::Construct(ParentArgs, InViewportConstructionArgs);
 	Client->VisibilityDelegate.BindSP(this, &SDataflowEditorViewport::IsVisible);
+
+	if(static_cast<FDataflowPreviewScene*>(Client->GetPreviewScene())->CanRunSimulation())
+	{
+		TSharedPtr<FDataflowEditorViewportClient> DataflowClient = StaticCastSharedPtr<FDataflowEditorViewportClient>(Client);
+		TWeakPtr<FDataflowSimulationScene> SimulationScene = DataflowClient->GetDataflowEditorToolkit().Pin()->GetSimulationScene();
+            
+		ViewportOverlay->AddSlot()
+		[
+			SNew(SHorizontalBox)
+			+ SHorizontalBox::Slot()
+			.HAlign(HAlign_Fill)
+			.VAlign(VAlign_Bottom)
+			.FillWidth(1)
+			.Padding(10.0f, 0.0f)
+			[
+				SNew(SBorder)
+				.BorderImage(FAppStyle::Get().GetBrush("EditorViewport.OverlayBrush"))
+				.Visibility(EVisibility::Visible)
+				.Padding(10.0f, 2.0f)
+				[
+					SNew(SDataflowSimulationPanel, SimulationScene)
+					.ViewInputMin(this, &SDataflowEditorViewport::GetViewMinInput)
+					.ViewInputMax(this, &SDataflowEditorViewport::GetViewMaxInput)
+				]
+			]
+		];
+	}
 }
 
 TSharedPtr<SWidget> SDataflowEditorViewport::MakeViewportToolbar()
@@ -32,9 +61,9 @@ TSharedPtr<SWidget> SDataflowEditorViewport::MakeViewportToolbar()
 
 void SDataflowEditorViewport::OnFocusViewportToSelection()
 {
-	if(UDataflowEditorMode* DataflowMode = GetEdMode())
+	if(const FDataflowPreviewScene* PreviewScene = static_cast<FDataflowPreviewScene*>(Client->GetPreviewScene()))
 	{
-		const FBox SceneBoundingBox = DataflowMode->SceneBoundingBox();
+		const FBox SceneBoundingBox = PreviewScene->GetBoundingBox();
 		Client->FocusViewportOnBox(SceneBoundingBox);
 	}
 }
@@ -77,4 +106,15 @@ void SDataflowEditorViewport::OnFloatingButtonClicked()
 {
 }
 
+float SDataflowEditorViewport:: GetViewMinInput() const
+{
+	return static_cast<FDataflowPreviewScene*>(Client->GetPreviewScene())->GetDataflowContent()->GetSimulationRange()[0];
+}
+
+float SDataflowEditorViewport::GetViewMaxInput() const
+{
+	return static_cast<FDataflowPreviewScene*>(Client->GetPreviewScene())->GetDataflowContent()->GetSimulationRange()[1];
+}
+
+
 #undef LOCTEXT_NAMESPACE
diff --git a/Engine/Plugins/Experimental/Dataflow/Source/DataflowEditor/Private/Dataflow/DataflowEditorViewportClient.cpp b/Engine/Plugins/Experimental/Dataflow/Source/DataflowEditor/Private/Dataflow/DataflowEditorViewportClient.cpp
--- a/Engine/Plugins/Experimental/Dataflow/Source/DataflowEditor/Private/Dataflow/DataflowEditorViewportClient.cpp
+++ b/Engine/Plugins/Experimental/Dataflow/Source/DataflowEditor/Private/Dataflow/DataflowEditorViewportClient.cpp
@@ -1,23 +1,16 @@
 // Copyright Epic Games, Inc. All Rights Reserved.
 #include "Dataflow/DataflowEditorViewportClient.h"
 
-#include "Dataflow/DataflowActor.h"
-#include "Dataflow/DataflowEditor.h"
-#include "Dataflow/DataflowEditorContent.h"
-#include "Dataflow/DataflowEditorMode.h"
+#include "Dataflow/DataflowObject.h"
 #include "Dataflow/DataflowEditorToolkit.h"
 #include "Dataflow/DataflowEngineSceneHitProxies.h"
 #include "Dataflow/DataflowPreviewScene.h"
-#include "Dataflow/DataflowXml.h"
 #include "EditorModeManager.h"
-#include "EdModeInteractiveToolsContext.h"
-#include "HAL/PlatformApplicationMisc.h"
 #include "PreviewScene.h"
 #include "Selection.h"
-#include "InputBehaviorSet.h"
 
 FDataflowEditorViewportClient::FDataflowEditorViewportClient(FEditorModeTools* InModeTools,
-                                                             FPreviewScene* InPreviewScene,
+                                                             FPreviewScene* InPreviewScene,  const bool bCouldTickScene,
                                                              const TWeakPtr<SEditorViewport> InEditorViewportWidget)
 	: FEditorViewportClient(InModeTools, InPreviewScene, InEditorViewportWidget)
 {
@@ -28,6 +21,7 @@ FDataflowEditorViewportClient::FDataflowEditorViewportClient(FEditorModeTools* I
 	EngineShowFlags.EnableAdvancedFeatures();
 
 	PreviewScene = static_cast<FDataflowPreviewScene*>(InPreviewScene);
+	bEnableSceneTicking = bCouldTickScene;
 }
 
 void FDataflowEditorViewportClient::SetDataflowEditorToolkit(TWeakPtr<FDataflowEditorToolkit> InDataflowEditorToolkitPtr)
@@ -45,43 +39,46 @@ void FDataflowEditorViewportClient::SetToolCommandList(TWeakPtr<FUICommandList>
 //	return BehaviorSet;
 //}
 
-Dataflow::FTimestamp LatestTimestamp(const UDataflow* Dataflow, const ::Dataflow::FContext* Context)
+void FDataflowEditorViewportClient::Tick(float DeltaSeconds)
 {
-	if (Dataflow && Context)
+	Super::Tick(DeltaSeconds);
+	if (PreviewScene)
 	{
-		return FMath::Max(Dataflow->GetRenderingTimestamp().Value, Context->GetTimestamp().Value);
+		PreviewScene->TickDataflowScene(DeltaSeconds);
 	}
-	return ::Dataflow::FTimestamp::Invalid;
 }
 
-void FDataflowEditorViewportClient::Tick(float DeltaSeconds)
+void FDataflowEditorViewportClient::ProcessClick(FSceneView& View, HHitProxy* HitProxy, FKey Key, EInputEvent Event, uint32 HitX, uint32 HitY)
 {
-	Super::Tick(DeltaSeconds);
+	Super::ProcessClick(View, HitProxy, Key, Event, HitX, HitY);
 
-	const TSharedPtr<FDataflowEditorToolkit> DataflowEditorToolkit = DataflowEditorToolkitPtr.Pin();
+	USelection* SelectedComponents = ModeTools->GetSelectedComponents();
 
-	if (PreviewScene)
+	TArray<UPrimitiveComponent*> PreviouslySelectedComponents;
+	SelectedComponents->GetSelectedObjects<UPrimitiveComponent>(PreviouslySelectedComponents);
+
+	SelectedComponents->Modify();
+	SelectedComponents->BeginBatchSelectOperation();
+
+	SelectedComponents->DeselectAll();
+
+	if (HitProxy && HitProxy->IsA(HActor::StaticGetType()))
 	{
-		if (TObjectPtr<UDataflowEditorContent> EditorContent = PreviewScene->GetDataflowEditorContent())
+		const HActor* ActorProxy = static_cast<HActor*>(HitProxy);
+		if (ActorProxy && ActorProxy->PrimComponent && ActorProxy->Actor)
 		{
-			if (TSharedPtr<Dataflow::FContext> Context = EditorContent->GetDataflowContext())
-			{
-				if (const UDataflow* Dataflow = EditorContent->GetDataflowAsset())
-				{
-					const Dataflow::FTimestamp SystemTimestamp = LatestTimestamp(Dataflow, Context.Get());
-					if (SystemTimestamp >= EditorContent->GetLastModifiedTimestamp())
-					{
-						EditorContent->SetLastModifiedTimestamp(LatestTimestamp(EditorContent->GetDataflowAsset(), EditorContent->GetDataflowContext().Get()).Value + 1);
-					}
-				}
-
-				if (EditorContent->IsDirty())
-				{
-					PreviewScene->Update();
-				}
-			}
+			UPrimitiveComponent* Component = const_cast<UPrimitiveComponent*>(ActorProxy->PrimComponent.Get());
+			SelectedComponents->Select(Component);
+			Component->PushSelectionToProxy();
 		}
 	}
+
+	SelectedComponents->EndBatchSelectOperation();
+
+	for (UPrimitiveComponent* const Component : PreviouslySelectedComponents)
+	{
+		Component->PushSelectionToProxy();
+	}
 }
 
 void FDataflowEditorViewportClient::SetConstructionViewMode(Dataflow::EDataflowPatternVertexType InViewMode)
diff --git a/Engine/Plugins/Experimental/Dataflow/Source/DataflowEditor/Private/Dataflow/DataflowEditorViewportToolbar.cpp b/Engine/Plugins/Experimental/Dataflow/Source/DataflowEditor/Private/Dataflow/DataflowEditorViewportToolbar.cpp
--- a/Engine/Plugins/Experimental/Dataflow/Source/DataflowEditor/Private/Dataflow/DataflowEditorViewportToolbar.cpp
+++ b/Engine/Plugins/Experimental/Dataflow/Source/DataflowEditor/Private/Dataflow/DataflowEditorViewportToolbar.cpp
@@ -1,86 +1,10 @@
 // Copyright Epic Games, Inc. All Rights Reserved.
 #include "Dataflow/DataflowEditorViewportToolbar.h"
-
-#include "Dataflow/DataflowEditorCommands.h"
-#include "Dataflow/DataflowEditorStyle.h"
 #include "Dataflow/DataflowEditorViewport.h"
-#include "EditorViewportCommands.h"
-#include "HAL/IConsoleManager.h"
-#include "Widgets/Layout/SBorder.h"
-
-bool bDataflowPrototypeSelectionMenu = false;
-FAutoConsoleVariableRef CVarDataflowPrototypeSelectionMenu(TEXT("p.Dataflow.prototype.selection"), bDataflowPrototypeSelectionMenu, TEXT("Work in progress development for selection controls in the 3D Viewport.[def:false]"));
-
 
 void SDataflowViewportSelectionToolBar::Construct(const FArguments& InArgs, TSharedPtr<SDataflowEditorViewport> InDataflowViewport)
 {
 	EditorViewport = InDataflowViewport;
-	static const FName DefaultForegroundName("DefaultForeground");
-
-	this->ChildSlot
-		[
-			SNew(SBorder)
-			.BorderImage(FAppStyle::GetBrush("NoBorder"))
-			.ForegroundColor(FAppStyle::GetSlateColor(DefaultForegroundName))
-			[
-			SNew(SHorizontalBox)
-			+ SHorizontalBox::Slot()
-				.AutoWidth()
-				.Padding(2.0f, 2.0f)
-				.HAlign(HAlign_Right)
-				[
-					MakeSelectionModeToolBar()
-				]
-			]
-		];
-
-	SViewportToolBar::Construct(SViewportToolBar::FArguments());
-}
-
-TSharedRef<SWidget> SDataflowViewportSelectionToolBar::MakeSelectionModeToolBar()
-{
-	FSlimHorizontalToolBarBuilder ToolbarBuilder(EditorViewport.Pin()->GetCommandList(), FMultiBoxCustomization::None);
-
-	const FName ToolBarStyle = TEXT("EditorViewportToolBar");
-	ToolbarBuilder.SetStyle(&FAppStyle::Get(), ToolBarStyle);
-	ToolbarBuilder.SetLabelVisibility(EVisibility::Collapsed);
-
-	ToolbarBuilder.SetIsFocusable(false);
-
-	ToolbarBuilder.BeginSection("Selection");
-	{
-		ToolbarBuilder.BeginBlockGroup();
-
-		if (bDataflowPrototypeSelectionMenu)
-		{
-			ToolbarBuilder.AddToolBarButton(FDataflowEditorCommands::Get().ToggleObjectSelection,
-											NAME_None,
-											TAttribute<FText>(),
-											TAttribute<FText>(),
-											FSlateIcon(FDataflowEditorStyle::Get().GetStyleSetName(), TEXT("Dataflow.SelectObject")),
-											FName(TEXT("SelectObjectMode"))
-			);
-			ToolbarBuilder.AddToolBarButton(FDataflowEditorCommands::Get().ToggleFaceSelection,
-											NAME_None,
-											TAttribute<FText>(),
-											TAttribute<FText>(),
-											FSlateIcon(FDataflowEditorStyle::Get().GetStyleSetName(), TEXT("Dataflow.SelectFace")),
-											FName(TEXT("SelectFaceMode"))
-			);
-			ToolbarBuilder.AddToolBarButton(FDataflowEditorCommands::Get().ToggleVertexSelection,
-											NAME_None,
-											TAttribute<FText>(),
-											TAttribute<FText>(),
-											FSlateIcon(FDataflowEditorStyle::Get().GetStyleSetName(), TEXT("Dataflow.SelectVertex")),
-											FName(TEXT("SelectVertexMode"))
-			);
-		}
-
-		ToolbarBuilder.EndBlockGroup();
-	}
-
-	ToolbarBuilder.EndSection();
-	ToolbarBuilder.AddSeparator();
-
-	return ToolbarBuilder.MakeWidget();
+	
+	SCommonEditorViewportToolbarBase::Construct(SCommonEditorViewportToolbarBase::FArguments(), InDataflowViewport);
 }
\ No newline at end of file
diff --git a/Engine/Plugins/Experimental/Dataflow/Source/DataflowEditor/Private/Dataflow/DataflowPreviewScene.cpp b/Engine/Plugins/Experimental/Dataflow/Source/DataflowEditor/Private/Dataflow/DataflowPreviewScene.cpp
--- a/Engine/Plugins/Experimental/Dataflow/Source/DataflowEditor/Private/Dataflow/DataflowPreviewScene.cpp
+++ b/Engine/Plugins/Experimental/Dataflow/Source/DataflowEditor/Private/Dataflow/DataflowPreviewScene.cpp
@@ -4,60 +4,128 @@
 
 #include "Animation/AnimSingleNodeInstance.h"
 #include "AssetEditorModeManager.h"
-#include "Components/SkeletalMeshComponent.h"
 #include "Components/DynamicMeshComponent.h"
 #include "Dataflow/CollectionRenderingPatternUtility.h"
-#include "Dataflow/DataflowActor.h"
-#include "Dataflow/DataflowComponent.h"
 #include "Dataflow/DataflowEditor.h"
-#include "Dataflow/DataflowEditorContent.h"
+#include "Dataflow/DataflowContent.h"
 #include "Dataflow/DataflowEditorStyle.h"
 #include "Dataflow/DataflowEditorUtil.h"
 #include "Drawing/MeshElementsVisualizer.h"
 #include "Elements/Framework/EngineElementsLibrary.h"
 #include "InteractiveTool.h"
-#include "ModelingToolTargetUtil.h"
 #include "Selection.h"
 
 #define LOCTEXT_NAMESPACE "FDataflowPreviewScene"
 
+FDataflowPreviewScene::FDataflowPreviewScene(FPreviewScene::ConstructionValues ConstructionValues, TObjectPtr<UDataflowBaseContent> InEditorContent) 
+	: FAdvancedPreviewScene(ConstructionValues), DataflowContent(InEditorContent)
+{
+	check(DataflowContent);
+	SetFloorVisibility(true, true);
+
+	RootSceneActor = GetWorld()->SpawnActor<AActor>(AActor::StaticClass());
+}
+
+FDataflowPreviewScene::~FDataflowPreviewScene()
+{}
 
-FDataflowPreviewScene::FDataflowPreviewScene(FPreviewScene::ConstructionValues ConstructionValues,TObjectPtr<UDataflowEditorContent> InEditorContent) 
-	: FAdvancedPreviewScene(ConstructionValues), EditorContent(InEditorContent)
+void FDataflowPreviewScene::AddReferencedObjects(FReferenceCollector& Collector)
 {
-	check(EditorContent);
-	SkeletalMeshActor = GetWorld()->SpawnActor<AActor>(AActor::StaticClass());
-	if (EditorContent->GetSkeletalMesh())
+	FAdvancedPreviewScene::AddReferencedObjects(Collector);
+	Collector.AddReferencedObject(DataflowContent);
+	Collector.AddReferencedObject(RootSceneActor);
+	
+	GetDataflowContent()->AddContentObjects(Collector);
+}
+
+bool FDataflowPreviewScene::IsComponentSelected(const UPrimitiveComponent* InComponent) const
+{
+	if(DataflowModeManager.IsValid())
 	{
-		SkeletalMeshComponent = NewObject<USkeletalMeshComponent>(SkeletalMeshActor);
-		SkeletalMeshComponent->SelectionOverrideDelegate = UPrimitiveComponent::FSelectionOverride::CreateRaw(this, &FDataflowPreviewScene::IsComponentSelected);
-		SkeletalMeshComponent->SetDisablePostProcessBlueprint(true);
-		UpdateSkeletalMeshComponent();
+		if (const UTypedElementSelectionSet* const TypedElementSelectionSet = DataflowModeManager->GetEditorSelectionSet())
+		{
+			if (const FTypedElementHandle ComponentElement = UEngineElementsLibrary::AcquireEditorComponentElementHandle(InComponent))
+			{
+				const bool bElementSelected = TypedElementSelectionSet->IsElementSelected(ComponentElement, FTypedElementIsSelectedOptions());
+				return bElementSelected;
+			}
+		}
 	}
-	SkeletalMeshActor->RegisterAllComponents();
-
-	DynamicMeshActor = GetWorld()->SpawnActor<AActor>(AActor::StaticClass());	
-	DynamicMeshActor->RegisterAllComponents();
-
-	SetFloorVisibility(false, true);
+	return false;
 }
 
-FDataflowPreviewScene::~FDataflowPreviewScene()
+FBox FDataflowPreviewScene::GetBoundingBox() const
 {
-	if (SkeletalMeshComponent)
+	FBox SceneBounds(ForceInitToZero);
+	if(DataflowModeManager.IsValid())
 	{
-		SkeletalMeshComponent->TransformUpdated.RemoveAll(this);
-		SkeletalMeshComponent->SelectionOverrideDelegate.Unbind();
-		SkeletalMeshComponent->UnregisterComponent();
-		SkeletalMeshComponent->DestroyComponent();
+		USelection* const SelectedComponents = DataflowModeManager->GetSelectedComponents();
+
+		TArray<TWeakObjectPtr<UObject>> SelectedObjects;
+		const int32 NumSelected = SelectedComponents->GetSelectedObjects(SelectedObjects);
+		
+		if(NumSelected > 0)
+		{
+			for(const TWeakObjectPtr<UObject> SelectedObject : SelectedObjects)
+			{
+				if(const UPrimitiveComponent* SelectedComponent = Cast<UPrimitiveComponent>(SelectedObject))
+				{
+					SceneBounds += SelectedComponent->Bounds.GetBox();
+				}
+			}
+		}
+		else
+		{
+			SceneBounds += RootSceneActor->GetComponentsBoundingBox(true);
+		}
 	}
+	return SceneBounds;
+}
+
+FDataflowConstructionScene::FDataflowConstructionScene(FPreviewScene::ConstructionValues ConstructionValues, TObjectPtr<UDataflowBaseContent> InEditorContent) 
+	: FDataflowPreviewScene(ConstructionValues,InEditorContent)
+{}
 
+FDataflowConstructionScene::~FDataflowConstructionScene()
+{
 	ResetDynamicMeshComponents();
 }
 
-void FDataflowPreviewScene::Tick(float DeltaTime)
+void FDataflowConstructionScene::AddReferencedObjects(FReferenceCollector& Collector)
 {
-	//@todo(brice) : Make sure this is being called. 
+	FDataflowPreviewScene::AddReferencedObjects(Collector);
+
+	Collector.AddReferencedObjects(DynamicMeshComponents);
+	Collector.AddReferencedObject(WireframeDraw);
+}
+
+FORCEINLINE Dataflow::FTimestamp LatestTimestamp(const UDataflow* Dataflow, const ::Dataflow::FContext* Context)
+{
+	if (Dataflow && Context)
+	{
+		return FMath::Max(Dataflow->GetRenderingTimestamp().Value, Context->GetTimestamp().Value);
+	}
+	return ::Dataflow::FTimestamp::Invalid;
+}
+
+void FDataflowConstructionScene::TickDataflowScene(const float DeltaSeconds)
+{
+	if (const TSharedPtr<Dataflow::FContext> DataflowContext = DataflowContent->GetDataflowContext())
+	{
+		if (const UDataflow* Dataflow = DataflowContent->DataflowAsset)
+		{
+			const Dataflow::FTimestamp SystemTimestamp = LatestTimestamp(Dataflow, DataflowContext.Get());
+			if (SystemTimestamp >= DataflowContent->GetLastModifiedTimestamp())
+			{
+				DataflowContent->SetLastModifiedTimestamp(SystemTimestamp.Value + 1);
+
+				if(DataflowContent->IsDirty())
+				{
+					UpdateConstructionScene();
+				}
+			}
+		}
+	}
 	for (TObjectPtr<UInteractiveToolPropertySet>& Propset : PropertyObjectsToTick)
 	{
 		if (Propset)
@@ -75,98 +143,77 @@ void FDataflowPreviewScene::Tick(float DeltaTime)
 
 	if (WireframeDraw)
 	{
-		WireframeDraw->OnTick(DeltaTime);
+		WireframeDraw->OnTick(DeltaSeconds);
 	}
 }
 
-void FDataflowPreviewScene::ReinitializeDynamicMeshComponents()
-{
-	Update();
-}
-
-void FDataflowPreviewScene::Update()
+void FDataflowConstructionScene::UpdateDynamicMeshComponents()
 {
-	// Some objects, like the UMeshElementsVisualizer and Settings Objects
-	// are not part of a tool, so they won't get ticked.This member holds
-	// ticked objects that get rebuilt on Update
-	PropertyObjectsToTick.Empty();
-
-	// Update the SkeletalMeshComponent for animation 
-	// changes.
-	UpdateSkeletalMeshComponent();
+	using namespace UE::Geometry;//FDynamicMesh3
 
 	// The preview scene for the construction view will be
 	// cleared and rebuilt from scratch. This will genrate a 
 	// list of UPrimitiveComponents for rendering.
-	UpdateDynamicMeshComponents();
-
-	// Attach a wireframe renderer to the DynamicMeshComponents
-	UpdateWireframeMeshElementsVisualizer();
-
+	ResetDynamicMeshComponents();
 
-	// Manage Selection and Tool Interaction
-	if (DataflowModeManager.IsValid())
+	if (DataflowContent)
 	{
-		USelection* SelectedComponents = DataflowModeManager->GetSelectedComponents();
-		for (const TObjectPtr<UDynamicMeshComponent>& DynamicMeshComponent : DynamicMeshComponents)
+		const TObjectPtr<UDataflow>& DataflowAsset = DataflowContent->GetDataflowAsset();
+		const TSharedPtr<Dataflow::FEngineContext>& DataflowContext = DataflowContent->GetDataflowContext();
+		if(DataflowAsset && DataflowContext)
 		{
-			SelectedComponents->DeselectAll();
-			SelectedComponents->Select(DynamicMeshComponent);
-			DynamicMeshComponent->PushSelectionToProxy();
-		}
-
-		// @todo(brice) : Deal with this
-		// Update the context object with the ConstructionViewMode and Collection used to build the DynamicMeshComponents, so
-		// tools know how to use the components.
-		//UEditorInteractiveToolsContext* RestSpaceToolsContext = DataflowModeManager->GetInteractiveToolsContext();
-		//UClothEditorContextObject* EditorContextObject = RestSpaceToolsContext->ContextObjectStore->FindContext<UClothEditorContextObject>();
-		//if (ensure(EditorContextObject))
-		//{
-		//	EditorContextObject->SetCollection(ConstructionViewMode, Collection);
-		//}
-	}
-}
-
-void FDataflowPreviewScene::Exit()
-{
-	PropertyObjectsToTick.Empty();
+			for (const UDataflowEdNode* Target : DataflowAsset->GetRenderTargets())
+			{
+				if (Target)
+				{
+					FDynamicMesh3 DynamicMesh;
+					TSharedPtr<FManagedArrayCollection> RenderCollection(new FManagedArrayCollection);
+					GeometryCollection::Facades::FRenderingFacade Facade(*RenderCollection);
+					Facade.DefineSchema();
 
-	if (WireframeDraw)
-	{
-		WireframeDraw->Disconnect();
+					Target->Render(Facade, DataflowContext);
+					Dataflow::Conversion::RenderingFacadeToDynamicMesh(Facade, DynamicMesh);
+					
+					if (Target == DataflowContent->GetPrimarySelectedNode())
+					{
+						DataflowContent->SetPrimaryRenderCollection(RenderCollection);
+					}
+					
+					AddDynamicMeshComponent(MoveTemp(DynamicMesh), {});
+				}
+			}
+		}
 	}
-	WireframeDraw = nullptr;
-
-	ResetDynamicMeshComponents();
 }
 
-
-void FDataflowPreviewScene::ResetDynamicMeshComponents()
+void FDataflowConstructionScene::ResetDynamicMeshComponents()
 {
 	USelection* SelectedComponents = DataflowModeManager->GetSelectedComponents();
 	for(const TObjectPtr<UDynamicMeshComponent>& DynamicMeshComponent : DynamicMeshComponents)
 	{
 		DynamicMeshComponent->SelectionOverrideDelegate.Unbind();
-		DynamicMeshComponent->UnregisterComponent();
-		DynamicMeshComponent->DestroyComponent();
-
 		if (SelectedComponents->IsSelected(DynamicMeshComponent))
 		{
 			SelectedComponents->Deselect(DynamicMeshComponent);
 			DynamicMeshComponent->PushSelectionToProxy();
 		}
+		RemoveComponent(DynamicMeshComponent);
 	}
 	DynamicMeshComponents.Reset();
 }
 
-TObjectPtr<UDynamicMeshComponent>& FDataflowPreviewScene::AddDynamicMeshComponent(UE::Geometry::FDynamicMesh3&& DynamicMesh, const TArray<UMaterialInterface*>& MaterialSet)
+TObjectPtr<UDynamicMeshComponent>& FDataflowConstructionScene::AddDynamicMeshComponent(UE::Geometry::FDynamicMesh3&& DynamicMesh, const TArray<UMaterialInterface*>& MaterialSet)
 {
-	TObjectPtr<UDynamicMeshComponent> DynamicMeshComponent = NewObject<UDynamicMeshComponent>(DynamicMeshActor);
+	TObjectPtr<UDynamicMeshComponent> DynamicMeshComponent = NewObject<UDynamicMeshComponent>(RootSceneActor);
 		
 	DynamicMeshComponent->SetMesh(MoveTemp(DynamicMesh));
-
-	// @todo(Dataflow) : Material support
-	//if (FDataflowEditorStyle::Get().DefaultMaterial)
+	
+	// @todo(Material) This is just to have a material, we should transfer the materials from the assets if they have them. 
+	if (DataflowContent && DataflowContent->DataflowAsset)
+	{
+		DynamicMeshComponent->ConfigureMaterialSet({ DataflowContent->DataflowAsset->Material });
+	}
+	//else if (FDataflowEditorStyle::Get().DefaultMaterial)
 	//{
 	//	DynamicMeshComponent->ConfigureMaterialSet({ FDataflowEditorStyle::Get().DefaultMaterial });
 	//}
@@ -174,81 +221,28 @@ TObjectPtr<UDynamicMeshComponent>& FDataflowPreviewScene::AddDynamicMeshComponen
 	//{
 	//	DynamicMeshComponent->ValidateMaterialSlots(true, false);
 	//}
-	// just vertex colors atm
-	DynamicMeshComponent->SetOverrideRenderMaterial(FDataflowEditorStyle::Get().VertexMaterial);
-	DynamicMeshComponent->SetShadowsEnabled(false);
-
-
+	else
+	{
+		DynamicMeshComponent->SetOverrideRenderMaterial(FDataflowEditorStyle::Get().VertexMaterial);
+		DynamicMeshComponent->SetShadowsEnabled(false);
+	}
 	DynamicMeshComponent->SelectionOverrideDelegate = UPrimitiveComponent::FSelectionOverride::CreateRaw(this, &FDataflowPreviewScene::IsComponentSelected);
-	DynamicMeshComponent->RegisterComponentWithWorld(GetWorld());
-	DynamicMeshComponent->OnMeshChanged.Add(FSimpleMulticastDelegate::FDelegate::CreateLambda([this](){}));
-
 	DynamicMeshComponent->UpdateBounds();
+
+	AddComponent(DynamicMeshComponent, DynamicMeshComponent->GetRelativeTransform());
 		
 	const int32 ElementIndex = DynamicMeshComponents.Emplace(DynamicMeshComponent);
 	return DynamicMeshComponents[ElementIndex];
 }
 
-void FDataflowPreviewScene::UpdateDynamicMeshComponents()
-{
-	using namespace UE::Geometry;//FDynamicMesh3
-
-	ResetDynamicMeshComponents();
-
-	if (EditorContent)
-	{
-		TObjectPtr<UDataflow> DataflowAsset = EditorContent->GetDataflowAsset();
-		TSharedPtr<Dataflow::FEngineContext> DataflowContext = EditorContent->GetDataflowContext();
-		if (DataflowAsset && DataflowContext)
-		{
-			for (const UDataflowEdNode* Target : DataflowAsset->GetRenderTargets())
-			{
-				if (Target)
-				{
-					FDynamicMesh3 DynamicMesh;
-					DynamicMesh.EnableAttributes();
-
-					TSharedPtr<FManagedArrayCollection> RenderCollection(new FManagedArrayCollection);
-					GeometryCollection::Facades::FRenderingFacade Facade(*RenderCollection);
-					Facade.DefineSchema();
-
-					Target->Render(Facade, DataflowContext);
-					Dataflow::Conversion::RenderingFacadeToDynamicMesh(Facade, DynamicMesh);
-
-					if (Target == EditorContent->GetPrimarySelectedNode())
-					{
-						EditorContent->SetPrimaryRenderCollection(RenderCollection);
-					}
-
-					// post updates
-					{
-						// Use per-triangle normals for the 2D view
-						//UE::Geometry::FMeshNormals::InitializeMeshToPerTriangleNormals(&LodMesh);
-					}
-					{
-						//@todo(Dataflow) :: Add material support
-						//SetUpDynamicMeshComponentMaterial(ClothFacade, *DynamicMeshComponent);
-					}
-
-
-					AddDynamicMeshComponent(MoveTemp(DynamicMesh), {});
-				}
-			}
-		}
-
-		EditorContent->SetIsDirty(false);
-	}
-}
-
-
-void FDataflowPreviewScene::AddWireframeMeshElementsVisualizer()
+void FDataflowConstructionScene::AddWireframeMeshElementsVisualizer()
 {
 	ensure(WireframeDraw==nullptr);
 	if (DynamicMeshComponents.Num())
 	{
 		// Set up the wireframe display of the rest space mesh.
 
-		WireframeDraw = NewObject<UMeshElementsVisualizer>(DynamicMeshActor);
+		WireframeDraw = NewObject<UMeshElementsVisualizer>(RootSceneActor);
 		WireframeDraw->CreateInWorld(GetWorld(), FTransform::Identity);
 
 		WireframeDraw->Settings->DepthBias = 2.0;
@@ -274,29 +268,11 @@ void FDataflowPreviewScene::AddWireframeMeshElementsVisualizer()
 			const bool bRestSpaceMeshVisible = DynamicMeshComponent->GetVisibleFlag();
 			WireframeDraw->Settings->bVisible = bRestSpaceMeshVisible && bConstructionViewWireframe;
 		}
-
-		// Some interactive tools will hide the input DynamicMeshComponent and create their own temporary PreviewMesh for visualization. If this
-		// occurs, we should also hide the corresponding Wireframe and Seam drawing (and un-hide it when the tool finishes).
-		/*
-		* // @todo(brice) : Deal with this.
-		UActorComponent::MarkRenderStateDirtyEvent.AddWeakLambda(this, [this](UActorComponent& ActorComponent)
-		{
-			if (!DynamicMeshComponent)
-			{
-				return;
-			}
-		    const bool bRestSpaceMeshVisible = DynamicMeshComponent->GetVisibleFlag();
-			if (WireframeDraw)
-			{
-				WireframeDraw->Settings->bVisible = bRestSpaceMeshVisible && bConstructionViewWireframe;
-			}
-		});
-		*/
 		PropertyObjectsToTick.Add(WireframeDraw->Settings);
 	}
 }
 
-void FDataflowPreviewScene::ResetWireframeMeshElementsVisualizer()
+void FDataflowConstructionScene::ResetWireframeMeshElementsVisualizer()
 {
 	if (WireframeDraw)
 	{
@@ -305,131 +281,93 @@ void FDataflowPreviewScene::ResetWireframeMeshElementsVisualizer()
 	WireframeDraw = nullptr;
 }
 
-void FDataflowPreviewScene::UpdateWireframeMeshElementsVisualizer()
+void FDataflowConstructionScene::UpdateWireframeMeshElementsVisualizer()
 {
 	ResetWireframeMeshElementsVisualizer();
 	AddWireframeMeshElementsVisualizer();
 }
 
-void FDataflowPreviewScene::UpdateSkeletalMeshComponent()
+bool FDataflowConstructionScene::HasRenderableGeometry()
 {
-	if (SkeletalMeshComponent)
+	for (auto& DynamicMeshComponent : DynamicMeshComponents)
 	{
-		if (EditorContent->GetSkeletalMesh())
-		{
-			if (EditorContent->GetSkeletalMesh() != SkeletalMeshComponent->GetSkeletalMeshAsset())
-			{
-				SkeletalMeshComponent->SetSkeletalMeshAsset(EditorContent->GetSkeletalMesh());
-			}
-
-			if (EditorContent->GetAnimationAsset())
-			{
-				PreviewAnimInstance = NewObject<UAnimSingleNodeInstance>(SkeletalMeshComponent);
-				PreviewAnimInstance->SetAnimationAsset(EditorContent->GetAnimationAsset());
-				SkeletalMeshComponent->SetAnimationMode(EAnimationMode::AnimationSingleNode);
-				SkeletalMeshComponent->InitAnim(true);
-				SkeletalMeshComponent->AnimationData.PopulateFrom(PreviewAnimInstance);
-				SkeletalMeshComponent->AnimScriptInstance = PreviewAnimInstance;
-				SkeletalMeshComponent->AnimScriptInstance->InitializeAnimation();
-				SkeletalMeshComponent->ValidateAnimation();
-			}
-			else
-			{
-				SkeletalMeshComponent->Stop();
-				SkeletalMeshComponent->AnimationData = FSingleAnimationPlayData();
-				SkeletalMeshComponent->AnimScriptInstance = nullptr;
-			}
-		}
-		else
+		if (DynamicMeshComponent->GetMesh()->TriangleCount() > 0)
 		{
-			SkeletalMeshComponent->SetSkeletalMeshAsset(nullptr);
-			SkeletalMeshComponent->Stop();
-			SkeletalMeshComponent->AnimationData = FSingleAnimationPlayData();
-			SkeletalMeshComponent->AnimScriptInstance = nullptr;
+			return true;
 		}
-
-		SkeletalMeshComponent->UpdateBounds();
 	}
+	return false;
 }
 
-FBox FDataflowPreviewScene::GetBoundingBox() const
+void FDataflowConstructionScene::ResetConstructionScene()
 {
-	FBox SceneBounds(EForceInit::ForceInit);
-
-	if (SkeletalMeshComponent)
-	{
-		FTransform ComponentTransform = SkeletalMeshComponent->GetComponentTransform();
-		FBox LocalBox = SkeletalMeshComponent->GetLocalBounds().GetBox();
-		SceneBounds += LocalBox.TransformBy(ComponentTransform);
-	}
+	// Some objects, like the UMeshElementsVisualizer and Settings Objects
+	// are not part of a tool, so they won't get ticked.This member holds
+	// ticked objects that get rebuilt on Update
+	PropertyObjectsToTick.Empty();
 
-	for (const TObjectPtr<UDynamicMeshComponent>& MeshComponent : DynamicMeshComponents)
-	{
-		if (MeshComponent)
-		{
-			SceneBounds += MeshComponent->Bounds.GetBox();
-		}
-	}
+	ResetWireframeMeshElementsVisualizer();
 
-	return SceneBounds;
+	ResetDynamicMeshComponents();
 }
 
-bool FDataflowPreviewScene::IsComponentSelected(const UPrimitiveComponent* InComponent) const
+void FDataflowConstructionScene::UpdateConstructionScene()
 {
-	if(DataflowModeManager.IsValid())
-	{
-		if (const UTypedElementSelectionSet* const TypedElementSelectionSet = DataflowModeManager->GetEditorSelectionSet())
-		{
-			if (const FTypedElementHandle ComponentElement = UEngineElementsLibrary::AcquireEditorComponentElementHandle(InComponent))
-			{
-				const bool bElementSelected = TypedElementSelectionSet->IsElementSelected(ComponentElement, FTypedElementIsSelectedOptions());
-				return bElementSelected;
-			}
-		}
-	}
-	return false;
+	ResetConstructionScene();
+
+	// The preview scene for the construction view will be
+	// cleared and rebuilt from scratch. This will genrate a 
+	// list of UPrimitiveComponents for rendering.
+	UpdateDynamicMeshComponents();
+	
+	// Attach a wireframe renderer to the DynamicMeshComponents
+	UpdateWireframeMeshElementsVisualizer();
+
+	// Manage Selection and Tool Interaction
+    if (DataflowModeManager.IsValid())
+    {
+    	USelection* SelectedComponents = DataflowModeManager->GetSelectedComponents();
+    	SelectedComponents->DeselectAll();
+    	for (const TObjectPtr<UDynamicMeshComponent>& DynamicMeshComponent : DynamicMeshComponents)
+    	{
+    		SelectedComponents->Select(DynamicMeshComponent);
+    		DynamicMeshComponent->PushSelectionToProxy();
+    	}
+    }
+	DataflowContent->SetIsDirty(false);
 }
 
-FBox FDataflowPreviewScene::SelectedComponentBounds() const
+FDataflowSimulationScene::FDataflowSimulationScene(FPreviewScene::ConstructionValues ConstructionValues, TObjectPtr<UDataflowBaseContent> InEditorContent) 
+	: FDataflowPreviewScene(ConstructionValues,InEditorContent)
 {
-	FBox Bounds(ForceInit);
-	if (DataflowModeManager.IsValid())
+	DataflowContent->RegisterWorldContent(this, RootSceneActor);
+
+	TInlineComponentArray<UPrimitiveComponent*> PrimComponents;
+	RootSceneActor->GetComponents(PrimComponents);
+
+	for(UPrimitiveComponent* PrimComponent : PrimComponents)
 	{
-		const USelection* const SelectedComponents = DataflowModeManager->GetSelectedComponents();
-		for (int32 i = 0; i < SelectedComponents->Num(); ++i)
-		{
-			const UObject* const SelectedObject = SelectedComponents->GetSelectedObject(i);
-			if (const UDynamicMeshComponent* const DynamicMeshComponent = Cast<UDynamicMeshComponent>(SelectedObject))
-			{
-				Bounds += DynamicMeshComponent->Bounds.GetBox();
-			}
-		}
+		PrimComponent->SelectionOverrideDelegate =
+			UPrimitiveComponent::FSelectionOverride::CreateRaw(this, &FDataflowPreviewScene::IsComponentSelected);
 	}
-	return Bounds;
 }
 
-bool FDataflowPreviewScene::HasRenderableGeometry()
+FDataflowSimulationScene::~FDataflowSimulationScene()
 {
-	for (auto& DynamicMeshComponent : DynamicMeshComponents)
+	TInlineComponentArray<UPrimitiveComponent*> PrimComponents;
+	RootSceneActor->GetComponents(PrimComponents);
+
+	for(UPrimitiveComponent* PrimComponent : PrimComponents)
 	{
-		if (DynamicMeshComponent->GetMesh()->TriangleCount() > 0)
-		{
-			return true;
-		}
+		PrimComponent->SelectionOverrideDelegate.Unbind();
 	}
-	return false;
+	
+	DataflowContent->UnregisterWorldContent(this);
 }
 
-
-void FDataflowPreviewScene::AddReferencedObjects(FReferenceCollector& Collector)
+void FDataflowSimulationScene::TickDataflowScene(const float DeltaSeconds)
 {
-	FAdvancedPreviewScene::AddReferencedObjects(Collector);
-	Collector.AddReferencedObject(SkeletalMeshComponent);
-	Collector.AddReferencedObject(SkeletalMeshActor);
-	Collector.AddReferencedObject(EditorContent);
-	Collector.AddReferencedObject(DynamicMeshActor);
-	Collector.AddReferencedObjects(DynamicMeshComponents);
-	Collector.AddReferencedObject(WireframeDraw);
+	GetWorld()->Tick(ELevelTick::LEVELTICK_All, DeltaSeconds);
 }
 
 #undef LOCTEXT_NAMESPACE
diff --git a/Engine/Plugins/Experimental/Dataflow/Source/DataflowEditor/Private/Dataflow/DataflowSimulationPanel.cpp b/Engine/Plugins/Experimental/Dataflow/Source/DataflowEditor/Private/Dataflow/DataflowSimulationPanel.cpp
new file mode 100644
--- /dev/null
+++ b/Engine/Plugins/Experimental/Dataflow/Source/DataflowEditor/Private/Dataflow/DataflowSimulationPanel.cpp
@@ -0,0 +1,451 @@
+﻿// Copyright Epic Games, Inc. All Rights Reserved.
+
+#include "Dataflow/DataflowSimulationPanel.h"
+#include "Dataflow/DataflowPreviewScene.h"
+#include "Widgets/SBoxPanel.h"
+#include "SScrubControlPanel.h"
+#include "Animation/AnimSingleNodeInstance.h"
+#include "Animation/AnimSequence.h"
+#include "Animation/BlendSpace.h"
+#include "Dataflow/DataflowContent.h"
+#include "Widgets/Input/SButton.h"
+#include "Widgets/Images/SImage.h"
+
+#define LOCTEXT_NAMESPACE "DataflowSimulationScrubPanel"
+
+void SDataflowSimulationPanel::Construct( const SDataflowSimulationPanel::FArguments& InArgs, TWeakPtr<FDataflowSimulationScene> InPreviewScene)
+{
+	SimulationScene = InPreviewScene;
+
+	// Skip adding the the Loop button so we can add our own
+	TArray<FTransportControlWidget> TransportControlWidgets;
+	for (const ETransportControlWidgetType Type : TEnumRange<ETransportControlWidgetType>())
+	{
+		if ((Type != ETransportControlWidgetType::Custom) && (Type != ETransportControlWidgetType::Loop))
+		{
+			TransportControlWidgets.Add(FTransportControlWidget(Type));
+		}
+	}
+	const FTransportControlWidget NewWidget(FOnMakeTransportWidget::CreateSP(this, &SDataflowSimulationPanel::OnCreatePreviewPlaybackModeWidget));
+	TransportControlWidgets.Add(NewWidget);
+
+	this->ChildSlot
+	[
+		SNew(SHorizontalBox)
+		.AddMetaData<FTagMetaData>(TEXT("ClothAnimScrub.Scrub"))
+		+SHorizontalBox::Slot()
+		.HAlign(HAlign_Fill) 
+		.VAlign(VAlign_Center)
+		.FillWidth(1)
+		.Padding(0.0f)
+		[
+			SAssignNew(ScrubControlPanel, SScrubControlPanel)
+			.IsEnabled(true)
+			.Value(this, &SDataflowSimulationPanel::GetScrubValue)
+			.NumOfKeys(this, &SDataflowSimulationPanel::GetNumberOfKeys)
+			.SequenceLength(this, &SDataflowSimulationPanel::GetSequenceLength)
+			.DisplayDrag(this, &SDataflowSimulationPanel::GetDisplayDrag)
+			.OnValueChanged(this, &SDataflowSimulationPanel::OnValueChanged)
+			.OnBeginSliderMovement(this, &SDataflowSimulationPanel::OnBeginSliderMovement)
+			.OnClickedForwardPlay(this, &SDataflowSimulationPanel::OnClick_Forward)
+			.OnClickedForwardStep(this, &SDataflowSimulationPanel::OnClick_Forward_Step)
+			.OnClickedForwardEnd(this, &SDataflowSimulationPanel::OnClick_Forward_End)
+			.OnClickedBackwardPlay(this, &SDataflowSimulationPanel::OnClick_Backward)
+			.OnClickedBackwardStep(this, &SDataflowSimulationPanel::OnClick_Backward_Step)
+			.OnClickedBackwardEnd(this, &SDataflowSimulationPanel::OnClick_Backward_End)
+			.OnTickPlayback(this, &SDataflowSimulationPanel::OnTickPlayback)
+			.OnGetPlaybackMode(this, &SDataflowSimulationPanel::GetPlaybackMode)
+			.ViewInputMin(InArgs._ViewInputMin)
+			.ViewInputMax(InArgs._ViewInputMax)
+			.bDisplayAnimScrubBarEditing(false)
+			.bAllowZoom(false)
+			.IsRealtimeStreamingMode(false)
+			.TransportControlWidgetsToCreate(TransportControlWidgets)
+		]
+	];
+	UpdatePreviewAnimationInstance();
+	ApplyPlaybackSettings();
+}
+
+
+TSharedRef<SWidget> SDataflowSimulationPanel::OnCreatePreviewPlaybackModeWidget()
+{
+	PreviewPlaybackModeButton = SNew(SButton)
+		.OnClicked(this, &SDataflowSimulationPanel::OnClick_PreviewPlaybackMode)
+		.ButtonStyle( FAppStyle::Get(), "Animation.PlayControlsButton" )
+		.IsFocusable(false)
+		.ToolTipText_Lambda([&]()
+		{
+			if (PreviewPlaybackMode == EDataflowPlaybackMode::Default)
+			{
+				return LOCTEXT("PlaybackModeDefaultTooltip", "Linear playback");
+			}
+			else if (PreviewPlaybackMode == EDataflowPlaybackMode::Looping)
+			{
+				return LOCTEXT("PlaybackModeLoopingTooltip", "Looping playback");
+			}
+			else
+			{
+				return LOCTEXT("PlaybackModePingPongTooltip", "Ping pong playback");
+			}
+		})
+		.ContentPadding(0.0f);
+
+	TWeakPtr<SButton> WeakButton = PreviewPlaybackModeButton;
+
+	PreviewPlaybackModeButton->SetContent(SNew(SImage)
+		.ColorAndOpacity(FSlateColor::UseSubduedForeground())
+		.Image_Lambda([&, WeakButton]()
+		{
+			if (PreviewPlaybackMode == EDataflowPlaybackMode::Default)
+			{
+				return FAppStyle::Get().GetBrush("Animation.Loop.Disabled");
+			}
+			else if (PreviewPlaybackMode == EDataflowPlaybackMode::Looping)
+			{
+				return FAppStyle::Get().GetBrush("Animation.Loop.Enabled");
+			}
+			else
+			{
+				return FAppStyle::Get().GetBrush("Animation.Loop.SelectionRange");		// TODO: Replace with a back and forth type icon
+			}
+		})
+	);
+
+	TSharedRef<SHorizontalBox> PreviewPlaybackModeBox = SNew(SHorizontalBox);
+	PreviewPlaybackModeBox->AddSlot()
+	.AutoWidth()
+	[
+		PreviewPlaybackModeButton.ToSharedRef()
+	];
+
+	return PreviewPlaybackModeBox;
+}
+
+FReply SDataflowSimulationPanel::OnClick_Forward_Step()
+{
+	if (UAnimSingleNodeInstance* const PreviewInstance = GetPreviewAnimationInstance())
+	{
+		PreviewInstance->SetPlaying(false);
+		PreviewInstance->StepForward();
+	}
+
+	return FReply::Handled();
+}
+
+FReply SDataflowSimulationPanel::OnClick_Forward_End()
+{
+	if (UAnimSingleNodeInstance* const PreviewInstance = GetPreviewAnimationInstance())
+	{
+		PreviewInstance->SetPlaying(false);
+		PreviewInstance->SetPosition(PreviewInstance->GetLength(), false);
+	}
+
+	return FReply::Handled();
+}
+
+FReply SDataflowSimulationPanel::OnClick_Backward_Step()
+{
+	if (UAnimSingleNodeInstance* const PreviewInstance = GetPreviewAnimationInstance())
+	{
+		PreviewInstance->SetPlaying(false);
+		PreviewInstance->StepBackward();
+	}
+
+	return FReply::Handled();
+}
+
+FReply SDataflowSimulationPanel::OnClick_Backward_End()
+{
+	if (UAnimSingleNodeInstance* const PreviewInstance = GetPreviewAnimationInstance())
+	{
+		PreviewInstance->SetPlaying(false);
+		PreviewInstance->SetPosition(0.f, false);
+	}
+
+	return FReply::Handled();
+}
+
+FReply SDataflowSimulationPanel::OnClick_Forward()
+{
+	if (UAnimSingleNodeInstance* const PreviewInstance = GetPreviewAnimationInstance())
+	{
+		const bool bIsReverse = PreviewInstance->IsReverse();
+		const bool bIsPlaying = PreviewInstance->IsPlaying();
+
+		// if current bIsReverse and bIsPlaying, we'd like to just turn off reverse
+		if (bIsReverse && bIsPlaying)
+		{
+			PreviewInstance->SetReverse(false);
+		}
+		// already playing, simply pause
+		else if (bIsPlaying) 
+		{
+			PreviewInstance->SetPlaying(false);
+		}
+		// if not playing, play forward
+		else 
+		{
+			//if we're at the end of the animation, jump back to the beginning before playing
+			if ( GetScrubValue() >= GetSequenceLength() )
+			{
+				PreviewInstance->SetPosition(0.0f, false);
+			}
+
+			PreviewInstance->SetReverse(false);
+			PreviewInstance->SetPlaying(true);
+		}
+	}
+
+	return FReply::Handled();
+}
+
+FReply SDataflowSimulationPanel::OnClick_Backward()
+{
+	if (UAnimSingleNodeInstance* const PreviewInstance = GetPreviewAnimationInstance())
+	{
+		const bool bIsReverse = PreviewInstance->IsReverse();
+		const bool bIsPlaying = PreviewInstance->IsPlaying();
+
+		// if currently playing forward, just simply turn on reverse
+		if (!bIsReverse && bIsPlaying)
+		{
+			PreviewInstance->SetReverse(true);
+		}
+		else if (bIsPlaying)
+		{
+			PreviewInstance->SetPlaying(false);
+		}
+		else
+		{
+			// if we're at the beginning of the animation, jump back to the end before playing
+			if ( GetScrubValue() <= 0.0f )
+			{
+				PreviewInstance->SetPosition(GetSequenceLength(), false);
+			}
+
+			PreviewInstance->SetPlaying(true);
+			PreviewInstance->SetReverse(true);
+		}
+	}
+
+	return FReply::Handled();
+}
+
+FReply SDataflowSimulationPanel::OnClick_PreviewPlaybackMode()
+{
+	if (PreviewPlaybackMode == EDataflowPlaybackMode::Default)
+	{
+		PreviewPlaybackMode = EDataflowPlaybackMode::Looping;
+		
+		if (UAnimSingleNodeInstance* const PreviewInstance = GetPreviewAnimationInstance())
+		{
+			// If we paused due to hitting the end point, start playing again when entering loop mode
+			const float CurrentTime = PreviewInstance->GetCurrentTime();
+			const float PlayRate = PreviewInstance->GetPlayRate();
+			const float AssetPlayLength = PreviewInstance->CurrentAsset->GetPlayLength();
+			if (PlayRate < 0.0 && CurrentTime <= 0.0)
+			{
+				PreviewInstance->SetPlaying(true);
+			}
+			else if (PlayRate > 0.0 && CurrentTime >= AssetPlayLength)
+			{
+				PreviewInstance->SetPlaying(true);
+			}
+		}
+	}
+	else if (PreviewPlaybackMode == EDataflowPlaybackMode::Looping)
+	{
+		PreviewPlaybackMode = EDataflowPlaybackMode::PingPong;
+	}
+	else
+	{
+		PreviewPlaybackMode = EDataflowPlaybackMode::Default;
+
+		if (UAnimSingleNodeInstance* const PreviewInstance = GetPreviewAnimationInstance())
+		{
+			// If we're switching to linear playback, set it to forward mode
+			PreviewInstance->SetReverse(false);
+		}
+	}
+
+	ApplyPlaybackSettings();
+
+	return FReply::Handled();
+}
+
+void SDataflowSimulationPanel::ApplyPlaybackSettings()
+{
+	UAnimSingleNodeInstance* const PreviewInstance = GetPreviewAnimationInstance();
+
+	switch (PreviewPlaybackMode)
+	{
+	case EDataflowPlaybackMode::Default:
+		if (PreviewInstance)
+		{
+			PreviewInstance->SetLooping(false);
+		}
+		break;
+	case EDataflowPlaybackMode::Looping:
+		if (PreviewInstance)
+		{
+			PreviewInstance->SetLooping(true);
+		}
+		break;
+	case EDataflowPlaybackMode::PingPong:
+		if (PreviewInstance)
+		{
+			PreviewInstance->SetLooping(false);
+		}
+		break;
+	}
+}
+
+void SDataflowSimulationPanel::OnTickPlayback(double InCurrentTime, float InDeltaTime)
+{
+	if (UAnimSingleNodeInstance* const PreviewInstance = GetPreviewAnimationInstance())
+	{
+		const float CurrentTime = PreviewInstance->GetCurrentTime();
+		const float PlayRate = PreviewInstance->GetPlayRate();
+		const float AssetPlayLength = PreviewInstance->CurrentAsset->GetPlayLength();
+
+		if (PreviewPlaybackMode == EDataflowPlaybackMode::PingPong)
+		{
+			if (PlayRate < 0.0 && CurrentTime <= 0.0)
+			{
+				PreviewInstance->SetReverse(!PreviewInstance->IsReverse());
+			}
+			else if (PlayRate > 0.0 && CurrentTime >= AssetPlayLength)
+			{
+				PreviewInstance->SetReverse(!PreviewInstance->IsReverse());
+			}
+		}
+	}
+}
+
+EPlaybackMode::Type SDataflowSimulationPanel::GetPlaybackMode() const
+{
+	if (const UAnimSingleNodeInstance* const PreviewInstance = GetPreviewAnimationInstance())
+	{
+		if (PreviewInstance->IsPlaying())
+		{
+			return PreviewInstance->IsReverse() ? EPlaybackMode::PlayingReverse : EPlaybackMode::PlayingForward;
+		}
+		return EPlaybackMode::Stopped;
+	}
+	
+	return EPlaybackMode::Stopped;
+}
+
+void SDataflowSimulationPanel::OnValueChanged(float NewValue)
+{
+	if (UAnimSingleNodeInstance* const PreviewInstance = GetPreviewAnimationInstance())
+	{
+		PreviewInstance->SetPosition(NewValue);
+	}
+}
+
+void SDataflowSimulationPanel::OnBeginSliderMovement()
+{
+	if (UAnimSingleNodeInstance* const PreviewInstance = GetPreviewAnimationInstance())
+	{
+		PreviewInstance->SetPlaying(false);
+	}
+}
+
+uint32 SDataflowSimulationPanel::GetNumberOfKeys() const
+{
+	if (const UAnimSingleNodeInstance* const PreviewInstance = GetPreviewAnimationInstance())
+	{
+		const float Length = const_cast<UAnimSingleNodeInstance*>(PreviewInstance)->GetLength();
+
+		// if anim sequence, use correct num frames
+		int32 NumKeys = (int32)(Length / 0.0333f);
+
+		if (PreviewInstance->CurrentAsset)
+		{
+			if (PreviewInstance->CurrentAsset->IsA(UAnimSequenceBase::StaticClass()))
+			{
+				NumKeys = CastChecked<UAnimSequenceBase>(PreviewInstance->CurrentAsset)->GetNumberOfSampledKeys();
+			}
+			else if (PreviewInstance->CurrentAsset->IsA(UBlendSpace::StaticClass()))
+			{
+				// Blendspaces dont display frame notches, so just return 0 here
+				NumKeys = 0;
+			}
+		}
+		return NumKeys;
+	}
+
+	return 1;
+}
+
+float SDataflowSimulationPanel::GetSequenceLength() const
+{
+	if (const UAnimSingleNodeInstance* const PreviewInstance = GetPreviewAnimationInstance())	// non-const because UAnimSingleNodeInstance::GetLength() is non-const
+	{
+		return const_cast<UAnimSingleNodeInstance*>(PreviewInstance)->GetLength();
+	}
+
+	return 0.f;
+}
+
+float SDataflowSimulationPanel::GetScrubValue() const
+{
+	if (const UAnimSingleNodeInstance* const PreviewInstance = GetPreviewAnimationInstance())
+	{
+		return PreviewInstance->GetCurrentTime(); 
+	}
+
+	return 0.f;
+}
+
+bool SDataflowSimulationPanel::GetDisplayDrag() const
+{
+	const UAnimSingleNodeInstance* const PreviewInstance = GetPreviewAnimationInstance();
+	if (PreviewInstance && PreviewInstance->CurrentAsset)
+	{
+		return true;
+	}
+
+	return false;
+}
+
+UAnimSingleNodeInstance* SDataflowSimulationPanel::GetPreviewAnimationInstance()
+{
+	if (const TSharedPtr<const FDataflowSimulationScene> PreviewScene = SimulationScene.Pin())
+	{
+		if(const TObjectPtr<UDataflowSkeletalContent> DataflowContent = Cast<UDataflowSkeletalContent>(PreviewScene->GetDataflowContent()))
+		{
+			return DataflowContent->GetAnimationInstance();
+		}
+	}
+
+	return nullptr;
+}
+
+const UAnimSingleNodeInstance* SDataflowSimulationPanel::GetPreviewAnimationInstance() const
+{
+	if (const TSharedPtr<const FDataflowSimulationScene> PreviewScene = SimulationScene.Pin())
+	{
+		if(const TObjectPtr<UDataflowSkeletalContent> DataflowContent = Cast<UDataflowSkeletalContent>(PreviewScene->GetDataflowContent()))
+		{
+			return DataflowContent->GetAnimationInstance();
+		}
+	}
+
+	return nullptr;
+}
+
+void SDataflowSimulationPanel::UpdatePreviewAnimationInstance()
+{
+	if (const TSharedPtr<FDataflowSimulationScene> PreviewScene = SimulationScene.Pin())
+	{
+		if(TObjectPtr<UDataflowSkeletalContent> DataflowContent = Cast<UDataflowSkeletalContent>(PreviewScene->GetDataflowContent()))
+		{
+			DataflowContent->UpdateAnimationInstance();
+		}
+	}
+}
+
+#undef LOCTEXT_NAMESPACE
diff --git a/Engine/Plugins/Experimental/Dataflow/Source/DataflowEditor/Private/Dataflow/DataflowToolTarget.cpp b/Engine/Plugins/Experimental/Dataflow/Source/DataflowEditor/Private/Dataflow/DataflowToolTarget.cpp
--- a/Engine/Plugins/Experimental/Dataflow/Source/DataflowEditor/Private/Dataflow/DataflowToolTarget.cpp
+++ b/Engine/Plugins/Experimental/Dataflow/Source/DataflowEditor/Private/Dataflow/DataflowToolTarget.cpp
@@ -5,7 +5,7 @@
 #include "ConversionUtils/DynamicMeshViaMeshDescriptionUtil.h"
 #include "Dataflow/CollectionRenderingPatternUtility.h"
 #include "Dataflow/DataflowEditor.h"
-#include "Dataflow/DataflowEditorContent.h"
+#include "Dataflow/DataflowContent.h"
 #include "Dataflow/DataflowEditorToolkit.h"
 #include "Dataflow/DataflowEditorUtil.h"
 #include "Dataflow/DataflowEdNode.h"
@@ -161,30 +161,37 @@ bool UDataflowReadOnlyToolTargetFactory::CanBuildTarget(UObject* SourceObject, c
 	// If you want to make the tool target work with some subclass of UDataflow,
 	// just add another factory that allows that class specifically(but make sure that
 	// GetMeshDescription and such work properly)
-	UDataflowEditorContent* Content = CastChecked<UDataflowEditorContent>(SourceObject);
-	const UDataflow* Dataflow = Content->GetDataflowAsset();
-
-
-
-	return Dataflow &&
-		ExactCast<UDataflow>(Cast<UDataflow>(Dataflow)) &&
-		!ExactCast<UDataflow>(Cast<UDataflow>(Dataflow))->GetOutermost()->bIsCookedForEditor &&
-		Requirements.AreSatisfiedBy(UDataflowReadOnlyToolTarget::StaticClass());
+	const TObjectPtr<UDataflowBaseContent> BaseContent = CastChecked<UDataflowBaseContent>(SourceObject);
+	if(BaseContent)
+	{
+		const UDataflow* Dataflow = BaseContent->GetDataflowAsset();
+	
+		return Dataflow &&
+			ExactCast<UDataflow>(Cast<UDataflow>(Dataflow)) &&
+			!ExactCast<UDataflow>(Cast<UDataflow>(Dataflow))->GetOutermost()->bIsCookedForEditor &&
+			Requirements.AreSatisfiedBy(UDataflowReadOnlyToolTarget::StaticClass());
+	}
+	return false;
 }
 
 UToolTarget* UDataflowReadOnlyToolTargetFactory::BuildTarget(UObject* SourceObject, const FToolTargetTypeRequirements& Requirements)
 {
-	UDataflowEditorContent* Content = CastChecked<UDataflowEditorContent>(SourceObject);
+	const TObjectPtr<UDataflowBaseContent> BaseContent = CastChecked<UDataflowBaseContent>(SourceObject);
+
+	if(BaseContent)
+	{
+		UDataflowReadOnlyToolTarget* Target = NewObject<UDataflowReadOnlyToolTarget>();
+		Target->Asset = BaseContent->GetDataflowOwner();
+		Target->Dataflow = BaseContent->GetDataflowAsset();
+		Target->Context = Dataflow::GetContext(BaseContent);
 
-	UDataflowReadOnlyToolTarget* Target = NewObject<UDataflowReadOnlyToolTarget>();
-	Target->Asset = Content->GetDataflowOwner();
-	Target->Dataflow = Content->GetDataflowAsset();
-	Target->Context = Dataflow::GetContext(Content);
+		// @todo(brice) : I needed to comment this out?
+		//checkSlow(Target->Component.IsValid() && Requirements.AreSatisfiedBy(Target));
 
-	// @todo(brice) : I needed to comment this out?
-	//checkSlow(Target->Component.IsValid() && Requirements.AreSatisfiedBy(Target));
+		return Target;
+	}
 
-	return Target;
+	return nullptr;
 }
 
 bool UDataflowToolTargetFactory::CanBuildTarget(UObject* SourceObject, const FToolTargetTypeRequirements& Requirements) const
@@ -195,24 +202,34 @@ bool UDataflowToolTargetFactory::CanBuildTarget(UObject* SourceObject, const FTo
 	// If you want to make the tool target work with some subclass of UDataflow,
 	// just add another factory that allows that class specifically(but make sure that
 	// GetMeshDescription and such work properly)
-	UDataflowEditorContent* Content = CastChecked<UDataflowEditorContent>(SourceObject);
-	const UDataflow* Dataflow = Content->GetDataflowAsset();
 
-	return Dataflow &&
-		ExactCast<UDataflow>(Cast<UDataflow>(Dataflow)) &&
-		!ExactCast<UDataflow>(Cast<UDataflow>(Dataflow))->GetOutermost()->bIsCookedForEditor &&
-		Requirements.AreSatisfiedBy(UDataflowToolTarget::StaticClass());
+	const TObjectPtr<UDataflowBaseContent> BaseContent = CastChecked<UDataflowBaseContent>(SourceObject);
+	if(BaseContent)
+	{
+		const UDataflow* Dataflow = BaseContent->GetDataflowAsset();
+		
+		return Dataflow &&
+			ExactCast<UDataflow>(Cast<UDataflow>(Dataflow)) &&
+			!ExactCast<UDataflow>(Cast<UDataflow>(Dataflow))->GetOutermost()->bIsCookedForEditor &&
+			Requirements.AreSatisfiedBy(UDataflowToolTarget::StaticClass());
+	}
+	return false;
 }
 
 UToolTarget* UDataflowToolTargetFactory::BuildTarget(UObject* SourceObject, const FToolTargetTypeRequirements& Requirements)
 {
-	UDataflowEditorContent* Content = CastChecked<UDataflowEditorContent>(SourceObject);
-	UDataflowToolTarget* Target = NewObject<UDataflowToolTarget>();
-	Target->Asset = Content->GetDataflowOwner();
-	Target->Dataflow = Content->GetDataflowAsset();
-	Target->Context = Dataflow::GetContext(Content);
-	//checkSlow(Target->Component.IsValid() && Requirements.AreSatisfiedBy(Target));
-
-	return Target;
+	const TObjectPtr<UDataflowBaseContent> BaseContent = CastChecked<UDataflowBaseContent>(SourceObject);
+
+	if(BaseContent)
+	{
+		UDataflowToolTarget* Target = NewObject<UDataflowToolTarget>();
+		Target->Asset = BaseContent->GetDataflowOwner();
+		Target->Dataflow = BaseContent->GetDataflowAsset();
+		Target->Context = Dataflow::GetContext(BaseContent);
+		//checkSlow(Target->Component.IsValid() && Requirements.AreSatisfiedBy(Target));
+
+		return Target;
+	}
+	return nullptr;
 }
 
diff --git a/Engine/Plugins/Experimental/Dataflow/Source/DataflowEditor/Private/DataflowEditorTools/DataflowEditorWeightMapPaintTool.cpp b/Engine/Plugins/Experimental/Dataflow/Source/DataflowEditor/Private/DataflowEditorTools/DataflowEditorWeightMapPaintTool.cpp
--- a/Engine/Plugins/Experimental/Dataflow/Source/DataflowEditor/Private/DataflowEditorTools/DataflowEditorWeightMapPaintTool.cpp
+++ b/Engine/Plugins/Experimental/Dataflow/Source/DataflowEditor/Private/DataflowEditorTools/DataflowEditorWeightMapPaintTool.cpp
@@ -11,11 +11,10 @@
 #include "Changes/BasicChanges.h"
 #include "ContextObjectStore.h"
 #include "DataflowEditorTools/DataflowEditorWeightMapPaintBrushOps.h"
-#include "Dataflow/DataflowEditorContent.h"
+#include "Dataflow/DataflowContent.h"
 #include "Dataflow/DataflowEdNode.h"
 #include "Dataflow/DataflowObject.h"
 #include "Dataflow/DataflowCollectionAddScalarVertexPropertyNode.h"
-#include "Dataflow/DataflowPatternVertexType.h"
 #include "Dataflow/DataflowSNode.h"
 #include "Dataflow/DataflowGraphEditor.h"
 #include "Drawing/MeshElementsVisualizer.h"
@@ -75,7 +74,7 @@ UMeshSurfacePointTool* UDataflowEditorWeightMapPaintToolBuilder::CreateNewTool(c
 	UDataflowEditorWeightMapPaintTool* PaintTool = NewObject<UDataflowEditorWeightMapPaintTool>(SceneState.ToolManager);
 	PaintTool->SetWorld(SceneState.World);
 
-	if (UDataflowEditorContextObject* ContextObject = SceneState.ToolManager->GetContextObjectStore()->FindContext<UDataflowEditorContextObject>())
+	if (UDataflowContextObject* ContextObject = SceneState.ToolManager->GetContextObjectStore()->FindContext<UDataflowContextObject>())
 	{
 		PaintTool->SetDataflowEditorContextObject(ContextObject);
 	}
@@ -362,7 +361,7 @@ void UDataflowEditorWeightMapPaintTool::Setup()
 }
 
 
-void UDataflowEditorWeightMapPaintTool::SetDataflowEditorContextObject(TObjectPtr<UDataflowEditorContextObject> InDataflowEditorContextObject)
+void UDataflowEditorWeightMapPaintTool::SetDataflowEditorContextObject(TObjectPtr<UDataflowContextObject> InDataflowEditorContextObject)
 {
 	DataflowEditorContextObject = InDataflowEditorContextObject;
 }
diff --git a/Engine/Plugins/Experimental/Dataflow/Source/DataflowEditor/Public/Dataflow/DataflowEditor.h b/Engine/Plugins/Experimental/Dataflow/Source/DataflowEditor/Public/Dataflow/DataflowEditor.h
--- a/Engine/Plugins/Experimental/Dataflow/Source/DataflowEditor/Public/Dataflow/DataflowEditor.h
+++ b/Engine/Plugins/Experimental/Dataflow/Source/DataflowEditor/Public/Dataflow/DataflowEditor.h
@@ -8,7 +8,7 @@
 #include "Dataflow/AssetDefinition_DataflowAsset.h"
 #include "Dataflow/DataflowAssetFactory.h"
 #include "Dataflow/DataflowEditorCommands.h"
-#include "Dataflow/DataflowEditorModule.h"
+#include "Dataflow/DataflowContent.h"
 #include "Dataflow/DataflowSchema.h"
 #include "Dataflow/DataflowSNode.h"
 #include "Dataflow/DataflowSNodeFactories.h"
@@ -17,7 +17,7 @@
 #include "DataflowEditor.generated.h"
 
 class FDataflowEditorToolkit;
-class UDataflowEditorContent;
+class UDataflowBaseContent;
 
 /** 
  * The actual asset editor class doesn't have that much in it, intentionally. 
@@ -44,6 +44,9 @@ public:
 	virtual TSharedPtr<FBaseAssetToolkit> CreateToolkit() override;
 	virtual void Initialize(const TArray<TObjectPtr<UObject>>& InObjects) override;
 
+	/** Initialize an editor with a given content and an owner */
+	void InitializeContent(TObjectPtr<UDataflowBaseContent> BaseContent, const TObjectPtr<UObject>& ContentOwner);
+
 private :
 
 	friend class FDataflowEditorToolkit;
@@ -53,7 +56,7 @@ private :
 	// It is why the data flow asset/owner/skelmesh have been added here. Could be added
 	// in the subsystem if necessary
 	UPROPERTY()
-	TObjectPtr<UDataflowEditorContent> Content;
+	TObjectPtr<UDataflowBaseContent> DataflowContent;
 };
 
 DECLARE_LOG_CATEGORY_EXTERN(LogDataflowEditor, Log, All);
\ No newline at end of file
diff --git a/Engine/Plugins/Experimental/Dataflow/Source/DataflowEditor/Public/Dataflow/DataflowEditorCommands.h b/Engine/Plugins/Experimental/Dataflow/Source/DataflowEditor/Public/Dataflow/DataflowEditorCommands.h
--- a/Engine/Plugins/Experimental/Dataflow/Source/DataflowEditor/Public/Dataflow/DataflowEditorCommands.h
+++ b/Engine/Plugins/Experimental/Dataflow/Source/DataflowEditor/Public/Dataflow/DataflowEditorCommands.h
@@ -16,7 +16,7 @@ class UDataflowEdNode;
 struct FDataflowNode;
 class UEdGraphNode;
 class SDataflowGraphEditor;
-class UDataflowEditorContent;
+class UDataflowBaseContent;
 
 typedef TSet<class UObject*> FGraphPanelSelectionSet;
 
@@ -117,7 +117,7 @@ public:
 	*  OnPropertyValueChanged
 	*/
 	static void OnPropertyValueChanged(UDataflow* Graph, TSharedPtr<Dataflow::FEngineContext>& Context, Dataflow::FTimestamp& OutLastNodeTimestamp, const FPropertyChangedEvent& PropertyChangedEvent, const TSet<UObject*>& SelectedNodes = TSet<UObject*>());
-	static void OnAssetPropertyValueChanged(TObjectPtr<UDataflowEditorContent> Content, const FPropertyChangedEvent& PropertyChangedEvent);
+	static void OnAssetPropertyValueChanged(TObjectPtr<UDataflowBaseContent> Content, const FPropertyChangedEvent& PropertyChangedEvent);
 
 	/*
 	*  OnSelectedNodesChanged
diff --git a/Engine/Plugins/Experimental/Dataflow/Source/DataflowEditor/Public/Dataflow/DataflowEditorContent.h b/Engine/Plugins/Experimental/Dataflow/Source/DataflowEditor/Public/Dataflow/DataflowEditorContent.h
deleted file mode 100644
diff --git a/Engine/Plugins/Experimental/Dataflow/Source/DataflowEditor/Public/Dataflow/DataflowEditorMode.h b/Engine/Plugins/Experimental/Dataflow/Source/DataflowEditor/Public/Dataflow/DataflowEditorMode.h
--- a/Engine/Plugins/Experimental/Dataflow/Source/DataflowEditor/Public/Dataflow/DataflowEditorMode.h
+++ b/Engine/Plugins/Experimental/Dataflow/Source/DataflowEditor/Public/Dataflow/DataflowEditorMode.h
@@ -3,12 +3,11 @@
 #pragma once
 
 #include "BaseCharacterFXEditorMode.h"
-#include "Dataflow/DataflowPatternVertexType.h"
-
+#include "Dataflow/DataflowContent.h"
 #include "DataflowEditorMode.generated.h"
 
 class UDataflowComponent;
-class FDataflowPreviewScene;
+class FDataflowConstructionScene;
 class SDataflowGraphEditor;
 class FDataflowEditorViewportClient;
 class IDataflowEditorToolBuilder;
@@ -61,8 +60,8 @@ public:
 	}
 
 	/** Set the data flow preview scene to store the dynamic meshes generated by the tools */
-	void SetDataflowPreviewScene(FDataflowPreviewScene* DataflowPreviewScene);
-	const FDataflowPreviewScene* GetDataflowPreviewScene() const { return PreviewScene; }
+	void SetDataflowConstructionScene(FDataflowConstructionScene* DataflowPreviewScene);
+	const FDataflowConstructionScene* GetDataflowConstructionScene() const { return PreviewScene; }
 
 	/** Set the data flow graph editor to create nodes once the tools have ended*/
 	void SetDataflowGraphEditor(TSharedPtr<SDataflowGraphEditor> DataflowGraphEditor);
@@ -103,8 +102,6 @@ private:
 
 	void RegisterAddNodeCommand(TSharedPtr<FUICommandInfo> AddNodeCommand, const FName& NewNodeType, TSharedPtr<FUICommandInfo> StartToolCommand);
 
-	void SetPreviewScene(FDataflowPreviewScene* PreviewScene);
-
 	void SetRestSpaceViewportClient(TWeakPtr<FDataflowEditorViewportClient, ESPMode::ThreadSafe> ViewportClient);
 	void RefocusRestSpaceViewportClient();
 	void FirstTimeFocusRestSpaceViewport();
@@ -154,7 +151,7 @@ private:
 	void DeleteContextObject();
 
 	/** Dataflow preview scene from the toolkit */
-	FDataflowPreviewScene* PreviewScene = nullptr;
+	FDataflowConstructionScene* PreviewScene = nullptr;
 
 	/** Correspondence between node types and commands to launch tools */
 	TMap<FName, TSharedPtr<const FUICommandInfo>> NodeTypeToToolCommandMap;
@@ -166,6 +163,7 @@ private:
 
 	// Whether we should restore the previous view mode when a tool ends
 	bool bShouldRestoreSavedConstructionViewMode = false;
+	
 	// The Construction view mode that was active before starting the current tool. When the tool ends, restore this view mode if bShouldRestoreSavedConstructionViewMode is true.
 	Dataflow::EDataflowPatternVertexType SavedConstructionViewMode;
 
diff --git a/Engine/Plugins/Experimental/Dataflow/Source/DataflowEditor/Public/Dataflow/DataflowEditorToolkit.h b/Engine/Plugins/Experimental/Dataflow/Source/DataflowEditor/Public/Dataflow/DataflowEditorToolkit.h
--- a/Engine/Plugins/Experimental/Dataflow/Source/DataflowEditor/Public/Dataflow/DataflowEditorToolkit.h
+++ b/Engine/Plugins/Experimental/Dataflow/Source/DataflowEditor/Public/Dataflow/DataflowEditorToolkit.h
@@ -8,11 +8,9 @@
 #include "Misc/NotifyHook.h"
 #include "GraphEditor.h"
 #include "TickableEditorObject.h"
-#include "Dataflow/CollectionSpreadSheetWidget.h"
 #include "Dataflow/DataflowSelectionView.h"
 #include "Dataflow/DataflowCollectionSpreadSheet.h"
 #include "Dataflow/DataflowEditorViewport.h"
-#include "Dataflow/SelectionViewWidget.h"
 
 class FEditorViewportTabContent;
 class IDetailsView;
@@ -22,8 +20,9 @@ class IToolkitHost;
 class UDataflow;
 class USkeletalMesh;
 class SDataflowGraphEditor;
-class UDataflowEditorContent;
-class FDataflowPreviewScene;
+class FDataflowConstructionScene;
+class UDataflowBaseContent;
+class FDataflowSimulationScene;
 
 namespace Dataflow
 {
@@ -52,10 +51,9 @@ public:
 	static UDataflow* GetDataflowAsset(UObject* ObjectToEdit);
 	static const UDataflow* GetDataflowAsset(const UObject* ObjectToEdit);
 
-	/** Datafloe Editor Content Access */
-	TObjectPtr<const UDataflowEditorContent> GetDataflowEditorContent() const;
-	TObjectPtr<UDataflowEditorContent> GetDataflowEditorContent();
-
+	/** Dataflow Content Access */
+	const TObjectPtr<UDataflowBaseContent> GetDataflowContent() const;
+	TObjectPtr<UDataflowBaseContent> GetDataflowContent();
 
 	// IToolkit interface
 	virtual FName GetToolkitFName() const override;
@@ -67,6 +65,9 @@ public:
 	virtual void RegisterTabSpawners(const TSharedRef<FTabManager>& TabManager) override;
 	virtual void UnregisterTabSpawners(const TSharedRef<class FTabManager>& TabManager) override;
 
+	/** Dataflow preview scenes accessor */
+	const TSharedPtr<FDataflowSimulationScene>& GetSimulationScene() const {return SimulationScene;}
+
 protected:
 
 	// List of dataflow actions callbacks
@@ -90,6 +91,7 @@ private:
 	TSharedRef<SDockTab> SpawnTab_SelectionView(const FSpawnTabArgs& Args);
 	TSharedRef<SDockTab> SpawnTab_CollectionSpreadSheet(const FSpawnTabArgs& Args);
 	TSharedRef<SDockTab> SpawnTab_AssetDetails(const FSpawnTabArgs& Args);
+	TSharedRef<SDockTab> SpawnTab_SimulationViewport(const FSpawnTabArgs& Args);
 	
 	// FTickableEditorObject interface
 	virtual void Tick(float DeltaTime) override;
@@ -124,9 +126,11 @@ private:
 	static const FName CollectionSpreadSheetTabId_2;
 	static const FName CollectionSpreadSheetTabId_3;
 	static const FName CollectionSpreadSheetTabId_4;
+	static const FName SimulationViewportTabId;
 
 	// List of all the widgets shared ptr that will be built in the editor
 	TSharedPtr<SDataflowEditorViewport> DataflowEditorViewport;
+	TSharedPtr<SDataflowEditorViewport> DataflowSimulationViewport;
 	TSharedPtr<SDataflowGraphEditor> GraphEditor;
 	TSharedPtr<IStructureDetailsView> NodeDetailsEditor;
 	TSharedPtr<class ISkeletonTree> SkeletalEditor;
@@ -144,7 +148,10 @@ private:
 	TSharedRef<SDataflowGraphEditor> CreateGraphEditorWidget(UDataflow* ObjectToEdit, TSharedPtr<IStructureDetailsView> PropertiesEditor);
     TSharedPtr<IDetailsView> CreateAssetDetailsEditorWidget(UObject* ObjectToEdit);
     TSharedPtr<IStructureDetailsView> CreateNodeDetailsEditorWidget(UObject* ObjectToEdit);
-    TSharedPtr<ISkeletonTree> CreateSkeletalEditorWidget(UObject* ObjectToEdit);
+	TSharedPtr<ISkeletonTree> CreateSkeletalEditorWidget();
+
+	/** Create the simulation viewport client */
+	void CreateSimulationViewportClient();
 
 	// List of editor commands used  for the dataflow asset
 	TSharedPtr<FUICommandList> GraphEditorCommands;
@@ -165,7 +172,19 @@ private:
 	// The most recently selected dataflow node.
 	UPROPERTY()
 	TObjectPtr<UDataflowEdNode> PrimarySelection;
+	
+	/** PreviewScene showing the objects being simulated */
+	TSharedPtr<FDataflowSimulationScene> SimulationScene;
+
+	/** The editor mode manager used by the simulation preview scene */
+	TSharedPtr<FEditorModeTools> SimulationModeManager;
+
+	/** Simulation tab content */
+	TSharedPtr<class FEditorViewportTabContent> SimulationTabContent;
+
+	/** Simulation viewport delegate */
+	AssetEditorViewportFactoryFunction SimulationViewportDelegate;
 
-	/** Scene in which the 3D sim space preview meshes live. Ownership shared with AdvancedPreviewSettingsWidget*/
-	TSharedPtr<FDataflowPreviewScene> DataflowPreviewScene;
+	/** Simulation Viewport client */
+	TSharedPtr<FDataflowEditorViewportClient> SimulationViewportClient;
 };
diff --git a/Engine/Plugins/Experimental/Dataflow/Source/DataflowEditor/Public/Dataflow/DataflowEditorViewport.h b/Engine/Plugins/Experimental/Dataflow/Source/DataflowEditor/Public/Dataflow/DataflowEditorViewport.h
--- a/Engine/Plugins/Experimental/Dataflow/Source/DataflowEditor/Public/Dataflow/DataflowEditorViewport.h
+++ b/Engine/Plugins/Experimental/Dataflow/Source/DataflowEditor/Public/Dataflow/DataflowEditorViewport.h
@@ -35,6 +35,8 @@ public:
 	virtual void OnFloatingButtonClicked() override;
 
 private:
+	float GetViewMinInput() const;
+	float GetViewMaxInput() const;
 
 	UDataflowEditorMode* GetEdMode() const;
 };
diff --git a/Engine/Plugins/Experimental/Dataflow/Source/DataflowEditor/Public/Dataflow/DataflowEditorViewportClient.h b/Engine/Plugins/Experimental/Dataflow/Source/DataflowEditor/Public/Dataflow/DataflowEditorViewportClient.h
--- a/Engine/Plugins/Experimental/Dataflow/Source/DataflowEditor/Public/Dataflow/DataflowEditorViewportClient.h
+++ b/Engine/Plugins/Experimental/Dataflow/Source/DataflowEditor/Public/Dataflow/DataflowEditorViewportClient.h
@@ -5,7 +5,7 @@
 #include "EditorViewportClient.h"
 #include "Dataflow/DataflowNodeParameters.h"
 #include "Dataflow/DataflowComponentSelectionState.h"
-#include "Dataflow/DataflowPatternVertexType.h"
+#include "Dataflow/DataflowContent.h"
 #include "InputBehaviorSet.h"
 
 class FDataflowEditorToolkit;
@@ -23,7 +23,7 @@ class DATAFLOWEDITOR_API FDataflowEditorViewportClient : public FEditorViewportC
 public:
 	using Super = FEditorViewportClient;
 
-	FDataflowEditorViewportClient(FEditorModeTools* InModeTools, FPreviewScene* InPreviewScene,
+	FDataflowEditorViewportClient(FEditorModeTools* InModeTools, FPreviewScene* InPreviewScene,  const bool bCouldTickScene,
 								  const TWeakPtr<SEditorViewport> InEditorViewportWidget = nullptr);
 
 	void SetConstructionViewMode(Dataflow::EDataflowPatternVertexType InViewMode);
@@ -35,6 +35,10 @@ public:
 	/** Set the data flow toolkit used to create the client*/
 	void SetDataflowEditorToolkit(TWeakPtr<FDataflowEditorToolkit> DataflowToolkit);
 	
+	/** Get the data flow toolkit  */
+	const TWeakPtr<FDataflowEditorToolkit>& GetDataflowEditorToolkit() const { return DataflowEditorToolkitPtr; }
+
+	/** Set the tool command list */
 	void SetToolCommandList(TWeakPtr<FUICommandList> ToolCommandList);
 
 	// FGCObject Interface
@@ -43,7 +47,9 @@ public:
 
 private:
 
+	// FEditorViewportClient interface
 	virtual void Tick(float DeltaSeconds) override;
+	virtual void ProcessClick(FSceneView& View, HHitProxy* HitProxy, FKey Key, EInputEvent Event, uint32 HitX, uint32 HitY) override;
 
 	/** Toolkit used to create the viewport client */
 	TWeakPtr<FDataflowEditorToolkit> DataflowEditorToolkitPtr = nullptr;
@@ -54,8 +60,12 @@ private:
 	// @todo(brice) : Is this needed?
 	TWeakPtr<FUICommandList> ToolCommandList;
 
+	/** Construction view mode */
 	Dataflow::EDataflowPatternVertexType ConstructionViewMode = Dataflow::EDataflowPatternVertexType::Sim3D;
 
+	/** Behavior set for the behavior UI */
 	TObjectPtr<UInputBehaviorSet> BehaviorSet;
-
+	
+	/** Flag to enable scene ticking from the client */
+	bool bEnableSceneTicking = false;
 };
diff --git a/Engine/Plugins/Experimental/Dataflow/Source/DataflowEditor/Public/Dataflow/DataflowEditorViewportToolbar.h b/Engine/Plugins/Experimental/Dataflow/Source/DataflowEditor/Public/Dataflow/DataflowEditorViewportToolbar.h
--- a/Engine/Plugins/Experimental/Dataflow/Source/DataflowEditor/Public/Dataflow/DataflowEditorViewportToolbar.h
+++ b/Engine/Plugins/Experimental/Dataflow/Source/DataflowEditor/Public/Dataflow/DataflowEditorViewportToolbar.h
@@ -2,12 +2,12 @@
 #pragma once
 
 #include "CoreMinimal.h"
-#include "SEditorViewportToolBarMenu.h"
 #include "SCommonEditorViewportToolbarBase.h"
 
 // ----------------------------------------------------------------------------------
 class SDataflowEditorViewport;
 
+/** Base toolbar for the dataflow. Should be extended to add more features */
 class SDataflowViewportSelectionToolBar : public SCommonEditorViewportToolbarBase
 {
 public:
@@ -16,10 +16,7 @@ public:
 
 	/** Constructs this widget with the given parameters */
 	void Construct(const FArguments& InArgs, TSharedPtr<SDataflowEditorViewport> InDataflowViewport);
-
-	TSharedRef<SWidget> MakeSelectionModeToolBar();
-
-
+	
 private:
 	/** Reference to the parent viewport */
 	TWeakPtr<SDataflowEditorViewport> EditorViewport;
diff --git a/Engine/Plugins/Experimental/Dataflow/Source/DataflowEditor/Public/Dataflow/DataflowPreviewScene.h b/Engine/Plugins/Experimental/Dataflow/Source/DataflowEditor/Public/Dataflow/DataflowPreviewScene.h
--- a/Engine/Plugins/Experimental/Dataflow/Source/DataflowEditor/Public/Dataflow/DataflowPreviewScene.h
+++ b/Engine/Plugins/Experimental/Dataflow/Source/DataflowEditor/Public/Dataflow/DataflowPreviewScene.h
@@ -4,7 +4,7 @@
 
 #include "AdvancedPreviewScene.h"
 
-class UDataflowEditorContent;
+class UDataflowBaseContent;
 class USkeletalMeshComponent;
 class UDynamicMeshComponent;
 class UDataflowComponent;
@@ -33,148 +33,129 @@ class DATAFLOWEDITOR_API FDataflowPreviewScene : public FAdvancedPreviewScene
 {
 public:
 
-	FDataflowPreviewScene(FPreviewScene::ConstructionValues ConstructionValues, TObjectPtr<UDataflowEditorContent> InEditorContent);
+	FDataflowPreviewScene(FPreviewScene::ConstructionValues ConstructionValues, TObjectPtr<UDataflowBaseContent> InDataflowContent);
 	virtual ~FDataflowPreviewScene();
+	
+	// FGCObject interface
+	virtual void AddReferencedObjects(FReferenceCollector& Collector) override;
 
+	/** Dataflow content accessors */
+	TObjectPtr<UDataflowBaseContent> GetDataflowContent() { return DataflowContent; }
+	const TObjectPtr<UDataflowBaseContent> GetDataflowContent() const { return DataflowContent; }
 
-	//
-	// Functional API
-	//
-public:
-	/** Update the state of the scene.*/
-	void Update();
-
-	/** Tick */
-	// @todo(brice) : depricate
-	virtual void Tick(float DeltaTime) override;
-
-	/** Cleam up the scene on exit. */
-	void Exit();
-
-
-	//
-	// SkeletalMesh Rendering
-	//
-public:
-	/** Update Skeletal Mesh Component */
-	void UpdateSkeletalMeshComponent();
-
-	/** Modify/Get the skeletal mesh component */
-	USkeletalMeshComponent* GetSkeletalMeshComponent() { return SkeletalMeshComponent; }
-	const USkeletalMeshComponent* GetSkeletalMeshComponent() const { return SkeletalMeshComponent; }
-
-
-
+	/** Root scene actor accessors */
+	TObjectPtr<AActor> GetRootActor() { return RootSceneActor; }
+	const TObjectPtr<AActor> GetRootActor() const { return RootSceneActor; }
 
-	//
-	// ManagedArrayCollectionRendering
-	//
-public:
-	void UpdateDynamicMeshComponents();
+	/** Dataflow mode manager accessors */
+	TSharedPtr<FAssetEditorModeManager>& GetDataflowModeManager() { return DataflowModeManager; }
+	const TSharedPtr<FAssetEditorModeManager>& GetDataflowModeManager() const { return DataflowModeManager; }
+	
+	/** Build the scene bounding box */
+	FBox GetBoundingBox() const;
 
-	/** Modify/Get the skeletal mesh component */
-	TArray<TObjectPtr<UDynamicMeshComponent>> GetDynamicMeshComponents() { return DynamicMeshComponents; }
-	const TArray<TObjectPtr<UDynamicMeshComponent>> GetDynamicMeshComponents() const { return DynamicMeshComponents; }
+	/** Tick data flow scene */
+	virtual void TickDataflowScene(const float DeltaSeconds) {}
 
-	/** Trigger a Update when the render geometry has changed.*/
-	//@todo(brice) : deprecate
-	void ReinitializeDynamicMeshComponents();
+	/** Check if a primitive component is selected */
+	bool IsComponentSelected(const UPrimitiveComponent* InComponent) const;
 
-private:
-	/** Reset all the dynamic mesh components */
-	void ResetDynamicMeshComponents();
+	/** Check if the preview scene can run simulation */
+	virtual bool CanRunSimulation() const {return false;}
 	
-	/** Add a dynamic mesh component to the scene */
-	TObjectPtr<UDynamicMeshComponent>& AddDynamicMeshComponent(UE::Geometry::FDynamicMesh3&& DynamicMesh, const TArray<UMaterialInterface*>& MaterialSet);
-
-
+protected:
+	
+	/** Root scene actor */
+	TObjectPtr<AActor> RootSceneActor = nullptr;
+	
+	/** Data flows assets datas used to build the scene components */
+	TObjectPtr<UDataflowBaseContent> DataflowContent = nullptr;
 
+	/** Mode Manager for selection */
+	TSharedPtr<FAssetEditorModeManager> DataflowModeManager;
+};
 
-	//
-	// Wireframe
-	//
+/**
+ * Dataflow construction scene holding all the dynamic meshes used/displayed in the dataflow graph
+ */
+class DATAFLOWEDITOR_API FDataflowConstructionScene : public FDataflowPreviewScene
+{
 public:
-	void UpdateWireframeMeshElementsVisualizer();
-
-private:
-	void ResetWireframeMeshElementsVisualizer();
-
-	void AddWireframeMeshElementsVisualizer();
-
-
 
+	FDataflowConstructionScene(FPreviewScene::ConstructionValues ConstructionValues, TObjectPtr<UDataflowBaseContent> InEditorContent);
+	virtual ~FDataflowConstructionScene();
 
+	// FGCObject interface
+	virtual void AddReferencedObjects(FReferenceCollector& Collector) override;
+	
+	/** Tick data flow scene */
+	virtual void TickDataflowScene(const float DeltaSeconds) override;
 
-	//
-	// Visibility
-	//
-public:
+	/** Check if the preview scene can run simulation */
+	virtual bool CanRunSimulation() const override {return false;}
+	
+	/**  Get all the dynamic mesh components */
+	TArray<TObjectPtr<UDynamicMeshComponent>> GetDynamicMeshComponents() { return DynamicMeshComponents; }
+	
 	/** Is there ANYTHING to render */
 	bool HasRenderableGeometry();
 
-	/** Build the scene bounding box */
-	FBox GetBoundingBox() const;
-
-	/** Build the selected objects bounding box */
-	FBox SelectedComponentBounds() const;
-
-
-	//
-	// General Access
-	//
-public:
-
-	/** Modify/Get the data flow datas */
-	TObjectPtr<UDataflowEditorContent> GetDataflowEditorContent() { return EditorContent; }
-	const TObjectPtr<UDataflowEditorContent> GetDataflowEditorContent() const { return EditorContent; }
+	/** Close the construction scene */
+	void ResetConstructionScene();
 
-	/** Modify/Get the data flow mode manager */
-	TSharedPtr<FAssetEditorModeManager>& GetDataflowModeManager() { return DataflowModeManager; }
-	const TSharedPtr<FAssetEditorModeManager>& GetDataflowModeManager() const { return DataflowModeManager; }
+	/** Update the construction scene */
+	void UpdateConstructionScene();
+	
+private:
 
-	/** Modify/Get the preview anim instance */
-	UAnimSingleNodeInstance* GetPreviewAnimInstance() { return PreviewAnimInstance; }
-	const UAnimSingleNodeInstance* const GetPreviewAnimInstance() const { return PreviewAnimInstance; }
+	/** Reset the scene mesh visualizer */
+	void ResetWireframeMeshElementsVisualizer();
 
-	/** FGCObject interface */
-	virtual void AddReferencedObjects(FReferenceCollector& Collector) override;
+	/** Update the scene mesh visualizer */
+	void UpdateWireframeMeshElementsVisualizer();
 
+	/** Add a mesh visualizer to the scene */
+	void AddWireframeMeshElementsVisualizer();
 
-private:
+	/** Reset all the dynamic mesh components */
+    void ResetDynamicMeshComponents();
 	
-	/** Check if a primitive component is selected */
-	bool IsComponentSelected(const UPrimitiveComponent* InComponent) const;
-
-	/** Data flows assets datas used to build the scene components */
-	TObjectPtr<UDataflowEditorContent> EditorContent = nullptr;
+	/** Build the dynamic mesh components from the dataflow graph render targets*/
+	void UpdateDynamicMeshComponents();
+	
+	/** Add a dynamic mesh component to the scene */
+	TObjectPtr<UDynamicMeshComponent>& AddDynamicMeshComponent(
+		UE::Geometry::FDynamicMesh3&& DynamicMesh, const TArray<UMaterialInterface*>& MaterialSet);
 	
-	/** Dynamic mesh actor */
-	TObjectPtr<AActor> DynamicMeshActor = nullptr;
-
 	/** Dynamic mesh components generated by the tools for edit/manipulation*/
 	TArray<TObjectPtr<UDynamicMeshComponent>> DynamicMeshComponents;
 
-	/** Skeletal mesh actor */
-	TObjectPtr<AActor> SkeletalMeshActor = nullptr;
-
-	/** Skeletal mesh component used for potential binding */
-	TObjectPtr<USkeletalMeshComponent> SkeletalMeshComponent = nullptr;
+	/** Rest space wireframe. They have to get ticked to be able to respond to setting changes. [ticked] */
+	TObjectPtr<UMeshElementsVisualizer> WireframeDraw = nullptr;
 
-	/** Preview anim instance for the skeletal mesh */
-	TObjectPtr<UAnimSingleNodeInstance> PreviewAnimInstance = nullptr;
+	/** Property objects (visible or not) that get ticked. [ticked] */
+	TArray<TObjectPtr<UInteractiveToolPropertySet>> PropertyObjectsToTick;
 
-	/** Mode Manager for selection */
-	TSharedPtr<FAssetEditorModeManager> DataflowModeManager;
+	/** Show the wireframe on render */ 
+	bool bConstructionViewWireframe = true;
+};
 
-	// Property objects (visible or not) that get ticked. [ticked]
-	TArray<TObjectPtr<UInteractiveToolPropertySet>> PropertyObjectsToTick;
 
-	// Rest space wireframe. They have to get ticked to be able to respond to setting changes. [ticked]
-	TObjectPtr<UMeshElementsVisualizer> WireframeDraw = nullptr;
+/**
+ * Dataflow simulation scene holding all the dataflow content components
+ */
+class DATAFLOWEDITOR_API FDataflowSimulationScene : public FDataflowPreviewScene
+{
+public:
 
-	// Show the wireframe on render. 
-	bool bConstructionViewWireframe = true;
+	FDataflowSimulationScene(FPreviewScene::ConstructionValues ConstructionValues, TObjectPtr<UDataflowBaseContent> InEditorContent);
+	virtual ~FDataflowSimulationScene();
 
+	/** Tick data flow scene */
+	virtual void TickDataflowScene(const float DeltaSeconds) override;
 	
+	/** Check if the preview scene can run simulation */
+	virtual bool CanRunSimulation() const { return true; }
 };
 
+
diff --git a/Engine/Plugins/Experimental/Dataflow/Source/DataflowEditor/Public/Dataflow/DataflowSimulationPanel.h b/Engine/Plugins/Experimental/Dataflow/Source/DataflowEditor/Public/Dataflow/DataflowSimulationPanel.h
new file mode 100644
--- /dev/null
+++ b/Engine/Plugins/Experimental/Dataflow/Source/DataflowEditor/Public/Dataflow/DataflowSimulationPanel.h
@@ -0,0 +1,88 @@
+﻿// Copyright Epic Games, Inc. All Rights Reserved.
+
+#pragma once
+
+#include "DataflowPreviewScene.h"
+#include "Widgets/SCompoundWidget.h"
+#include "ITransportControl.h"  // EPlaybackMode::Type
+
+class SScrubControlPanel;
+class SButton;
+class UAnimSingleNodeInstance;
+
+/** Dataflow simulation panel to control an animation/simulation */
+class SDataflowSimulationPanel : public SCompoundWidget
+{
+private:
+
+	/** Data flow playback mode */
+	enum class EDataflowPlaybackMode : int32
+	{
+		Default,
+		Looping,
+		PingPong
+	};
+
+	/** View min/max to run the simulation */
+	SLATE_BEGIN_ARGS(SDataflowSimulationPanel)	{}
+		SLATE_ATTRIBUTE( float, ViewInputMin )
+		SLATE_ATTRIBUTE( float, ViewInputMax )
+	SLATE_END_ARGS()
+
+	/** Construct the simulation widget */
+	void Construct( const FArguments& InArgs, TWeakPtr<FDataflowSimulationScene> InPreviewScene );
+
+	/** Delegate when the playback mode is changed */
+	TSharedRef<SWidget> OnCreatePreviewPlaybackModeWidget();
+
+	/** Delegate when the simulation controls are pressed*/
+	FReply OnClick_Forward_Step();
+	FReply OnClick_Forward_End();
+	FReply OnClick_Backward_Step();
+	FReply OnClick_Backward_End();
+	FReply OnClick_Forward();
+	FReply OnClick_Backward();
+	FReply OnClick_PreviewPlaybackMode();
+
+	/** Change the playback settings */
+	void ApplyPlaybackSettings();
+
+	/** Main simulation delegates */
+	void OnTickPlayback(double InCurrentTime, float InDeltaTime);
+	void OnValueChanged(float NewValue);
+	void OnBeginSliderMovement();
+
+	/** Get the playback mode used in the widget */
+	EPlaybackMode::Type GetPlaybackMode() const;
+
+	/** Get ethe current scrub value */
+	float GetScrubValue() const;
+
+	/** Animation instance accessors */
+	UAnimSingleNodeInstance* GetPreviewAnimationInstance();
+	const UAnimSingleNodeInstance* GetPreviewAnimationInstance() const;
+
+	/** Update the widget based on the modified animation instance*/
+	void UpdatePreviewAnimationInstance();
+
+	/** Get the number of keys */
+	uint32 GetNumberOfKeys() const;
+
+	/** Get the sequence length */
+	float GetSequenceLength() const;
+
+	/** Get the display drag */
+	bool GetDisplayDrag() const;
+
+	/** Simulation scene to be used for the widget */
+	TWeakPtr<FDataflowSimulationScene> SimulationScene;
+
+	/** Scrub widget defined for the timeline */
+	TSharedPtr<SScrubControlPanel> ScrubControlPanel;
+
+	/** Playback mode button */
+	TSharedPtr<SButton> PreviewPlaybackModeButton;
+
+	/** Preview playback mode (looping...)*/
+	EDataflowPlaybackMode PreviewPlaybackMode = EDataflowPlaybackMode::Looping;
+};
diff --git a/Engine/Plugins/Experimental/Dataflow/Source/DataflowEditor/Public/DataflowEditorTools/DataflowEditorWeightMapPaintTool.h b/Engine/Plugins/Experimental/Dataflow/Source/DataflowEditor/Public/DataflowEditorTools/DataflowEditorWeightMapPaintTool.h
--- a/Engine/Plugins/Experimental/Dataflow/Source/DataflowEditor/Public/DataflowEditorTools/DataflowEditorWeightMapPaintTool.h
+++ b/Engine/Plugins/Experimental/Dataflow/Source/DataflowEditor/Public/DataflowEditorTools/DataflowEditorWeightMapPaintTool.h
@@ -31,7 +31,7 @@ class UDataflowWeightMapEraseBrushOpProps;
 class UDataflowWeightMapPaintBrushOpProps;
 class UDataflowWeightMapSmoothBrushOpProps;
 class UPolygonSelectionMechanic;
-class UDataflowEditorContextObject;
+class UDataflowContextObject;
 struct FDataflowCollectionAddScalarVertexPropertyNode;
 
 DECLARE_STATS_GROUP(TEXT("WeightMapPaintTool"), STATGROUP_WeightMapPaintTool, STATCAT_Advanced);
@@ -282,7 +282,7 @@ public:
 
 	virtual void CommitResult(UBaseDynamicMeshComponent* Component, bool bModifiedTopology) override;
 
-	void SetDataflowEditorContextObject(TObjectPtr<UDataflowEditorContextObject> InDataflowEditorContextObject);
+	void SetDataflowEditorContextObject(TObjectPtr<UDataflowContextObject> InDataflowEditorContextObject);
 
 public:
 
@@ -402,7 +402,7 @@ protected:
 	TObjectPtr<UMeshElementsVisualizer> MeshElementsDisplay;
 
 	UPROPERTY()
-	TObjectPtr<UDataflowEditorContextObject> DataflowEditorContextObject = nullptr;
+	TObjectPtr<UDataflowContextObject> DataflowEditorContextObject = nullptr;
 
 	// realtime visualization
 	void OnDynamicMeshComponentChanged(UDynamicMeshComponent* Component, const FMeshVertexChange* Change, bool bRevert);
diff --git a/Engine/Source/Runtime/Experimental/Dataflow/Engine/Private/Dataflow/DataflowContent.cpp b/Engine/Source/Runtime/Experimental/Dataflow/Engine/Private/Dataflow/DataflowContent.cpp
new file mode 100644
--- /dev/null
+++ b/Engine/Source/Runtime/Experimental/Dataflow/Engine/Private/Dataflow/DataflowContent.cpp
@@ -0,0 +1,161 @@
+﻿// Copyright Epic Games, Inc. All Rights Reserved.
+
+#include "Dataflow/DataflowContent.h"
+#include "Animation/AnimSingleNodeInstance.h"
+#include "Components/SkeletalMeshComponent.h"
+#include "Animation/Skeleton.h"
+#include "Engine/SkeletalMesh.h"
+#include "GameFramework/Actor.h"
+#include "PreviewScene.h"
+
+#include UE_INLINE_GENERATED_CPP_BY_NAME(DataflowContent)
+
+UDataflowBaseContent::UDataflowBaseContent()
+{
+}
+
+void UDataflowBaseContent::BuildBaseContent(TObjectPtr<UObject> ContentOwner)
+{
+	DataflowContext = MakeShared<Dataflow::FEngineContext>(ContentOwner, DataflowAsset, FPlatformTime::Cycles64());
+	LastModifiedTimestamp = DataflowContext->GetTimestamp();
+}
+
+UDataflowSkeletalContent::UDataflowSkeletalContent() : Super()
+{
+}
+
+void UDataflowSkeletalContent::RegisterWorldContent(FPreviewScene* PreviewScene, AActor* RootActor)
+{
+	Super::RegisterWorldContent(PreviewScene, RootActor);
+
+	SkeletalMeshComponent = NewObject<USkeletalMeshComponent>(RootActor);
+	SkeletalMeshComponent->SetDisablePostProcessBlueprint(true);
+	
+	if(SkeletalMesh)
+	{
+		SkeletalMeshComponent->SetSkeletalMeshAsset(SkeletalMesh);
+
+		Skeleton = SkeletalMesh->GetSkeleton();
+		UpdateAnimationInstance();
+	}
+	SkeletalMeshComponent->UpdateBounds();
+	PreviewScene->AddComponent(SkeletalMeshComponent, SkeletalMeshComponent->GetRelativeTransform());
+}
+
+void UDataflowSkeletalContent::UpdateAnimationInstance()
+{
+	if (AnimationAsset && SkeletalMesh && (AnimationAsset->GetSkeleton() == SkeletalMesh->GetSkeleton()))
+	{
+		AnimationNodeInstance = NewObject<UAnimSingleNodeInstance>(SkeletalMeshComponent);
+		AnimationNodeInstance->SetAnimationAsset(AnimationAsset);
+
+		SkeletalMeshComponent->SetAnimationMode(EAnimationMode::AnimationSingleNode);
+		SkeletalMeshComponent->InitAnim(true);
+		SkeletalMeshComponent->AnimationData.PopulateFrom(AnimationNodeInstance);
+		SkeletalMeshComponent->AnimScriptInstance = AnimationNodeInstance;
+		SkeletalMeshComponent->AnimScriptInstance->InitializeAnimation();
+
+#if WITH_EDITOR
+		SkeletalMeshComponent->ValidateAnimation();
+#endif
+	}
+	else
+	{
+		SkeletalMeshComponent->Stop();
+		SkeletalMeshComponent->AnimationData = FSingleAnimationPlayData();
+		SkeletalMeshComponent->AnimScriptInstance = nullptr;
+		AnimationNodeInstance = nullptr;
+		AnimationAsset = nullptr;
+	}
+}
+	
+void UDataflowSkeletalContent::UnregisterWorldContent(FPreviewScene* PreviewScene)
+{
+	Super::UnregisterWorldContent(PreviewScene);
+	if (SkeletalMeshComponent)
+	{
+		SkeletalMeshComponent->TransformUpdated.RemoveAll(this);
+		PreviewScene->RemoveComponent(SkeletalMeshComponent);
+	}
+}
+
+void UDataflowSkeletalContent::SetSkeletalMesh(const TObjectPtr<USkeletalMesh>& SkeletalMeshAsset)
+{
+	SkeletalMesh = SkeletalMeshAsset;
+	if(SkeletalMesh && SkeletalMesh->GetSkeleton())
+	{
+		Skeleton = SkeletalMesh->GetSkeleton();
+	}
+	if(SkeletalMeshComponent)
+	{
+ 		SkeletalMeshComponent->SetSkeletalMeshAsset(SkeletalMesh);
+
+ 		UpdateAnimationInstance();
+	}
+	bIsDirty = true;
+}
+
+void UDataflowSkeletalContent::SetAnimationAsset(const TObjectPtr<UAnimationAsset>& SkeletalAnimationAsset)
+{
+	AnimationAsset = SkeletalAnimationAsset;
+	if(SkeletalMeshComponent)
+	{
+ 		UpdateAnimationInstance();
+	}
+	bIsDirty = true;
+}
+
+void UDataflowSkeletalContent::SetSkeleton(const TObjectPtr<USkeleton>& SkeletonAsset)
+{
+	if(SkeletonAsset)
+	{
+		Skeleton = SkeletonAsset;
+		if(SkeletalMesh && (SkeletalMesh->GetSkeleton() != Skeleton))
+		{
+			SetSkeletalMesh(nullptr);
+		}
+	}
+	bIsDirty = true;
+}
+
+#if WITH_EDITOR
+
+void UDataflowSkeletalContent::PostEditChangeProperty(FPropertyChangedEvent& PropertyChangedEvent) 
+{
+	Super::PostEditChangeProperty(PropertyChangedEvent);
+	
+	const FName PropertyName = PropertyChangedEvent.Property->GetFName();
+	
+	if (PropertyName == GET_MEMBER_NAME_CHECKED(UDataflowSkeletalContent, SkeletalMesh))
+	{
+		SetSkeletalMesh(SkeletalMesh);
+	}
+	if (PropertyName == GET_MEMBER_NAME_CHECKED(UDataflowSkeletalContent, AnimationAsset))
+	{
+		SetAnimationAsset(AnimationAsset);
+	}
+	if (PropertyName == GET_MEMBER_NAME_CHECKED(UDataflowSkeletalContent, Skeleton))
+	{
+		SetSkeleton(Skeleton);
+	}
+}
+#endif //if WITH_EDITOR
+
+void UDataflowSkeletalContent::AddContentObjects(FReferenceCollector& Collector)
+{
+	Super::AddContentObjects(Collector);
+	
+	Collector.AddReferencedObject(SkeletalMeshComponent);
+	Collector.AddReferencedObject(AnimationNodeInstance);
+}
+
+FVector2f UDataflowSkeletalContent::GetSimulationRange() const
+{
+	if(AnimationNodeInstance)
+	{
+		return FVector2f(0.0f, AnimationNodeInstance->GetLength());
+	}
+	return FVector2f(0.0f, 0.0f);
+}
+
+
diff --git a/Engine/Source/Runtime/Experimental/Dataflow/Engine/Private/Dataflow/DataflowObjectInterface.cpp b/Engine/Source/Runtime/Experimental/Dataflow/Engine/Private/Dataflow/DataflowObjectInterface.cpp
--- a/Engine/Source/Runtime/Experimental/Dataflow/Engine/Private/Dataflow/DataflowObjectInterface.cpp
+++ b/Engine/Source/Runtime/Experimental/Dataflow/Engine/Private/Dataflow/DataflowObjectInterface.cpp
@@ -1,3 +1,9 @@
 ﻿// Copyright Epic Games, Inc. All Rights Reserved.
 
 #include "Dataflow/DataflowObjectInterface.h"
+
+namespace Dataflow
+{
+	template class TEngineContext<FContextSingle>;
+	template class TEngineContext<FContextThreaded>;
+}
diff --git a/Engine/Source/Runtime/Experimental/Dataflow/Engine/Public/Dataflow/DataflowContent.h b/Engine/Source/Runtime/Experimental/Dataflow/Engine/Public/Dataflow/DataflowContent.h
new file mode 100644
--- /dev/null
+++ b/Engine/Source/Runtime/Experimental/Dataflow/Engine/Public/Dataflow/DataflowContent.h
@@ -0,0 +1,213 @@
+﻿// Copyright Epic Games, Inc. All Rights Reserved.
+
+#pragma once
+
+#include "Dataflow/DataflowObjectInterface.h"
+#include "Components/PrimitiveComponent.h"
+#include "Templates/SharedPointer.h"
+#include "Dataflow/DataflowEdNode.h"
+#include "DataflowContent.generated.h"
+
+class FDataflowEditorToolkit;
+class UDataflow;
+class USkeletalMesh;
+class USkeleton;
+class USkeletalMeshComponent;
+class UAnimationAsset;
+class UDataflowBaseContent;
+class FPreviewScene;
+class UAnimSingleNodeInstance;
+class AActor;
+
+namespace Dataflow
+{
+	enum class EDataflowPatternVertexType : uint8
+	{
+		Sim2D = 0,
+		Sim3D = 1,
+		Render = 2
+	};
+}
+
+/** 
+ * Context object used for selection/rendering 
+ */
+
+UCLASS()
+class DATAFLOWENGINE_API UDataflowContextObject : public UObject
+{
+	GENERATED_BODY()
+public:
+
+	/** Selection Collection Access */
+	void SetPrimarySelectedNode(TObjectPtr<UDataflowEdNode> InSelectedNode) { PrimarySelectedNode = InSelectedNode; }
+	TObjectPtr<UDataflowEdNode> GetPrimarySelectedNode() const { return PrimarySelectedNode; }
+
+	/** Render Collection used to generate the DynamicMesh3D on the PrimarySelection */
+	void SetPrimaryRenderCollection(const TSharedPtr<FManagedArrayCollection>& InCollection) { PrimaryRenderCollection = InCollection; }
+	TSharedPtr<const FManagedArrayCollection> GetPrimaryRenderCollection() const { return PrimaryRenderCollection; }
+
+	/** ViewMode Access */
+	void SetConstructionViewMode(Dataflow::EDataflowPatternVertexType InMode) {ConstructionViewMode = InMode;}
+	Dataflow::EDataflowPatternVertexType GetConstructionViewMode() const { return ConstructionViewMode; }
+
+	/** Get a single selected node of the specified type. Return nullptr if the specified node is not selected, or if multiple nodes are selected*/
+	template<typename NodeType>
+	NodeType* GetPrimarySelectedNodeOfType() const 
+	{
+		if (PrimarySelectedNode && PrimarySelectedNode->GetDataflowNode()) 
+		{
+			return PrimarySelectedNode->GetDataflowNode()->AsType<NodeType>();
+		}
+		return nullptr;
+	}
+
+protected:
+
+	/** Render collection to be used */
+	TSharedPtr<FManagedArrayCollection> PrimaryRenderCollection = nullptr;
+
+	/** Primary node that is selected in the graph */
+	TObjectPtr<UDataflowEdNode> PrimarySelectedNode = nullptr;
+
+	/** Construction view mode for the context object @todo(michael) : is it only for construction or for simulation as well*/
+	Dataflow::EDataflowPatternVertexType ConstructionViewMode = Dataflow::EDataflowPatternVertexType::Sim3D;
+};
+
+/** 
+ * Dataflow content owning dataflow asset that that will be used to evaluate the graph
+ */
+UCLASS()
+class DATAFLOWENGINE_API UDataflowBaseContent : public UDataflowContextObject
+{
+	GENERATED_BODY()
+
+public:
+	UDataflowBaseContent();
+
+	/** Check if the dataflow data are dirty and need to be re-rendered */
+	bool IsDirty() const { return bIsDirty; }
+	void SetIsDirty(bool InDirty) { bIsDirty = InDirty; }
+
+	/** Data flow asset that we will edit */
+	UPROPERTY(EditAnywhere, Category = "Dataflow")
+	TObjectPtr<UDataflow> DataflowAsset = nullptr;
+
+	/** Data flow terminal path for evaluation */
+	UPROPERTY(EditAnywhere, Category = "Dataflow")
+	FString DataflowTerminal = "";
+	
+	/** Return the simulation time range to be used in the simulation viewport */
+	virtual FVector2f GetSimulationRange() const { return FVector2f(0.0f, 100.0f); } 
+	
+	/** Register components to the simulation world */
+	virtual void RegisterWorldContent(FPreviewScene* PreviewScene, AActor* RootActor) {}
+	
+	/** Unregister components to the simulation world */
+	virtual void UnregisterWorldContent(FPreviewScene* PreviewScene) {}
+
+	/** Build the content context, timestamp*/
+	void BuildBaseContent(TObjectPtr<UObject> ContentOwner);
+
+	/** Collect reference objects for GC */
+	virtual void AddContentObjects(FReferenceCollector& Collector) {}
+	
+	/** Data flow owner accessors */
+	void SetDataflowOwner(const TObjectPtr<UObject>& InOwner) { if(DataflowContext) { DataflowContext->Owner = InOwner;  bIsDirty = true; }}
+	TObjectPtr<UObject> GetDataflowOwner() const { return DataflowContext ? DataflowContext->Owner : nullptr; }
+
+	/** Data flow asset accessors */
+	void SetDataflowAsset(const TObjectPtr<UDataflow>& InAsset) { DataflowAsset = InAsset;  bIsDirty = true;}
+	const TObjectPtr<UDataflow>& GetDataflowAsset() const { return DataflowAsset; }
+
+	/** Data flow terminal accessors */
+	void SetDataflowTerminal(const FString& InPath) { DataflowTerminal = InPath;  bIsDirty = true;}
+	const FString& GetDataflowTerminal() const { return DataflowTerminal; }
+
+	/**  Data flow context accessors */
+	void SetDataflowContext(const TSharedPtr<Dataflow::FEngineContext>& InContext) { DataflowContext = InContext; bIsDirty = true;}
+	const TSharedPtr<Dataflow::FEngineContext>& GetDataflowContext() const { return DataflowContext; }
+
+	/** Data flow timestamp accessors */
+	void SetLastModifiedTimestamp(Dataflow::FTimestamp InTimestamp) { LastModifiedTimestamp = InTimestamp; bIsDirty = true;}
+	const Dataflow::FTimestamp& GetLastModifiedTimestamp() const { return LastModifiedTimestamp; }
+	
+protected :
+	
+	/**  Engine context to be used for dataflow evaluation */
+    TSharedPtr<Dataflow::FEngineContext> DataflowContext = nullptr;
+
+    /** Last data flow evaluated node time stamp */
+    Dataflow::FTimestamp LastModifiedTimestamp = Dataflow::FTimestamp::Invalid;
+
+    /** Dirty flag to trigger rendering. Do we need that? since when accessing the member by non const ref we will not dirty it */
+    bool bIsDirty = true;
+};
+
+/** 
+ * Dataflow content owning dataflow and skelmesh assets that that will be used to evaluate the graph
+ */
+UCLASS()
+class DATAFLOWENGINE_API UDataflowSkeletalContent : public UDataflowBaseContent
+{
+	GENERATED_BODY()
+
+public:
+	UDataflowSkeletalContent();
+	virtual ~UDataflowSkeletalContent() override{}
+
+	/** Data flow skeletal mesh*/
+	UPROPERTY(EditAnywhere, Category = "Preview")
+	TObjectPtr<USkeletalMesh> SkeletalMesh = nullptr;
+
+	/** Animation asset to be used to preview simulation */
+	UPROPERTY(EditAnywhere, Category = "Preview")
+	TObjectPtr<UAnimationAsset> AnimationAsset;
+
+	/** Data flow skeleton */
+	UPROPERTY(EditAnywhere, Category = "Skeleton")
+	TObjectPtr<USkeleton> Skeleton = nullptr;
+
+	/** Return the simulation time range to be used in the simulation viewport */
+	virtual FVector2f GetSimulationRange() const override; 
+
+	/** Register components to the scene world */
+	virtual void RegisterWorldContent(FPreviewScene* PreviewScene, AActor* RootActor) override;
+	
+	/** Unregister components to the scene world */
+	virtual void UnregisterWorldContent(FPreviewScene* PreviewScene) override;
+	
+#if WITH_EDITOR
+	virtual void PostEditChangeProperty(FPropertyChangedEvent& PropertyChangedEvent) override;
+#endif //if WITH_EDITOR
+
+	/** Collect reference objects for GC */
+	virtual void AddContentObjects(FReferenceCollector& Collector) override;
+
+	/** Update the animation node instance and set it to the skelmesh component*/
+	void UpdateAnimationInstance();
+
+	/** Data flow animation instance accessors */
+	const TObjectPtr<UAnimSingleNodeInstance>& GetAnimationInstance() const {return AnimationNodeInstance;}
+	TObjectPtr<UAnimSingleNodeInstance>& GetAnimationInstance() {return AnimationNodeInstance;}
+
+	/** Data flow skeletal mesh accessors */
+	void SetSkeletalMesh(const TObjectPtr<USkeletalMesh>& InMesh);
+	const TObjectPtr<USkeletalMesh>& GetSkeletalMesh() const { return SkeletalMesh; }
+
+	/** Data flow skeleton accessors */
+	void SetSkeleton(const TObjectPtr<USkeleton>& InSkeleton);
+	const TObjectPtr<USkeleton>& GetSkeleton() const { return Skeleton; }
+
+	/** Data flow animation asset accessors */
+	void SetAnimationAsset(const TObjectPtr<UAnimationAsset>& InAnimation);
+	const TObjectPtr<UAnimationAsset>& GetAnimationAsset() const { return AnimationAsset; }
+	
+protected :
+	
+	/** Skeletal mesh component used in the preview scene */
+	TObjectPtr<USkeletalMeshComponent> SkeletalMeshComponent;
+
+	/** Anim node instance used with skelmesh component */
+	TObjectPtr<UAnimSingleNodeInstance> AnimationNodeInstance;
+};
diff --git a/Engine/Source/Runtime/Experimental/Dataflow/Engine/Public/Dataflow/DataflowObjectInterface.h b/Engine/Source/Runtime/Experimental/Dataflow/Engine/Public/Dataflow/DataflowObjectInterface.h
--- a/Engine/Source/Runtime/Experimental/Dataflow/Engine/Public/Dataflow/DataflowObjectInterface.h
+++ b/Engine/Source/Runtime/Experimental/Dataflow/Engine/Public/Dataflow/DataflowObjectInterface.h
@@ -18,18 +18,17 @@ namespace Dataflow
 	public:
 		DATAFLOW_CONTEXT_INTERNAL(Base, TEngineContext);
 
-		TEngineContext(UObject* InOwner, UDataflow* InGraph, FTimestamp InTimestamp)
+		TEngineContext(const TObjectPtr<UObject>& InOwner,const TObjectPtr<UDataflow>& InGraph, FTimestamp InTimestamp)
 				: Base(InTimestamp)
 				, Owner(InOwner)
 				, Graph(InGraph)
 		{}
 	
-		UObject* Owner = nullptr;
-		UDataflow* Graph = nullptr;
-	};
+		TObjectPtr<UObject> Owner = nullptr;
+		TObjectPtr<UDataflow> Graph = nullptr;
 
-	template class TEngineContext<FContextSingle>;
-	template class TEngineContext<FContextThreaded>;
+		~TEngineContext(){}
+	};
 
 	typedef TEngineContext<FContextSingle> FEngineContext;
 	typedef TEngineContext<FContextThreaded> FEngineContextThreaded;
