Dataflow : Default to vertex material if one is not set.

diff --git a/Engine/Plugins/Experimental/Dataflow/Source/DataflowEditor/Private/Dataflow/DataflowPreviewScene.cpp b/Engine/Plugins/Experimental/Dataflow/Source/DataflowEditor/Private/Dataflow/DataflowPreviewScene.cpp
--- a/Engine/Plugins/Experimental/Dataflow/Source/DataflowEditor/Private/Dataflow/DataflowPreviewScene.cpp
+++ b/Engine/Plugins/Experimental/Dataflow/Source/DataflowEditor/Private/Dataflow/DataflowPreviewScene.cpp
@@ -209,10 +209,15 @@ TObjectPtr<UDynamicMeshComponent>& FDataflowConstructionScene::AddDynamicMeshCom
 	DynamicMeshComponent->SetMesh(MoveTemp(DynamicMesh));
 	
 	// @todo(Material) This is just to have a material, we should transfer the materials from the assets if they have them. 
-	if (DataflowContent && DataflowContent->DataflowAsset)
+	if (DataflowContent && DataflowContent->DataflowAsset && DataflowContent->DataflowAsset->Material)
 	{
 		DynamicMeshComponent->ConfigureMaterialSet({ DataflowContent->DataflowAsset->Material });
 	}
+	else
+	{
+		DynamicMeshComponent->SetOverrideRenderMaterial(FDataflowEditorStyle::Get().VertexMaterial);
+		DynamicMeshComponent->SetShadowsEnabled(false);
+	}
 	//else if (FDataflowEditorStyle::Get().DefaultMaterial)
 	//{
 	//	DynamicMeshComponent->ConfigureMaterialSet({ FDataflowEditorStyle::Get().DefaultMaterial });
@@ -221,11 +226,7 @@ TObjectPtr<UDynamicMeshComponent>& FDataflowConstructionScene::AddDynamicMeshCom
 	//{
 	//	DynamicMeshComponent->ValidateMaterialSlots(true, false);
 	//}
-	else
-	{
-		DynamicMeshComponent->SetOverrideRenderMaterial(FDataflowEditorStyle::Get().VertexMaterial);
-		DynamicMeshComponent->SetShadowsEnabled(false);
-	}
+
 	DynamicMeshComponent->SelectionOverrideDelegate = UPrimitiveComponent::FSelectionOverride::CreateRaw(this, &FDataflowPreviewScene::IsComponentSelected);
 	DynamicMeshComponent->UpdateBounds();
 
