Dataflow : On Selection updates
- Removed Selected Collection on the Mode
- Halted continuous update in FDataflowConstructionScene::TickDataflowScene by only rebuilding when the scene is dirty or the graph has changed.
-  Prevent loading invalid state when the painting tool loads.

diff --git a/Engine/Plugins/Experimental/Dataflow/Source/DataflowEditor/Private/Dataflow/DataflowEditorMode.cpp b/Engine/Plugins/Experimental/Dataflow/Source/DataflowEditor/Private/Dataflow/DataflowEditorMode.cpp
--- a/Engine/Plugins/Experimental/Dataflow/Source/DataflowEditor/Private/Dataflow/DataflowEditorMode.cpp
+++ b/Engine/Plugins/Experimental/Dataflow/Source/DataflowEditor/Private/Dataflow/DataflowEditorMode.cpp
@@ -348,22 +348,6 @@ bool UDataflowEditorMode::IsComponentSelected(const UPrimitiveComponent* InCompo
 	return false;
 }
 
-void UDataflowEditorMode::SetSelectedCollection(TSharedPtr<FManagedArrayCollection> Collection)
-{
-	SelectedCollection = Collection;
-	PreviewScene->UpdateConstructionScene();
-
-	// The first time we get a valid mesh, refocus the camera on it
-	FirstTimeFocusRestSpaceViewport();
-}
-
-TSharedPtr<FManagedArrayCollection> UDataflowEditorMode::GetSelectedCollection()
-{
-	return SelectedCollection;
-
-	// TODO: If no cloth collection node is selected, show the ClothAsset's collection. In this case, also ensure that any interactive tools are disabled. (UE-181574)
-}
-
 void UDataflowEditorMode::RefocusRestSpaceViewportClient()
 {
 	TSharedPtr<FDataflowEditorViewportClient, ESPMode::ThreadSafe> PinnedVC = ConstructionViewportClient.Pin();
@@ -396,7 +380,7 @@ void UDataflowEditorMode::RefocusRestSpaceViewportClient()
 void UDataflowEditorMode::FirstTimeFocusRestSpaceViewport()
 {
 	// If this is the first time seeing a valid 2D or 3D mesh, refocus the camera on it.
-	const bool bIsValid = (SelectedCollection && PreviewScene->HasRenderableGeometry());
+	const bool bIsValid = (PreviewScene->HasRenderableGeometry());
 	const bool bIs2D = ConstructionViewMode == Dataflow::EDataflowPatternVertexType::Sim3D;
 
 	if (bIsValid)
diff --git a/Engine/Plugins/Experimental/Dataflow/Source/DataflowEditor/Private/Dataflow/DataflowEditorToolkit.cpp b/Engine/Plugins/Experimental/Dataflow/Source/DataflowEditor/Private/Dataflow/DataflowEditorToolkit.cpp
--- a/Engine/Plugins/Experimental/Dataflow/Source/DataflowEditor/Private/Dataflow/DataflowEditorToolkit.cpp
+++ b/Engine/Plugins/Experimental/Dataflow/Source/DataflowEditor/Private/Dataflow/DataflowEditorToolkit.cpp
@@ -494,26 +494,6 @@ void FDataflowEditorToolkit::OnNodeSelectionChanged(const TSet<UObject*>& InNewS
 		}
 	};
 
-	auto EvaluateCollection = [](TObjectPtr<UDataflowEdNode> InDataflowEdNode, const TSharedPtr<Dataflow::FEngineContext> Context)
-	{
-		if (InDataflowEdNode && Context)
-		{
-			if (const TSharedPtr<FDataflowNode> InDataflowNode = InDataflowEdNode->GetDataflowNode())
-			{
-				for (const FDataflowOutput* const Output : InDataflowNode->GetOutputs())
-				{
-					if (Output->GetType() == FName("FManagedArrayCollection"))
-					{
-						const FManagedArrayCollection DefaultValue;
-						return TSharedPtr<FManagedArrayCollection>(new FManagedArrayCollection(Output->GetValue<FManagedArrayCollection>(*Context, DefaultValue)));
-					}
-				}
-			}
-		}
-
-		return TSharedPtr<FManagedArrayCollection>(new FManagedArrayCollection());
-	};
-
 	// Despite this function's name, we might not have actually changed which node is selected
 	bool bPrimarySelectionChanged = false;
 
@@ -582,15 +562,9 @@ void FDataflowEditorToolkit::OnNodeSelectionChanged(const TSet<UObject*>& InNewS
 					ToolsContext->EndTool(EToolShutdownType::Completed);
 				}
 
-				// Update the Construction viewport with the newly selected node's Collection
-				// @todo(brice) : Is this necessary? FDataflowPreviewScene::Update will generate a new RenderCollection 
-				TSharedPtr<FManagedArrayCollection> Collection = EvaluateCollection(PrimarySelection, EditorContent->GetDataflowContext());
-				if(Collection->HasGroup("Geometry"))
-				{
-					EditorContent->SetPrimarySelectedNode(PrimarySelection);
-					DataflowMode->SetSelectedCollection(Collection);
-				}
+				EditorContent->SetPrimarySelectedNode(PrimarySelection);
 			}
+			EditorContent->SetIsDirty(true);
 		}
 	}
 }
@@ -632,8 +606,8 @@ void FDataflowEditorToolkit::Tick(float DeltaTime)
 				EditorContent->SetDataflowContext(MakeShared<Dataflow::FEngineContext>(EditorContent->GetDataflowOwner(), EditorContent->GetDataflowAsset(), Dataflow::FTimestamp::Invalid));
 				TimeStamp = Dataflow::FTimestamp::Invalid;
 			}
-			FDataflowEditorCommands::EvaluateTerminalNode(*EditorContent->GetDataflowContext().Get(), TimeStamp, EditorContent->GetDataflowAsset(), 
-			nullptr, nullptr, EditorContent->GetDataflowOwner(), EditorContent->GetDataflowTerminal());
+			FDataflowEditorCommands::EvaluateTerminalNode(*EditorContent->GetDataflowContext().Get(), TimeStamp, EditorContent->GetDataflowAsset(),
+				nullptr, nullptr, EditorContent->GetDataflowOwner(), EditorContent->GetDataflowTerminal());
 			EditorContent->SetLastModifiedTimestamp(TimeStamp);
 		}
 	}
@@ -662,8 +636,8 @@ TSharedRef<SDataflowGraphEditor> FDataflowEditorToolkit::CreateGraphEditorWidget
 				Node->Invalidate();
 				Dataflow::FTimestamp TimeStamp = Dataflow::FTimestamp::Invalid;
 				
-				FDataflowEditorCommands::EvaluateTerminalNode(*EditorContent->GetDataflowContext().Get(), TimeStamp,
-															  EditorContent->GetDataflowAsset(), Node, Out, EditorContent->GetDataflowOwner(), EditorContent->GetDataflowTerminal());
+				FDataflowEditorCommands::EvaluateTerminalNode(*EditorContent->GetDataflowContext().Get(), TimeStamp, EditorContent->GetDataflowAsset(),
+					nullptr, nullptr, EditorContent->GetDataflowOwner(), EditorContent->GetDataflowTerminal());
 				EditorContent->SetLastModifiedTimestamp(TimeStamp);
 			}
 		}
diff --git a/Engine/Plugins/Experimental/Dataflow/Source/DataflowEditor/Private/Dataflow/DataflowPreviewScene.cpp b/Engine/Plugins/Experimental/Dataflow/Source/DataflowEditor/Private/Dataflow/DataflowPreviewScene.cpp
--- a/Engine/Plugins/Experimental/Dataflow/Source/DataflowEditor/Private/Dataflow/DataflowPreviewScene.cpp
+++ b/Engine/Plugins/Experimental/Dataflow/Source/DataflowEditor/Private/Dataflow/DataflowPreviewScene.cpp
@@ -115,7 +115,7 @@ void FDataflowConstructionScene::TickDataflowScene(const float DeltaSeconds)
 		if (const UDataflow* Dataflow = DataflowContent->DataflowAsset)
 		{
 			const Dataflow::FTimestamp SystemTimestamp = LatestTimestamp(Dataflow, DataflowContext.Get());
-			if (SystemTimestamp >= DataflowContent->GetLastModifiedTimestamp())
+			if (SystemTimestamp >= DataflowContent->GetLastModifiedTimestamp() || DataflowContent->IsDirty())
 			{
 				DataflowContent->SetLastModifiedTimestamp(SystemTimestamp.Value + 1);
 
diff --git a/Engine/Plugins/Experimental/Dataflow/Source/DataflowEditor/Private/DataflowEditorTools/DataflowEditorWeightMapPaintTool.cpp b/Engine/Plugins/Experimental/Dataflow/Source/DataflowEditor/Private/DataflowEditorTools/DataflowEditorWeightMapPaintTool.cpp
--- a/Engine/Plugins/Experimental/Dataflow/Source/DataflowEditor/Private/DataflowEditorTools/DataflowEditorWeightMapPaintTool.cpp
+++ b/Engine/Plugins/Experimental/Dataflow/Source/DataflowEditor/Private/DataflowEditorTools/DataflowEditorWeightMapPaintTool.cpp
@@ -69,6 +69,46 @@ void UDataflowEditorWeightMapPaintToolBuilder::GetSupportedViewModes(TArray<Data
 	//Modes.Add(Dataflow::EDataflowPatternVertexType::Sim2D);
 }
 
+bool UDataflowEditorWeightMapPaintToolBuilder::CanBuildTool(const FToolBuilderState& SceneState) const
+{
+	auto HasRenderableCollection = [](const FDataflowNode* InDataflowNode, const TSharedPtr<Dataflow::FEngineContext> Context)
+	{
+		if (InDataflowNode && Context)
+		{
+			for (const FDataflowOutput* const Output : InDataflowNode->GetOutputs())
+			{
+				if (Output->GetType() == FName("FManagedArrayCollection"))
+				{
+					const FManagedArrayCollection DefaultValue;
+					const FManagedArrayCollection& Collection = Output->GetValue<FManagedArrayCollection>(*Context, DefaultValue);
+					if (Collection.HasGroup("Geometry"))
+					{
+						return true;
+					}
+				}
+			}
+		}
+
+		return false;
+	};
+
+	if (UMeshSurfacePointMeshEditingToolBuilder::CanBuildTool(SceneState))
+	{
+		if (UDataflowBaseContent* ContextObject = SceneState.ToolManager->GetContextObjectStore()->FindContext<UDataflowBaseContent>())
+		{
+			if (const TSharedPtr<Dataflow::FEngineContext> EvaluationContext = ContextObject->GetDataflowContext())
+			{
+				if (const FDataflowNode* PrimarySelection = ContextObject->GetPrimarySelectedNodeOfType<FDataflowCollectionAddScalarVertexPropertyNode>())
+				{
+					return HasRenderableCollection(PrimarySelection, EvaluationContext);
+				}
+			}
+		}
+	}
+	return false;
+}
+
+
 UMeshSurfacePointTool* UDataflowEditorWeightMapPaintToolBuilder::CreateNewTool(const FToolBuilderState& SceneState) const
 {
 	UDataflowEditorWeightMapPaintTool* PaintTool = NewObject<UDataflowEditorWeightMapPaintTool>(SceneState.ToolManager);
@@ -373,13 +413,20 @@ void UDataflowEditorWeightMapPaintTool::Shutdown(EToolShutdownType ShutdownType)
 		DynamicMeshComponent->OnMeshChanged.Remove(OnDynamicMeshComponentChangedHandle);
 	}
 
-	if (ensure(MeshElementsDisplay->Settings))
+	if (MeshElementsDisplay)
+	{
+		if (ensure(MeshElementsDisplay->Settings))
+		{
+			MeshElementsDisplay->Settings->SaveProperties(this, TEXT("DataflowEditorWeightMapPaintTool2"));
+		}
+		MeshElementsDisplay->Disconnect();
+	}
+
+	if (FilterProperties)
 	{
-		MeshElementsDisplay->Settings->SaveProperties(this, TEXT("DataflowEditorWeightMapPaintTool2"));
+		FilterProperties->SaveProperties(this);
 	}
-	MeshElementsDisplay->Disconnect();
 
-	FilterProperties->SaveProperties(this);
 	if (PreviewMeshActor != nullptr)
 	{
 		PreviewMeshActor->Destroy();
diff --git a/Engine/Plugins/Experimental/Dataflow/Source/DataflowEditor/Public/Dataflow/DataflowEditorMode.h b/Engine/Plugins/Experimental/Dataflow/Source/DataflowEditor/Public/Dataflow/DataflowEditorMode.h
--- a/Engine/Plugins/Experimental/Dataflow/Source/DataflowEditor/Public/Dataflow/DataflowEditorMode.h
+++ b/Engine/Plugins/Experimental/Dataflow/Source/DataflowEditor/Public/Dataflow/DataflowEditorMode.h
@@ -105,14 +105,7 @@ private:
 	void SetRestSpaceViewportClient(TWeakPtr<FDataflowEditorViewportClient, ESPMode::ThreadSafe> ViewportClient);
 	void RefocusRestSpaceViewportClient();
 	void FirstTimeFocusRestSpaceViewport();
-
-	// @todo(brice) : Can this be in the content?
-	// intended to be called by the toolkit when selected node in the Dataflow graph changes
-	void SetSelectedCollection(TSharedPtr<FManagedArrayCollection> Collection);
 	bool IsComponentSelected(const UPrimitiveComponent* InComponent);
-	TSharedPtr<FManagedArrayCollection> GetSelectedCollection();
-	TSharedPtr<FManagedArrayCollection> SelectedCollection = nullptr;
-
 	void StartToolForSelectedNode(const UObject* SelectedNode);
 	void OnDataflowNodeDeleted(const TSet<UObject*>& DeletedNodes);
 
diff --git a/Engine/Plugins/Experimental/Dataflow/Source/DataflowEditor/Public/DataflowEditorTools/DataflowEditorWeightMapPaintTool.h b/Engine/Plugins/Experimental/Dataflow/Source/DataflowEditor/Public/DataflowEditorTools/DataflowEditorWeightMapPaintTool.h
--- a/Engine/Plugins/Experimental/Dataflow/Source/DataflowEditor/Public/DataflowEditorTools/DataflowEditorWeightMapPaintTool.h
+++ b/Engine/Plugins/Experimental/Dataflow/Source/DataflowEditor/Public/DataflowEditorTools/DataflowEditorWeightMapPaintTool.h
@@ -61,6 +61,7 @@ class DATAFLOWEDITOR_API UDataflowEditorWeightMapPaintToolBuilder : public UMesh
 
 private:
 	virtual void GetSupportedViewModes(TArray<Dataflow::EDataflowPatternVertexType>& Modes) const override;
+	virtual bool CanBuildTool(const FToolBuilderState& SceneState) const override;
 	virtual UMeshSurfacePointTool* CreateNewTool(const FToolBuilderState& SceneState) const override;
 	virtual bool CanSetConstructionViewWireframeActive() const { return false; }
 };
diff --git a/Engine/Source/Runtime/Experimental/Dataflow/Engine/Private/Dataflow/DataflowContent.cpp b/Engine/Source/Runtime/Experimental/Dataflow/Engine/Private/Dataflow/DataflowContent.cpp
--- a/Engine/Source/Runtime/Experimental/Dataflow/Engine/Private/Dataflow/DataflowContent.cpp
+++ b/Engine/Source/Runtime/Experimental/Dataflow/Engine/Private/Dataflow/DataflowContent.cpp
@@ -1,9 +1,11 @@
-﻿// Copyright Epic Games, Inc. All Rights Reserved.
+// Copyright Epic Games, Inc. All Rights Reserved.
 
 #include "Dataflow/DataflowContent.h"
+
 #include "Animation/AnimSingleNodeInstance.h"
-#include "Components/SkeletalMeshComponent.h"
 #include "Animation/Skeleton.h"
+#include "Components/SkeletalMeshComponent.h"
+#include "Dataflow/DataflowNodeParameters.h"
 #include "Engine/SkeletalMesh.h"
 #include "GameFramework/Actor.h"
 #include "PreviewScene.h"
@@ -14,6 +16,15 @@ UDataflowBaseContent::UDataflowBaseContent()
 {
 }
 
+void UDataflowBaseContent::SetLastModifiedTimestamp(Dataflow::FTimestamp InTimestamp) 
+{ 
+	if (InTimestamp.IsInvalid() || LastModifiedTimestamp < InTimestamp)
+	{
+		LastModifiedTimestamp = InTimestamp; 
+		bIsDirty = true;
+	}
+}
+
 void UDataflowBaseContent::BuildBaseContent(TObjectPtr<UObject> ContentOwner)
 {
 	DataflowContext = MakeShared<Dataflow::FEngineContext>(ContentOwner, DataflowAsset, FPlatformTime::Cycles64());
diff --git a/Engine/Source/Runtime/Experimental/Dataflow/Engine/Public/Dataflow/DataflowContent.h b/Engine/Source/Runtime/Experimental/Dataflow/Engine/Public/Dataflow/DataflowContent.h
--- a/Engine/Source/Runtime/Experimental/Dataflow/Engine/Public/Dataflow/DataflowContent.h
+++ b/Engine/Source/Runtime/Experimental/Dataflow/Engine/Public/Dataflow/DataflowContent.h
@@ -1,4 +1,4 @@
-﻿// Copyright Epic Games, Inc. All Rights Reserved.
+// Copyright Epic Games, Inc. All Rights Reserved.
 
 #pragma once
 
@@ -85,10 +85,6 @@ class DATAFLOWENGINE_API UDataflowBaseContent : public UDataflowContextObject
 public:
 	UDataflowBaseContent();
 
-	/** Check if the dataflow data are dirty and need to be re-rendered */
-	bool IsDirty() const { return bIsDirty; }
-	void SetIsDirty(bool InDirty) { bIsDirty = InDirty; }
-
 	/** Data flow asset that we will edit */
 	UPROPERTY(EditAnywhere, Category = "Dataflow")
 	TObjectPtr<UDataflow> DataflowAsset = nullptr;
@@ -97,6 +93,27 @@ public:
 	UPROPERTY(EditAnywhere, Category = "Dataflow")
 	FString DataflowTerminal = "";
 	
+	/** 
+	*	Dirty - State Invalidation
+	*   Check if non-graph specific data has been changed, this usually requires a re-render 
+	*/
+	bool IsDirty() const { return bIsDirty; }
+	void SetIsDirty(bool InDirty) { bIsDirty = InDirty; }
+
+	/** 
+	*	LastModifiedTimestamp - State Invalidation 
+	*   Dataflow timestamp accessors can be used to see if the EvaluationContext has been invalidated. 
+	*/
+	void SetLastModifiedTimestamp(Dataflow::FTimestamp InTimestamp);
+	const Dataflow::FTimestamp& GetLastModifiedTimestamp() const { return LastModifiedTimestamp; }
+
+	/**  
+	*	Context - Dataflow Evaluation State
+	*   Dataflow context stores the evaluated state of the graph. 
+	*/
+	void SetDataflowContext(const TSharedPtr<Dataflow::FEngineContext>& InContext) { DataflowContext = InContext; bIsDirty = true; }
+	const TSharedPtr<Dataflow::FEngineContext>& GetDataflowContext() const { return DataflowContext; }
+
 	/** Return the simulation time range to be used in the simulation viewport */
 	virtual FVector2f GetSimulationRange() const { return FVector2f(0.0f, 100.0f); } 
 	
@@ -124,13 +141,6 @@ public:
 	void SetDataflowTerminal(const FString& InPath) { DataflowTerminal = InPath;  bIsDirty = true;}
 	const FString& GetDataflowTerminal() const { return DataflowTerminal; }
 
-	/**  Data flow context accessors */
-	void SetDataflowContext(const TSharedPtr<Dataflow::FEngineContext>& InContext) { DataflowContext = InContext; bIsDirty = true;}
-	const TSharedPtr<Dataflow::FEngineContext>& GetDataflowContext() const { return DataflowContext; }
-
-	/** Data flow timestamp accessors */
-	void SetLastModifiedTimestamp(Dataflow::FTimestamp InTimestamp) { LastModifiedTimestamp = InTimestamp; bIsDirty = true;}
-	const Dataflow::FTimestamp& GetLastModifiedTimestamp() const { return LastModifiedTimestamp; }
 	
 protected :
 	
