Dataflow : Fix cast on comment nodes.

diff --git a/Engine/Plugins/Experimental/Dataflow/Source/DataflowEditor/Private/Dataflow/DataflowEditorMode.cpp b/Engine/Plugins/Experimental/Dataflow/Source/DataflowEditor/Private/Dataflow/DataflowEditorMode.cpp
--- a/Engine/Plugins/Experimental/Dataflow/Source/DataflowEditor/Private/Dataflow/DataflowEditorMode.cpp
+++ b/Engine/Plugins/Experimental/Dataflow/Source/DataflowEditor/Private/Dataflow/DataflowEditorMode.cpp
@@ -665,22 +665,19 @@ UEdGraphNode* UDataflowEditorMode::GetSingleSelectedNodeWithOutputType(const FNa
 		return nullptr;
 	}
 
-	const UDataflowEdNode* const SelectedDataflowEdNode = CastChecked<UDataflowEdNode>(SelectedNode);
-	const TSharedPtr<const FDataflowNode> SelectedDataflowNode = SelectedDataflowEdNode->GetDataflowNode();
-
-	if (!SelectedDataflowNode)
-	{
-		// This can happen when the user deletes a node. Seems like the Dataflow FGraph is updated with the removed node before the graph editor can update.
-		return nullptr;
-	}
-
-	for (const FDataflowOutput* const Output : SelectedDataflowNode->GetOutputs())
-	{
-		if (Output->GetType() == SelectedNodeOutputTypeName)
-		{
-			return SelectedNode;
-		}
-	}
+	 if (const UDataflowEdNode* const SelectedDataflowEdNode = Cast<UDataflowEdNode>(SelectedNode))
+	 {
+		 if (TSharedPtr<const FDataflowNode> SelectedDataflowNode = SelectedDataflowEdNode->GetDataflowNode())
+		 {
+			 for (const FDataflowOutput* const Output : SelectedDataflowNode->GetOutputs())
+			 {
+				 if (Output->GetType() == SelectedNodeOutputTypeName)
+				 {
+					 return SelectedNode;
+				 }
+			 }
+		 }
+	 }
 
 	return nullptr;
 }
