Dataflow : Hide the floor

diff --git a/Engine/Plugins/Experimental/Dataflow/Source/DataflowEditor/Private/Dataflow/DataflowPreviewScene.cpp b/Engine/Plugins/Experimental/Dataflow/Source/DataflowEditor/Private/Dataflow/DataflowPreviewScene.cpp
--- a/Engine/Plugins/Experimental/Dataflow/Source/DataflowEditor/Private/Dataflow/DataflowPreviewScene.cpp
+++ b/Engine/Plugins/Experimental/Dataflow/Source/DataflowEditor/Private/Dataflow/DataflowPreviewScene.cpp
@@ -17,11 +17,16 @@
 
 #define LOCTEXT_NAMESPACE "FDataflowPreviewScene"
 
+
+bool bDataflowShowFloorDefault = false;
+FAutoConsoleVariableRef CVARDataflowShowFloorDefault(TEXT("p.Dataflow.Editor.ShowFloor"), bDataflowShowFloorDefault, TEXT("Show the floor in the dataflow editor[def:false]"));
+
+
 FDataflowPreviewScene::FDataflowPreviewScene(FPreviewScene::ConstructionValues ConstructionValues, TObjectPtr<UDataflowBaseContent> InEditorContent) 
 	: FAdvancedPreviewScene(ConstructionValues), DataflowContent(InEditorContent)
 {
 	check(DataflowContent);
-	SetFloorVisibility(true, true);
+	SetFloorVisibility(bDataflowShowFloorDefault, true);
 
 	RootSceneActor = GetWorld()->SpawnActor<AActor>(AActor::StaticClass());
 }
