Dataflow : Fixed extra UI tabs from launching.
#perflight 65bac29ad359c2603f291890

diff --git a/Engine/Plugins/Experimental/Dataflow/Source/DataflowEditor/Private/Dataflow/DataflowEditorToolkit.cpp b/Engine/Plugins/Experimental/Dataflow/Source/DataflowEditor/Private/Dataflow/DataflowEditorToolkit.cpp
--- a/Engine/Plugins/Experimental/Dataflow/Source/DataflowEditor/Private/Dataflow/DataflowEditorToolkit.cpp
+++ b/Engine/Plugins/Experimental/Dataflow/Source/DataflowEditor/Private/Dataflow/DataflowEditorToolkit.cpp
@@ -10,6 +10,7 @@
 #include "Dataflow/DataflowEditorCommands.h"
 #include "Dataflow/DataflowEditorMode.h"
 #include "Dataflow/DataflowEditorModeToolkit.h"
+#include "Dataflow/DataflowEditorModule.h"
 #include "Dataflow/DataflowEditorModeUILayer.h"
 #include "Dataflow/DataflowEditorViewport.h"
 #include "Dataflow/DataflowEditorViewportClient.h"
@@ -58,11 +59,11 @@ const FName FDataflowEditorToolkit::CollectionSpreadSheetTabId_4(TEXT("DataflowE
 const FName FDataflowEditorToolkit::SimulationViewportTabId(TEXT("DataflowEditor_SimulationViewport"));
 
 FDataflowEditorToolkit::FDataflowEditorToolkit(UAssetEditor* InOwningAssetEditor)
-	: FBaseCharacterFXEditorToolkit(InOwningAssetEditor, FName("DataflowEditor")) 
+	: FBaseCharacterFXEditorToolkit(InOwningAssetEditor, FName("DataflowEditor"))
 {
 	check(Cast<UDataflowEditor>(InOwningAssetEditor));
 
-	StandaloneDefaultLayout = FTabManager::NewLayout(FName("DataflowEditorLayout.V10"))
+	StandaloneDefaultLayout = FTabManager::NewLayout(FName("DataflowEditorLayout03"))
 		->AddArea
 		(
 			FTabManager::NewPrimaryArea()->SetOrientation(Orient_Horizontal)
@@ -73,7 +74,7 @@ FDataflowEditorToolkit::FDataflowEditorToolkit(UAssetEditor* InOwningAssetEditor
 				->Split
 				(
 					FTabManager::NewSplitter()->SetOrientation(Orient_Horizontal)
-					->SetSizeCoefficient(0.55f)	// Relative height of (Tools Panel, Construction Viewport, Preview Viewport) vs (Dataflow Graph Editor, Outliner)
+					->SetSizeCoefficient(0.60f)	// Relative height of (Tools Panel, Construction Viewport, Preview Viewport) vs (Dataflow Graph Editor, Outliner)
 					->Split
 					(
 						FTabManager::NewStack()
@@ -92,23 +93,59 @@ FDataflowEditorToolkit::FDataflowEditorToolkit(UAssetEditor* InOwningAssetEditor
 				)
 				->Split
 				(
-					FTabManager::NewStack()
-					->SetSizeCoefficient(0.45f)	// Relative height of (Dataflow Graph Editor, Outliner) vs (Tools Panel, Construction Viewport, Preview Viewport)
-					->AddTab(GraphCanvasTabId, ETabState::OpenedTab)
-					->SetExtensionId("GraphEditorArea")
-					->SetHideTabWell(false)
-					->SetForegroundTab(GraphCanvasTabId)
+					FTabManager::NewSplitter()->SetOrientation(Orient_Horizontal)
+					->SetSizeCoefficient(0.40f)	// Relative height of (Dataflow Node Details) vs (Asset Details, Preview Scene Details)
+					->Split
+					(
+						FTabManager::NewStack()
+						->SetSizeCoefficient(0.2f)
+						->AddTab(CollectionSpreadSheetTabId_1, ETabState::OpenedTab)
+						->SetExtensionId("CollectionSpreadSheetArea")
+						->SetHideTabWell(false)
+					)
+					->Split
+					(
+						FTabManager::NewStack()
+						->SetSizeCoefficient(0.8f)	// Relative height of (Dataflow Graph Editor, Outliner) vs (Tools Panel, Construction Viewport, Preview Viewport)
+						->AddTab(GraphCanvasTabId, ETabState::OpenedTab)
+						->SetExtensionId("GraphEditorArea")
+						->SetHideTabWell(false)
+						->SetForegroundTab(GraphCanvasTabId)
+					)
 				)
 			)
 			->Split
 			(
-				FTabManager::NewStack()
-				->SetSizeCoefficient(0.2f)	// Relative height of (Dataflow Node Details) vs (Asset Details, Preview Scene Details)
-				->AddTab(NodeDetailsTabId, ETabState::OpenedTab)
-				->SetExtensionId("NodeDetailsArea")
-				->SetHideTabWell(false)
+				FTabManager::NewSplitter()->SetOrientation(Orient_Vertical)
+				->SetSizeCoefficient(0.2f)	// Relative width of (Asset Details, Preview Scene Details, Dataflow Node Details) vs (Tools Panel, Construction Viewport, Preview Viewport, Dataflow Graph Editor, Outliner)
+				->Split
+				(
+					FTabManager::NewStack()
+					->SetSizeCoefficient(0.65f)	// Relative height of (Asset Details, Preview Scene Details) vs (Dataflow Node Details)
+					->AddTab(DetailsTabID, ETabState::OpenedTab)
+					->SetExtensionId("DetailsArea")
+					->SetHideTabWell(false)
+					->SetForegroundTab(DetailsTabID)
+				)
+				->Split
+				(
+					FTabManager::NewStack()
+					->SetSizeCoefficient(0.2f)	// Relative height of (Dataflow Node Details) vs (Asset Details, Preview Scene Details)
+					->AddTab(NodeDetailsTabId, ETabState::OpenedTab)
+					->SetExtensionId("NodeDetailsArea")
+					->SetHideTabWell(false)
+				)
 			)
 		);
+
+	// Add any extenders specified by the UISubsystem
+	// The extenders provide defined locations for FModeToolkit to attach
+	// tool palette tabs and detail panel tabs
+	LayoutExtender = MakeShared<FLayoutExtender>();
+	FDataflowEditorModule* Module = &FModuleManager::LoadModuleChecked<FDataflowEditorModule>("DataflowEditor");
+	Module->OnRegisterLayoutExtensions().Broadcast(*LayoutExtender);
+	StandaloneDefaultLayout->ProcessExtensions(*LayoutExtender);
+
 	FAdvancedPreviewScene::ConstructionValues PreviewSceneArgs;
 	PreviewSceneArgs.bShouldSimulatePhysics = 1;
 	PreviewSceneArgs.bCreatePhysicsScene = 1;
