Dataflow : Multi mesh selection per collection.

diff --git a/Engine/Plugins/Experimental/Dataflow/Source/DataflowAssetTools/Private/Dataflow/CollectionRenderingPatternUtility.cpp b/Engine/Plugins/Experimental/Dataflow/Source/DataflowAssetTools/Private/Dataflow/CollectionRenderingPatternUtility.cpp
--- a/Engine/Plugins/Experimental/Dataflow/Source/DataflowAssetTools/Private/Dataflow/CollectionRenderingPatternUtility.cpp
+++ b/Engine/Plugins/Experimental/Dataflow/Source/DataflowAssetTools/Private/Dataflow/CollectionRenderingPatternUtility.cpp
@@ -18,36 +18,62 @@ namespace Dataflow
 	namespace Conversion
 	{
 		// Convert a rendering facade to a dynamic mesh
-		void RenderingFacadeToDynamicMesh(const GeometryCollection::Facades::FRenderingFacade& Facade, FDynamicMesh3& DynamicMesh)
+		void RenderingFacadeToDynamicMesh(const GeometryCollection::Facades::FRenderingFacade& Facade, int32 InMeshIndex, FDynamicMesh3& DynamicMesh)
 		{
 			if (Facade.CanRenderSurface())
 			{
-				const int32 NumTriangles = Facade.NumTriangles();
-				const int32 NumVertices = Facade.NumVertices();
+				int32 StartTriangles = 0;
+				int32 StartVertices = 0;
+				int32 NumTriangles = Facade.NumTriangles();
+				int32 NumVertices = Facade.NumVertices();
+
+				if (InMeshIndex != INDEX_NONE)
+				{
+					if (ensure(0 <= InMeshIndex && InMeshIndex < Facade.NumGeometry()))
+					{
+						StartTriangles = Facade.GetIndicesStart()[InMeshIndex];
+						StartVertices = Facade.GetVertexStart()[InMeshIndex];
+						NumTriangles = Facade.GetIndicesCount()[InMeshIndex];
+						NumVertices = Facade.GetVertexCount()[InMeshIndex];
+					}
+				}
 
+				TArray<int32> Remapping;
 				const TManagedArray<FIntVector>& Indices = Facade.GetIndices();
 				const TManagedArray<FVector3f>& Positions = Facade.GetVertices();
 				const TManagedArray<FVector3f>& Normals = Facade.GetNormals();
 				const TManagedArray<FLinearColor>& Colors = Facade.GetVertexColor();
 
-				for (int32 VertexIndex = 0; VertexIndex < NumVertices; ++VertexIndex)
+				int32 LastVertexIndex = StartVertices + NumVertices;
+				for (int32 VertexIndex = StartVertices; VertexIndex < LastVertexIndex; ++VertexIndex)
 				{
 					DynamicMesh.AppendVertex(FVertexInfo(FVector3d(Positions[VertexIndex]), Normals[VertexIndex],
 						FVector3f(Colors[VertexIndex].R, Colors[VertexIndex].G, Colors[VertexIndex].B)));
+					Remapping.Add(VertexIndex);
 				}
-				for (int32 TriangleIndex = 0; TriangleIndex < NumTriangles; ++TriangleIndex)
+				int32 LastTriangleIndex = StartTriangles + NumTriangles;
+				for (int32 TriangleIndex = StartTriangles; TriangleIndex < LastTriangleIndex; ++TriangleIndex)
 				{
-					DynamicMesh.AppendTriangle(FIndex3i(Indices[TriangleIndex].X, Indices[TriangleIndex].Y, Indices[TriangleIndex].Z));
+					DynamicMesh.AppendTriangle(FIndex3i(
+						Indices[TriangleIndex].X - StartVertices,
+						Indices[TriangleIndex].Y - StartVertices,
+						Indices[TriangleIndex].Z - StartVertices)
+					);
 				}
 				FMeshNormals::QuickComputeVertexNormals(DynamicMesh);
 
 				DynamicMesh.EnableAttributes();
+
+				// Build Remmaping indices back into the colleciton. 
+				if (ensure(Remapping.Num() < Facade.NumVertices()))
+				{
+					UE::Geometry::FNonManifoldMappingSupport::AttachNonManifoldVertexMappingData(Remapping, DynamicMesh);
+				}
+
 				DynamicMesh.Attributes()->EnablePrimaryColors();
 				DynamicMesh.Attributes()->PrimaryColors()->CreateFromPredicate([](int ParentVID, int TriIDA, int TriIDB) {return true; }, 0.f);
 				DynamicMesh.EnableVertexColors(FVector3f::Zero());
-
 				FDynamicMeshColorOverlay* const ColorOverlay = DynamicMesh.Attributes()->PrimaryColors();
-
 				auto SetColorsFromWeights = [&](int TriangleID)
 				{
 					const FIndex3i Tri = DynamicMesh.GetTriangle(TriangleID);
@@ -65,27 +91,6 @@ namespace Dataflow
 			}
 		}
 
-		// Convert a dataflow component to a dynamic mesh
-		void DataflowToDynamicMesh(TSharedPtr<::Dataflow::FEngineContext> DataflowContext, UObject* Asset, UDataflow* Dataflow, FDynamicMesh3& DynamicMesh)
-		{
-			// just call update on the preview scene.
-
-			FManagedArrayCollection RenderCollection;
-			GeometryCollection::Facades::FRenderingFacade Facade(RenderCollection);
-			Facade.DefineSchema();
-
-
-			for (const UDataflowEdNode* Target : Dataflow->GetRenderTargets())
-			{
-				if (Target)
-				{
-					// @todo(brice) Fix this cast
-					// Target->Render(Facade, DataflowContext);
-				}
-			}
-			RenderingFacadeToDynamicMesh(Facade, DynamicMesh);
-		}
-
 		// Convert a dynamic mesh to a rendering facade
 		void DynamicMeshToRenderingFacade(const FDynamicMesh3& DynamicMesh, GeometryCollection::Facades::FRenderingFacade& Facade)
 		{
@@ -110,12 +115,6 @@ namespace Dataflow
 				}
 			}
 		}
-
-		// Convert a dynamic mesh to a dataflow component
-		void DynamicMeshToDataflow(const FDynamicMesh3& DynamicMesh, UDataflow* Dataflow)
-		{
-			//todo
-		}
 	}
 
 }	// namespace Dataflow
\ No newline at end of file
diff --git a/Engine/Plugins/Experimental/Dataflow/Source/DataflowAssetTools/Public/Dataflow/CollectionRenderingPatternUtility.h b/Engine/Plugins/Experimental/Dataflow/Source/DataflowAssetTools/Public/Dataflow/CollectionRenderingPatternUtility.h
--- a/Engine/Plugins/Experimental/Dataflow/Source/DataflowAssetTools/Public/Dataflow/CollectionRenderingPatternUtility.h
+++ b/Engine/Plugins/Experimental/Dataflow/Source/DataflowAssetTools/Public/Dataflow/CollectionRenderingPatternUtility.h
@@ -14,17 +14,11 @@ namespace Dataflow
 {
 	namespace Conversion
 	{
-		// Convert a rendering facade to a dynamic mesh
-		void DATAFLOWASSETTOOLS_API RenderingFacadeToDynamicMesh(const GeometryCollection::Facades::FRenderingFacade& Facade, UE::Geometry::FDynamicMesh3& DynamicMesh);
-
-		// Convert a dataflow component to a dynamic mesh
-		void DATAFLOWASSETTOOLS_API DataflowToDynamicMesh(TSharedPtr<::Dataflow::FEngineContext> DataflowContext, UObject* Asset, UDataflow* Dataflow, UE::Geometry::FDynamicMesh3& DynamicMesh);
+		// Convert a rendering facade to a dynamic mesh. if InMeshIndex==INDEX_NONE then all Geometry is converted. 
+		void DATAFLOWASSETTOOLS_API RenderingFacadeToDynamicMesh(const GeometryCollection::Facades::FRenderingFacade& Facade, int32 InMeshIndex, UE::Geometry::FDynamicMesh3& DynamicMesh);
 
 		// Convert a dynamic mesh to a rendering facade
 		void DATAFLOWASSETTOOLS_API DynamicMeshToRenderingFacade(const UE::Geometry::FDynamicMesh3& DynamicMesh, GeometryCollection::Facades::FRenderingFacade& Facade);
-
-		// Convert a dynamic mesh to a dataflow component
-		void DATAFLOWASSETTOOLS_API DynamicMeshToDataflow(const UE::Geometry::FDynamicMesh3& DynamicMesh, UDataflow* Dataflow);
 	}
 
 }	// namespace Dataflow
\ No newline at end of file
diff --git a/Engine/Plugins/Experimental/Dataflow/Source/DataflowEditor/Private/Dataflow/DataflowEditorCollectionComponent.cpp b/Engine/Plugins/Experimental/Dataflow/Source/DataflowEditor/Private/Dataflow/DataflowEditorCollectionComponent.cpp
new file mode 100644
--- /dev/null
+++ b/Engine/Plugins/Experimental/Dataflow/Source/DataflowEditor/Private/Dataflow/DataflowEditorCollectionComponent.cpp
@@ -0,0 +1,14 @@
+// Copyright Epic Games, Inc. All Rights Reserved.
+
+#include "Dataflow/DataflowEditorCollectionComponent.h"
+
+
+#include UE_INLINE_GENERATED_CPP_BY_NAME(DataflowEditorCollectionComponent)
+
+DEFINE_LOG_CATEGORY_STATIC(LogDataflowEditorCollectionComponentInternal, Log, All);
+
+UDataflowEditorCollectionComponent::UDataflowEditorCollectionComponent(const FObjectInitializer& ObjectInitializer)
+	: Super(ObjectInitializer)
+{
+}
+
diff --git a/Engine/Plugins/Experimental/Dataflow/Source/DataflowEditor/Private/Dataflow/DataflowEditorCollectionComponent.h b/Engine/Plugins/Experimental/Dataflow/Source/DataflowEditor/Private/Dataflow/DataflowEditorCollectionComponent.h
new file mode 100644
--- /dev/null
+++ b/Engine/Plugins/Experimental/Dataflow/Source/DataflowEditor/Private/Dataflow/DataflowEditorCollectionComponent.h
@@ -0,0 +1,27 @@
+// Copyright Epic Games, Inc. All Rights Reserved.
+#pragma once
+
+
+#include "Components/DynamicMeshComponent.h"
+#include "UObject/ObjectMacros.h"
+
+#include "DataflowEditorCollectionComponent.generated.h"
+
+class UDataflowEdNode;
+
+/**
+*	FleshComponent
+*/
+UCLASS(MinimalAPI)
+class UDataflowEditorCollectionComponent : public UDynamicMeshComponent
+{
+	GENERATED_UCLASS_BODY()
+public:
+
+	UPROPERTY()
+	int32 MeshIndex = INDEX_NONE;
+
+	UPROPERTY()
+	TObjectPtr<const UDataflowEdNode> Node = nullptr;
+
+};
diff --git a/Engine/Plugins/Experimental/Dataflow/Source/DataflowEditor/Private/Dataflow/DataflowPreviewScene.cpp b/Engine/Plugins/Experimental/Dataflow/Source/DataflowEditor/Private/Dataflow/DataflowPreviewScene.cpp
--- a/Engine/Plugins/Experimental/Dataflow/Source/DataflowEditor/Private/Dataflow/DataflowPreviewScene.cpp
+++ b/Engine/Plugins/Experimental/Dataflow/Source/DataflowEditor/Private/Dataflow/DataflowPreviewScene.cpp
@@ -7,6 +7,7 @@
 #include "Components/DynamicMeshComponent.h"
 #include "Dataflow/CollectionRenderingPatternUtility.h"
 #include "Dataflow/DataflowEditor.h"
+#include "Dataflow/DataflowEditorCollectionComponent.h"
 #include "Dataflow/DataflowContent.h"
 #include "Dataflow/DataflowEditorStyle.h"
 #include "Dataflow/DataflowEditorUtil.h"
@@ -101,7 +102,7 @@ void FDataflowConstructionScene::AddReferencedObjects(FReferenceCollector& Colle
 	FDataflowPreviewScene::AddReferencedObjects(Collector);
 
 	Collector.AddReferencedObjects(DynamicMeshComponents);
-	Collector.AddReferencedObject(WireframeDraw);
+	Collector.AddReferencedObjects(WireframeElements);
 }
 
 FORCEINLINE Dataflow::FTimestamp LatestTimestamp(const UDataflow* Dataflow, const ::Dataflow::FContext* Context)
@@ -146,9 +147,9 @@ void FDataflowConstructionScene::TickDataflowScene(const float DeltaSeconds)
 		}
 	}
 
-	if (WireframeDraw)
+	for (FRenderWireElement Elem : WireframeElements)
 	{
-		WireframeDraw->OnTick(DeltaSeconds);
+		Elem.Value->OnTick(DeltaSeconds);
 	}
 }
 
@@ -167,24 +168,32 @@ void FDataflowConstructionScene::UpdateDynamicMeshComponents()
 		const TSharedPtr<Dataflow::FEngineContext>& DataflowContext = DataflowContent->GetDataflowContext();
 		if(DataflowAsset && DataflowContext)
 		{
-			for (const UDataflowEdNode* Target : DataflowAsset->GetRenderTargets())
+			for (TObjectPtr<const UDataflowEdNode> Target : DataflowAsset->GetRenderTargets())
 			{
 				if (Target)
 				{
-					FDynamicMesh3 DynamicMesh;
 					TSharedPtr<FManagedArrayCollection> RenderCollection(new FManagedArrayCollection);
 					GeometryCollection::Facades::FRenderingFacade Facade(*RenderCollection);
 					Facade.DefineSchema();
 
 					Target->Render(Facade, DataflowContext);
-					Dataflow::Conversion::RenderingFacadeToDynamicMesh(Facade, DynamicMesh);
-					
-					if (Target == DataflowContent->GetPrimarySelectedNode())
+
+					int32 NumGeometry = Facade.NumGeometry();
+					for (int32 MeshIndex = 0; MeshIndex < NumGeometry; MeshIndex++)
 					{
-						DataflowContent->SetPrimaryRenderCollection(RenderCollection);
+						FDynamicMesh3 DynamicMesh;
+						Dataflow::Conversion::RenderingFacadeToDynamicMesh(Facade, MeshIndex, DynamicMesh);
+
+						if (DynamicMesh.VertexCount())
+						{
+							if (Target == DataflowContent->GetPrimarySelectedNode())
+							{
+								DataflowContent->SetPrimaryRenderCollection(RenderCollection);
+							}
+
+							AddDynamicMeshComponent({Target, MeshIndex }, MoveTemp(DynamicMesh), {});
+						}
 					}
-					
-					AddDynamicMeshComponent(MoveTemp(DynamicMesh), {});
 				}
 			}
 		}
@@ -194,8 +203,10 @@ void FDataflowConstructionScene::UpdateDynamicMeshComponents()
 void FDataflowConstructionScene::ResetDynamicMeshComponents()
 {
 	USelection* SelectedComponents = DataflowModeManager->GetSelectedComponents();
-	for(const TObjectPtr<UDynamicMeshComponent>& DynamicMeshComponent : DynamicMeshComponents)
+	for (FRenderElement RenderElement : DynamicMeshComponents)
 	{
+		TObjectPtr<UDynamicMeshComponent>& DynamicMeshComponent = RenderElement.Value;
+
 		DynamicMeshComponent->SelectionOverrideDelegate.Unbind();
 		if (SelectedComponents->IsSelected(DynamicMeshComponent))
 		{
@@ -207,10 +218,11 @@ void FDataflowConstructionScene::ResetDynamicMeshComponents()
 	DynamicMeshComponents.Reset();
 }
 
-TObjectPtr<UDynamicMeshComponent>& FDataflowConstructionScene::AddDynamicMeshComponent(UE::Geometry::FDynamicMesh3&& DynamicMesh, const TArray<UMaterialInterface*>& MaterialSet)
+TObjectPtr<UDynamicMeshComponent>& FDataflowConstructionScene::AddDynamicMeshComponent(FDataflowRenderKey InKey, UE::Geometry::FDynamicMesh3&& DynamicMesh, const TArray<UMaterialInterface*>& MaterialSet)
 {
-	TObjectPtr<UDynamicMeshComponent> DynamicMeshComponent = NewObject<UDynamicMeshComponent>(RootSceneActor);
-		
+	TObjectPtr<UDataflowEditorCollectionComponent> DynamicMeshComponent = NewObject<UDataflowEditorCollectionComponent>(RootSceneActor);
+	DynamicMeshComponent->MeshIndex = InKey.Value;
+	DynamicMeshComponent->Node = InKey.Key;;
 	DynamicMeshComponent->SetMesh(MoveTemp(DynamicMesh));
 	
 	// @todo(Material) This is just to have a material, we should transfer the materials from the assets if they have them. 
@@ -235,43 +247,44 @@ TObjectPtr<UDynamicMeshComponent>& FDataflowConstructionScene::AddDynamicMeshCom
 	DynamicMeshComponent->SelectionOverrideDelegate = UPrimitiveComponent::FSelectionOverride::CreateRaw(this, &FDataflowPreviewScene::IsComponentSelected);
 	DynamicMeshComponent->UpdateBounds();
 
-	AddComponent(DynamicMeshComponent, DynamicMeshComponent->GetRelativeTransform());
-		
-	const int32 ElementIndex = DynamicMeshComponents.Emplace(DynamicMeshComponent);
-	return DynamicMeshComponents[ElementIndex];
+	AddComponent(DynamicMeshComponent, DynamicMeshComponent->GetRelativeTransform());	
+	DynamicMeshComponents.Emplace(InKey, DynamicMeshComponent);
+	return DynamicMeshComponents[InKey];
 }
 
 void FDataflowConstructionScene::AddWireframeMeshElementsVisualizer()
 {
-	ensure(WireframeDraw==nullptr);
-	if (DynamicMeshComponents.Num())
+	ensure(WireframeElements.Num()==0);
+	for(FRenderElement Elem : DynamicMeshComponents)
 	{
+		TObjectPtr<UDynamicMeshComponent> DynamicMeshComponent = Elem.Value;
+
 		// Set up the wireframe display of the rest space mesh.
 
-		WireframeDraw = NewObject<UMeshElementsVisualizer>(RootSceneActor);
-		WireframeDraw->CreateInWorld(GetWorld(), FTransform::Identity);
+		TObjectPtr<UMeshElementsVisualizer> WireframeDraw = NewObject<UMeshElementsVisualizer>(RootSceneActor);
+		WireframeElements.Add(DynamicMeshComponent, WireframeDraw);
 
+		WireframeDraw->CreateInWorld(GetWorld(), FTransform::Identity);
 		WireframeDraw->Settings->DepthBias = 2.0;
 		WireframeDraw->Settings->bAdjustDepthBiasUsingMeshSize = false;
 		WireframeDraw->Settings->bShowWireframe = true;
 		WireframeDraw->Settings->bShowBorders = true;
 		WireframeDraw->Settings->bShowUVSeams = false;
-
 		WireframeDraw->WireframeComponent->BoundaryEdgeThickness = 2;
 
-		WireframeDraw->SetMeshAccessFunction([this](UMeshElementsVisualizer::ProcessDynamicMeshFunc ProcessFunc) 
+		WireframeDraw->SetMeshAccessFunction([DynamicMeshComponent](UMeshElementsVisualizer::ProcessDynamicMeshFunc ProcessFunc)
 		{
-			for (auto DynamicMeshComponent : DynamicMeshComponents) ProcessFunc(*DynamicMeshComponent->GetMesh());
+			ProcessFunc(*DynamicMeshComponent->GetMesh());
 		});
 
-		for (auto DynamicMeshComponent : DynamicMeshComponents)
+		for (FRenderElement RenderElement : DynamicMeshComponents)
 		{
-			DynamicMeshComponent->OnMeshChanged.Add(FSimpleMulticastDelegate::FDelegate::CreateLambda([this]()
+			RenderElement.Value->OnMeshChanged.Add(FSimpleMulticastDelegate::FDelegate::CreateLambda([WireframeDraw,this]()
 			{
 				WireframeDraw->NotifyMeshChanged();
 			}));
 
-			const bool bRestSpaceMeshVisible = DynamicMeshComponent->GetVisibleFlag();
+			const bool bRestSpaceMeshVisible = RenderElement.Value->GetVisibleFlag();
 			WireframeDraw->Settings->bVisible = bRestSpaceMeshVisible && bConstructionViewWireframe;
 		}
 		PropertyObjectsToTick.Add(WireframeDraw->Settings);
@@ -280,11 +293,11 @@ void FDataflowConstructionScene::AddWireframeMeshElementsVisualizer()
 
 void FDataflowConstructionScene::ResetWireframeMeshElementsVisualizer()
 {
-	if (WireframeDraw)
+	for (FRenderWireElement Elem : WireframeElements)
 	{
-		WireframeDraw->Disconnect();
+		Elem.Value->Disconnect();
 	}
-	WireframeDraw = nullptr;
+	WireframeElements.Empty();
 }
 
 void FDataflowConstructionScene::UpdateWireframeMeshElementsVisualizer()
@@ -295,9 +308,9 @@ void FDataflowConstructionScene::UpdateWireframeMeshElementsVisualizer()
 
 bool FDataflowConstructionScene::HasRenderableGeometry()
 {
-	for (auto& DynamicMeshComponent : DynamicMeshComponents)
+	for (FRenderElement RenderElement : DynamicMeshComponents)
 	{
-		if (DynamicMeshComponent->GetMesh()->TriangleCount() > 0)
+		if (RenderElement.Value->GetMesh()->TriangleCount() > 0)
 		{
 			return true;
 		}
@@ -334,10 +347,10 @@ void FDataflowConstructionScene::UpdateConstructionScene()
     {
     	USelection* SelectedComponents = DataflowModeManager->GetSelectedComponents();
     	SelectedComponents->DeselectAll();
-    	for (const TObjectPtr<UDynamicMeshComponent>& DynamicMeshComponent : DynamicMeshComponents)
+    	for (FRenderElement RenderElement : DynamicMeshComponents)
     	{
-    		SelectedComponents->Select(DynamicMeshComponent);
-    		DynamicMeshComponent->PushSelectionToProxy();
+    		SelectedComponents->Select(RenderElement.Value);
+			RenderElement.Value->PushSelectionToProxy();
     	}
     }
 	DataflowContent->SetIsDirty(false);
diff --git a/Engine/Plugins/Experimental/Dataflow/Source/DataflowEditor/Private/Dataflow/DataflowToolTarget.cpp b/Engine/Plugins/Experimental/Dataflow/Source/DataflowEditor/Private/Dataflow/DataflowToolTarget.cpp
--- a/Engine/Plugins/Experimental/Dataflow/Source/DataflowEditor/Private/Dataflow/DataflowToolTarget.cpp
+++ b/Engine/Plugins/Experimental/Dataflow/Source/DataflowEditor/Private/Dataflow/DataflowToolTarget.cpp
@@ -105,7 +105,7 @@ FDynamicMesh3 UDataflowReadOnlyToolTarget::GetDynamicMesh()
 	FDynamicMesh3 DynamicMesh;
 	if(IsValid())
 	{
-		Dataflow::Conversion::DataflowToDynamicMesh(Context, Asset, Dataflow, DynamicMesh);
+		// @todo(dataflow)
 	}
 	return DynamicMesh;
 }
@@ -147,7 +147,7 @@ void UDataflowToolTarget::CommitDynamicMesh(const FDynamicMesh3& DynamicMesh, co
 {
 	if(IsValid())
 	{
-		Dataflow::Conversion::DynamicMeshToDataflow(DynamicMesh, Dataflow);
+		// @todo(dataflow)
 	}
 }
 
@@ -161,7 +161,7 @@ bool UDataflowReadOnlyToolTargetFactory::CanBuildTarget(UObject* SourceObject, c
 	// If you want to make the tool target work with some subclass of UDataflow,
 	// just add another factory that allows that class specifically(but make sure that
 	// GetMeshDescription and such work properly)
-	const TObjectPtr<UDataflowBaseContent> BaseContent = CastChecked<UDataflowBaseContent>(SourceObject);
+	const TObjectPtr<UDataflowBaseContent> BaseContent = Cast<UDataflowBaseContent>(SourceObject);
 	if(BaseContent)
 	{
 		const UDataflow* Dataflow = BaseContent->GetDataflowAsset();
@@ -203,7 +203,7 @@ bool UDataflowToolTargetFactory::CanBuildTarget(UObject* SourceObject, const FTo
 	// just add another factory that allows that class specifically(but make sure that
 	// GetMeshDescription and such work properly)
 
-	const TObjectPtr<UDataflowBaseContent> BaseContent = CastChecked<UDataflowBaseContent>(SourceObject);
+	const TObjectPtr<UDataflowBaseContent> BaseContent = Cast<UDataflowBaseContent>(SourceObject);
 	if(BaseContent)
 	{
 		const UDataflow* Dataflow = BaseContent->GetDataflowAsset();
diff --git a/Engine/Plugins/Experimental/Dataflow/Source/DataflowEditor/Private/DataflowEditorTools/DataflowEditorWeightMapPaintTool.cpp b/Engine/Plugins/Experimental/Dataflow/Source/DataflowEditor/Private/DataflowEditorTools/DataflowEditorWeightMapPaintTool.cpp
--- a/Engine/Plugins/Experimental/Dataflow/Source/DataflowEditor/Private/DataflowEditorTools/DataflowEditorWeightMapPaintTool.cpp
+++ b/Engine/Plugins/Experimental/Dataflow/Source/DataflowEditor/Private/DataflowEditorTools/DataflowEditorWeightMapPaintTool.cpp
@@ -13,6 +13,7 @@
 #include "DataflowEditorTools/DataflowEditorWeightMapPaintBrushOps.h"
 #include "Dataflow/DataflowContent.h"
 #include "Dataflow/DataflowEdNode.h"
+#include "Dataflow/DataflowEditorCollectionComponent.h"
 #include "Dataflow/DataflowObject.h"
 #include "Dataflow/DataflowCollectionAddScalarVertexPropertyNode.h"
 #include "Dataflow/DataflowSNode.h"
@@ -94,13 +95,22 @@ bool UDataflowEditorWeightMapPaintToolBuilder::CanBuildTool(const FToolBuilderSt
 
 	if (UMeshSurfacePointMeshEditingToolBuilder::CanBuildTool(SceneState))
 	{
-		if (UDataflowBaseContent* ContextObject = SceneState.ToolManager->GetContextObjectStore()->FindContext<UDataflowBaseContent>())
+		if (SceneState.SelectedComponents.Num() == 1)
 		{
-			if (const TSharedPtr<Dataflow::FEngineContext> EvaluationContext = ContextObject->GetDataflowContext())
+			if (TObjectPtr<UDataflowEditorCollectionComponent> Component = Cast<UDataflowEditorCollectionComponent>(SceneState.SelectedComponents[0]))
 			{
-				if (const FDataflowNode* PrimarySelection = ContextObject->GetPrimarySelectedNodeOfType<FDataflowCollectionAddScalarVertexPropertyNode>())
+				if (UDataflowBaseContent* ContextObject = SceneState.ToolManager->GetContextObjectStore()->FindContext<UDataflowBaseContent>())
 				{
-					return HasRenderableCollection(PrimarySelection, EvaluationContext);
+					if (ContextObject->GetPrimarySelectedNode() == Component->Node)
+					{
+						if (const TSharedPtr<Dataflow::FEngineContext> EvaluationContext = ContextObject->GetDataflowContext())
+						{
+							if (const FDataflowNode* PrimarySelection = ContextObject->GetPrimarySelectedNodeOfType<FDataflowCollectionAddScalarVertexPropertyNode>())
+							{
+								return HasRenderableCollection(PrimarySelection, EvaluationContext);
+							}
+						}
+					}
 				}
 			}
 		}
@@ -1675,7 +1685,11 @@ void UDataflowEditorWeightMapPaintTool::UpdateSelectedNode()
 
 	if (bHaveDynamicMeshToWeightConversion)
 	{
-		WeightMapNodeToUpdate->VertexWeights.Init(0.f, WeightToDynamicMesh.Num());
+		if (WeightToDynamicMesh.Num() != WeightMapNodeToUpdate->VertexWeights.Num())
+		{
+			WeightMapNodeToUpdate->VertexWeights.Init(0.f, WeightToDynamicMesh.Num());
+		}
+
 		for (int32 DynamicMeshIdx = 0; DynamicMeshIdx < CurrentWeights.Num(); ++DynamicMeshIdx)
 		{
 			WeightMapNodeToUpdate->VertexWeights[DynamicMeshToWeight[DynamicMeshIdx]] = CurrentWeights[DynamicMeshIdx];
diff --git a/Engine/Plugins/Experimental/Dataflow/Source/DataflowEditor/Public/Dataflow/DataflowPreviewScene.h b/Engine/Plugins/Experimental/Dataflow/Source/DataflowEditor/Public/Dataflow/DataflowPreviewScene.h
--- a/Engine/Plugins/Experimental/Dataflow/Source/DataflowEditor/Public/Dataflow/DataflowPreviewScene.h
+++ b/Engine/Plugins/Experimental/Dataflow/Source/DataflowEditor/Public/Dataflow/DataflowPreviewScene.h
@@ -15,6 +15,7 @@ class FAssetEditorModeManager;
 class ADataflowActor;
 class UMeshElementsVisualizer;
 class UInteractiveToolPropertySet;
+class UDataflowEdNode;
 
 namespace UE
 {
@@ -80,6 +81,9 @@ protected:
  */
 class DATAFLOWEDITOR_API FDataflowConstructionScene : public FDataflowPreviewScene
 {
+	/*(const void* ObjectPointer, const int32 ObjectIndex)*/
+	typedef TPair<TObjectPtr<const UDataflowEdNode>, int32> FDataflowRenderKey;
+
 public:
 
 	FDataflowConstructionScene(FPreviewScene::ConstructionValues ConstructionValues, TObjectPtr<UDataflowBaseContent> InEditorContent);
@@ -94,9 +98,6 @@ public:
 	/** Check if the preview scene can run simulation */
 	virtual bool CanRunSimulation() const override {return false;}
 	
-	/**  Get all the dynamic mesh components */
-	TArray<TObjectPtr<UDynamicMeshComponent>> GetDynamicMeshComponents() { return DynamicMeshComponents; }
-	
 	/** Is there ANYTHING to render */
 	bool HasRenderableGeometry();
 
@@ -125,13 +126,15 @@ private:
 	
 	/** Add a dynamic mesh component to the scene */
 	TObjectPtr<UDynamicMeshComponent>& AddDynamicMeshComponent(
-		UE::Geometry::FDynamicMesh3&& DynamicMesh, const TArray<UMaterialInterface*>& MaterialSet);
+		FDataflowRenderKey InKey, UE::Geometry::FDynamicMesh3&& DynamicMesh,  const TArray<UMaterialInterface*>& MaterialSet);
 	
 	/** Dynamic mesh components generated by the tools for edit/manipulation*/
-	TArray<TObjectPtr<UDynamicMeshComponent>> DynamicMeshComponents;
+	typedef TPair<FDataflowRenderKey, TObjectPtr<UDynamicMeshComponent>> FRenderElement;
+	TMap<FDataflowRenderKey, TObjectPtr<UDynamicMeshComponent>> DynamicMeshComponents;
 
 	/** Rest space wireframe. They have to get ticked to be able to respond to setting changes. [ticked] */
-	TObjectPtr<UMeshElementsVisualizer> WireframeDraw = nullptr;
+	typedef TPair<TObjectPtr<UDynamicMeshComponent>, TObjectPtr<UMeshElementsVisualizer>> FRenderWireElement;
+	TMap<TObjectPtr<UDynamicMeshComponent>,TObjectPtr<UMeshElementsVisualizer> > WireframeElements;
 
 	/** Property objects (visible or not) that get ticked. [ticked] */
 	TArray<TObjectPtr<UInteractiveToolPropertySet>> PropertyObjectsToTick;
diff --git a/Engine/Plugins/Experimental/Dataflow/Source/DataflowEnginePlugin/Private/Dataflow/DataflowEngineRendering.cpp b/Engine/Plugins/Experimental/Dataflow/Source/DataflowEnginePlugin/Private/Dataflow/DataflowEngineRendering.cpp
--- a/Engine/Plugins/Experimental/Dataflow/Source/DataflowEnginePlugin/Private/Dataflow/DataflowEngineRendering.cpp
+++ b/Engine/Plugins/Experimental/Dataflow/Source/DataflowEnginePlugin/Private/Dataflow/DataflowEngineRendering.cpp
@@ -14,6 +14,211 @@
 #include "UDynamicMesh.h"
 namespace Dataflow
 {
+	void RenderBasicGeometryCollection(GeometryCollection::Facades::FRenderingFacade& RenderCollection, const Dataflow::FGraphRenderingState& State)
+	{
+		FManagedArrayCollection Default;
+		FName PrimaryOutput = State.GetRenderOutputs()[0]; // "Collection"
+		const FManagedArrayCollection& Collection = State.GetValue<FManagedArrayCollection>(PrimaryOutput, Default);
+
+		const TManagedArray<int32>& BoneIndex = Collection.GetAttribute<int32>("BoneMap", FGeometryCollection::VerticesGroup);
+		const TManagedArray<int32>& Parents = Collection.GetAttribute<int32>(FTransformCollection::ParentAttribute, FTransformCollection::TransformGroup);
+		const TManagedArray<FTransform3f>& Transforms = Collection.GetAttribute<FTransform3f>(FTransformCollection::TransformAttribute, FTransformCollection::TransformGroup);
+
+		TArray<FTransform> M;
+		GeometryCollectionAlgo::GlobalMatrices(Transforms, Parents, M);
+
+		// If Collection has "ExplodedVector" attribute then use it to modify the global matrices (ExplodedView node creates it)
+		GeometryCollection::Facades::FCollectionExplodedVectorFacade ExplodedViewFacade(Collection);
+		ExplodedViewFacade.UpdateGlobalMatricesWithExplodedVectors(M);
+
+		auto ToD = [](FVector3f V) { return FVector3d(V.X, V.Y, V.Z); };
+		auto ToF = [](FVector3d V) { return FVector3f(V.X, V.Y, V.Z); };
+
+
+		const TManagedArray<FVector3f>& Vertex = Collection.GetAttribute<FVector3f>("Vertex", FGeometryCollection::VerticesGroup);
+		const TManagedArray<FIntVector>& Faces = Collection.GetAttribute<FIntVector>("Indices", FGeometryCollection::FacesGroup);
+		const TManagedArray<bool>* FaceVisible = Collection.FindAttribute<bool>("Visible", FGeometryCollection::FacesGroup);
+
+		TArray<FVector3f> Vertices; Vertices.AddUninitialized(Vertex.Num());
+		TArray<FIntVector> Tris; Tris.AddUninitialized(Faces.Num());
+		TArray<bool> Visited; Visited.Init(false, Vertices.Num());
+
+		int32 Tdx = 0;
+		for (int32 FaceIdx = 0; FaceIdx < Faces.Num(); ++FaceIdx)
+		{
+			if (FaceVisible && !(*FaceVisible)[FaceIdx]) continue;
+
+			const FIntVector& Face = Faces[FaceIdx];
+
+			FIntVector Tri = FIntVector(Face[0], Face[1], Face[2]);
+			FTransform Ms[3] = { M[BoneIndex[Tri[0]]], M[BoneIndex[Tri[1]]], M[BoneIndex[Tri[2]]] };
+
+			Tris[Tdx++] = Tri;
+			if (!Visited[Tri[0]]) Vertices[Tri[0]] = ToF(Ms[0].TransformPosition(ToD(Vertex[Tri[0]])));
+			if (!Visited[Tri[1]]) Vertices[Tri[1]] = ToF(Ms[1].TransformPosition(ToD(Vertex[Tri[1]])));
+			if (!Visited[Tri[2]]) Vertices[Tri[2]] = ToF(Ms[2].TransformPosition(ToD(Vertex[Tri[2]])));
+
+			Visited[Tri[0]] = true; Visited[Tri[1]] = true; Visited[Tri[2]] = true;
+		}
+
+		Tris.SetNum(Tdx);
+
+		// Maybe these buffers should be shrunk, but there are unused vertices in the buffer. 
+		for (int i = 0; i < Visited.Num(); i++) if (!Visited[i]) Vertices[i] = FVector3f(0);
+
+		// Copy VertexNormals from the Collection if exists otherwise compute and set it
+		TArray<FVector3f> VertexNormals; VertexNormals.AddUninitialized(Vertex.Num());
+		if (const TManagedArray<FVector3f>* VertexNormal = Collection.FindAttribute<FVector3f>("Normal", FGeometryCollection::VerticesGroup))
+		{
+			for (int32 VertexIdx = 0; VertexIdx < VertexNormals.Num(); ++VertexIdx)
+			{
+				VertexNormals[VertexIdx] = (*VertexNormal)[VertexIdx];
+			}
+		}
+		else
+		{
+			for (int32 VertexIdx = 0; VertexIdx < VertexNormals.Num(); ++VertexIdx)
+			{
+				// TODO: Compute the normal
+				VertexNormals[VertexIdx] = FVector3f(0.f);
+			}
+		}
+
+		// Copy VertexColors from the Collection if exists otherwise set it to IDataflowEnginePlugin::SurfaceColor
+		TArray<FLinearColor> VertexColors; VertexColors.AddUninitialized(Vertex.Num());
+		if (const TManagedArray<FLinearColor>* VertexColorManagedArray = Collection.FindAttribute<FLinearColor>("Color", FGeometryCollection::VerticesGroup))
+		{
+			for (int32 VertexIdx = 0; VertexIdx < VertexColors.Num(); ++VertexIdx)
+			{
+				VertexColors[VertexIdx] = (*VertexColorManagedArray)[VertexIdx];
+			}
+		}
+		else
+		{
+			for (int32 VertexIdx = 0; VertexIdx < VertexColors.Num(); ++VertexIdx)
+			{
+				VertexColors[VertexIdx] = FLinearColor(IDataflowEnginePlugin::SurfaceColor);
+			}
+		}
+
+		// Set the data on the RenderCollection
+		int32 GeometryIndex = RenderCollection.StartGeometryGroup(State.GetGuid().ToString());
+		RenderCollection.AddSurface(MoveTemp(Vertices), MoveTemp(Tris), MoveTemp(VertexNormals), MoveTemp(VertexColors));
+		RenderCollection.EndGeometryGroup(GeometryIndex);
+
+	}
+
+	void RenderMeshIndexedGeometryCollection(GeometryCollection::Facades::FRenderingFacade& RenderCollection, const Dataflow::FGraphRenderingState& State)
+	{
+		auto ToD = [](FVector3f V) { return FVector3d(V.X, V.Y, V.Z); };
+		auto ToF = [](FVector3d V) { return FVector3f(V.X, V.Y, V.Z); };
+
+		FManagedArrayCollection Default;
+		FName PrimaryOutput = State.GetRenderOutputs()[0]; // "Collection"
+		const FManagedArrayCollection& Collection = State.GetValue<FManagedArrayCollection>(PrimaryOutput, Default);
+
+		const TManagedArray<int32>& BoneIndex = Collection.GetAttribute<int32>("BoneMap", FGeometryCollection::VerticesGroup);
+		const TManagedArray<int32>& Parents = Collection.GetAttribute<int32>(FTransformCollection::ParentAttribute, FTransformCollection::TransformGroup);
+		const TManagedArray<FTransform3f>& Transforms = Collection.GetAttribute<FTransform3f>(FTransformCollection::TransformAttribute, FTransformCollection::TransformGroup);
+		const TManagedArray<FVector3f>& Vertex = Collection.GetAttribute<FVector3f>("Vertex", FGeometryCollection::VerticesGroup);
+		const TManagedArray<FIntVector>& Faces = Collection.GetAttribute<FIntVector>("Indices", FGeometryCollection::FacesGroup);
+		const TManagedArray<bool>* FaceVisible = Collection.FindAttribute<bool>("Visible", FGeometryCollection::FacesGroup);
+
+		const TManagedArray<int32>& VertexStart = Collection.GetAttribute<int32>("VertexStart", FGeometryCollection::GeometryGroup);
+		const TManagedArray<int32>& VertexCount = Collection.GetAttribute<int32>("VertexCount", FGeometryCollection::GeometryGroup);
+		const TManagedArray<int32>& FacesStart = Collection.GetAttribute<int32>("FaceStart", FGeometryCollection::GeometryGroup);
+		const TManagedArray<int32>& FacesCount = Collection.GetAttribute<int32>("FaceCount", FGeometryCollection::GeometryGroup);
+
+		TArray<FTransform> M;
+		GeometryCollectionAlgo::GlobalMatrices(Transforms, Parents, M);
+		GeometryCollection::Facades::FCollectionExplodedVectorFacade ExplodedViewFacade(Collection);
+		ExplodedViewFacade.UpdateGlobalMatricesWithExplodedVectors(M);
+
+		for (int Gdx = 0; Gdx < Collection.NumElements(FGeometryCollection::GeometryGroup); Gdx++)
+		{
+			TArray<FVector3f> Vertices; Vertices.AddUninitialized(VertexCount[Gdx]);
+			TArray<FIntVector> Tris; Tris.AddUninitialized(FacesCount[Gdx]);
+			TArray<bool> Visited; Visited.Init(false, VertexCount[Gdx]);
+
+			int32 Tdx = 0;
+			int32 LastFaceIndex = FacesStart[Gdx] + FacesCount[Gdx];
+			for (int32 FaceIdx = FacesStart[Gdx]; FaceIdx < LastFaceIndex; ++FaceIdx)
+			{
+				if (FaceVisible && !(*FaceVisible)[FaceIdx]) continue;
+
+				const FIntVector& Face = Faces[FaceIdx];
+
+				FIntVector Tri = FIntVector(Face[0], Face[1], Face[2]);
+				FTransform Ms[3] = { M[BoneIndex[Tri[0]]], M[BoneIndex[Tri[1]]], M[BoneIndex[Tri[2]]] };
+				FIntVector MovedTri = FIntVector(Face[0]- VertexStart[Gdx], Face[1] - VertexStart[Gdx], Face[2] - VertexStart[Gdx]);
+
+				Tris[Tdx++] = MovedTri;
+				if (!Visited[MovedTri[0]]) Vertices[Tri[0] - VertexStart[Gdx]] = ToF(Ms[0].TransformPosition(ToD(Vertex[Tri[0]])));
+				if (!Visited[MovedTri[1]]) Vertices[Tri[1] - VertexStart[Gdx]] = ToF(Ms[1].TransformPosition(ToD(Vertex[Tri[1]])));
+				if (!Visited[MovedTri[2]]) Vertices[Tri[2] - VertexStart[Gdx]] = ToF(Ms[2].TransformPosition(ToD(Vertex[Tri[2]])));
+
+				Visited[MovedTri[0]] = true; Visited[MovedTri[1]] = true; Visited[MovedTri[2]] = true;
+			}
+
+			Tris.SetNum(Tdx);
+
+			// move the unused points too. Need to keep them for vertex alignment with ediior tools. 
+			for (int i = 0; i < Visited.Num(); i++)
+			{
+				if (!Visited[i])
+				{
+					Vertices[i] = ToF(M[BoneIndex[i+ VertexStart[Gdx]]].TransformPosition(ToD(Vertex[i+ VertexStart[Gdx]])));
+				}
+			}
+
+			// Copy VertexNormals from the Collection if exists otherwise compute and set it
+			ensure(VertexCount[Gdx] == Vertices.Num());
+			TArray<FVector3f> VertexNormals; VertexNormals.AddUninitialized(Vertices.Num());
+			if (const TManagedArray<FVector3f>* VertexNormal = Collection.FindAttribute<FVector3f>("Normal", FGeometryCollection::VerticesGroup))
+			{
+				int32 LastVertIndex = VertexStart[Gdx] + VertexCount[Gdx];
+				for (int32 VertexIdx = VertexStart[Gdx], SrcVertexIdx = 0; VertexIdx < LastVertIndex; ++VertexIdx, ++SrcVertexIdx)
+				{
+					VertexNormals[SrcVertexIdx] = (*VertexNormal)[VertexIdx];
+				}
+			}
+			else
+			{
+				for (int32 VertexIdx = 0; VertexIdx < Vertices.Num(); ++VertexIdx)
+				{
+					// TODO: Compute the normal
+					VertexNormals[VertexIdx] = FVector3f(0.f);
+				}
+			}
+
+			// Copy VertexColors from the Collection if exists otherwise set it to IDataflowEnginePlugin::SurfaceColor
+			TArray<FLinearColor> VertexColors; VertexColors.AddUninitialized(Vertices.Num());
+			if (const TManagedArray<FLinearColor>* VertexColorManagedArray = Collection.FindAttribute<FLinearColor>("Color", FGeometryCollection::VerticesGroup))
+			{
+				int32 LastVertIndex = VertexStart[Gdx] + VertexCount[Gdx];
+				for (int32 VertexIdx = VertexStart[Gdx], SrcVertexIdx = 0; VertexIdx < LastVertIndex; ++VertexIdx, ++SrcVertexIdx)
+				{
+					VertexColors[SrcVertexIdx] = (*VertexColorManagedArray)[VertexIdx];
+				}
+			}
+			else
+			{
+				for (int32 VertexIdx = 0; VertexIdx < VertexColors.Num(); ++VertexIdx)
+				{
+					VertexColors[VertexIdx] = FLinearColor(IDataflowEnginePlugin::SurfaceColor);
+				}
+			}
+
+			// Set the data on the RenderCollection
+			if (Vertices.Num() && Tris.Num())
+			{
+				FString GeometryName = State.GetGuid().ToString(); GeometryName.AppendChar('.').AppendInt(Gdx);
+				int32 GeometryIndex = RenderCollection.StartGeometryGroup(GeometryName);
+				RenderCollection.AddSurface(MoveTemp(Vertices), MoveTemp(Tris), MoveTemp(VertexNormals), MoveTemp(VertexColors));
+				RenderCollection.EndGeometryGroup(GeometryIndex);
+			}
+		}
+	}
 
 	void RenderingCallbacks()
 	{
@@ -31,109 +236,34 @@ namespace Dataflow
 			{
 				if (State.GetRenderOutputs().Num())
 				{
-					FName PrimaryOutput = State.GetRenderOutputs()[0]; // "Collection"
-
 					FManagedArrayCollection Default;
+					FName PrimaryOutput = State.GetRenderOutputs()[0]; // "Collection"
 					const FManagedArrayCollection& Collection = State.GetValue<FManagedArrayCollection>(PrimaryOutput, Default);
+
 					const bool bFoundIndices = Collection.FindAttributeTyped<FIntVector>("Indices", FGeometryCollection::FacesGroup) != nullptr;
 					const bool bFoundVertices = Collection.FindAttributeTyped<FVector3f>("Vertex", FGeometryCollection::VerticesGroup) != nullptr;
 					const bool bFoundTransforms = Collection.FindAttributeTyped<FTransform3f>(FTransformCollection::TransformAttribute, FTransformCollection::TransformGroup) != nullptr;
 					const bool bFoundBoneMap = Collection.FindAttributeTyped<int32>("BoneMap", FGeometryCollection::VerticesGroup) != nullptr;
 					const bool bFoundParents = Collection.FindAttributeTyped<int32>(FTransformCollection::ParentAttribute, FTransformCollection::TransformGroup) != nullptr;
 					UE_LOG(LogTemp, Warning, TEXT("Render GC with found params = %d %d %d %d %d"), bFoundIndices, bFoundVertices, bFoundTransforms, bFoundBoneMap, bFoundParents);
-					
-					if (Collection.FindAttributeTyped<FIntVector>("Indices", FGeometryCollection::FacesGroup)
-						&& Collection.FindAttributeTyped<FVector3f>("Vertex", FGeometryCollection::VerticesGroup)
-						&& Collection.FindAttributeTyped<FTransform3f>(FTransformCollection::TransformAttribute, FTransformCollection::TransformGroup)
-						&& Collection.FindAttributeTyped<int32>("BoneMap", FGeometryCollection::VerticesGroup)
-						&& Collection.FindAttributeTyped<int32>(FTransformCollection::ParentAttribute, FTransformCollection::TransformGroup))
+					bool bFoundRenderData = bFoundIndices && bFoundVertices && bFoundTransforms && bFoundBoneMap && bFoundParents 
+						&& Collection.NumElements(FTransformCollection::TransformGroup) > 0;
+
+					const bool bFoundVertexStart = Collection.FindAttributeTyped<int32>("VertexStart", FGeometryCollection::GeometryGroup) != nullptr;
+					const bool bFoundVertexCount = Collection.FindAttributeTyped<int32>("VertexCount", FGeometryCollection::GeometryGroup) != nullptr;
+					const bool bFoundFaceStart = Collection.FindAttributeTyped<int32>("FaceStart", FGeometryCollection::GeometryGroup) != nullptr;
+					const bool bFoundFaceCount = Collection.FindAttributeTyped<int32>("FaceCount", FGeometryCollection::GeometryGroup) != nullptr;
+					UE_LOG(LogTemp, Warning, TEXT("Render GC with found mesh group params = %d %d %d %d"), bFoundVertexStart, bFoundVertexCount, bFoundFaceStart, bFoundFaceCount);
+					bool bFoundGeometryAttributes = bFoundVertexStart && bFoundVertexCount && bFoundFaceStart && bFoundFaceCount
+						&& Collection.NumElements(FGeometryCollection::GeometryGroup) > 0;
+
+					if (bFoundRenderData && bFoundGeometryAttributes )
 					{
-						if (Collection.NumElements(FTransformCollection::TransformGroup) > 0)
-						{
-							const TManagedArray<int32>& BoneIndex = Collection.GetAttribute<int32>("BoneMap", FGeometryCollection::VerticesGroup);
-							const TManagedArray<int32>& Parents = Collection.GetAttribute<int32>(FTransformCollection::ParentAttribute, FTransformCollection::TransformGroup);
-							const TManagedArray<FTransform3f>& Transforms = Collection.GetAttribute<FTransform3f>(FTransformCollection::TransformAttribute, FTransformCollection::TransformGroup);
-
-							TArray<FTransform> M;
-							GeometryCollectionAlgo::GlobalMatrices(Transforms, Parents, M);
-
-							// If Collection has "ExplodedVector" attribute then use it to modify the global matrices (ExplodedView node creates it)
-							GeometryCollection::Facades::FCollectionExplodedVectorFacade ExplodedViewFacade(Collection);
-							ExplodedViewFacade.UpdateGlobalMatricesWithExplodedVectors(M);
-
-							auto ToD = [](FVector3f V) { return FVector3d(V.X, V.Y, V.Z); };
-							auto ToF = [](FVector3d V) { return FVector3f(V.X, V.Y, V.Z); };
-
-
-							const TManagedArray<FVector3f>& Vertex = Collection.GetAttribute<FVector3f>("Vertex", FGeometryCollection::VerticesGroup);
-							const TManagedArray<FIntVector>& Faces = Collection.GetAttribute<FIntVector>("Indices", FGeometryCollection::FacesGroup);
-							const TManagedArray<bool>* FaceVisible = Collection.FindAttribute<bool>("Visible", FGeometryCollection::FacesGroup);
-
-							TArray<FVector3f> Vertices; Vertices.AddUninitialized(Vertex.Num());
-							TArray<FIntVector> Tris; Tris.AddUninitialized(Faces.Num());
-							TArray<bool> Visited; Visited.Init(false, Vertices.Num());
-
-							int32 Tdx = 0;
-							for (int32 FaceIdx = 0; FaceIdx < Faces.Num(); ++FaceIdx)
-							{				
-								if (FaceVisible && !(*FaceVisible)[FaceIdx]) continue;
-
-								const FIntVector& Face = Faces[FaceIdx];
-
-								FIntVector Tri = FIntVector(Face[0], Face[1], Face[2]);
-								FTransform Ms[3] = { M[BoneIndex[Tri[0]]], M[BoneIndex[Tri[1]]], M[BoneIndex[Tri[2]]] };
-
-								Tris[Tdx++] = Tri;
-								if (!Visited[Tri[0]]) Vertices[Tri[0]] = ToF(Ms[0].TransformPosition(ToD(Vertex[Tri[0]])));
-								if (!Visited[Tri[1]]) Vertices[Tri[1]] = ToF(Ms[1].TransformPosition(ToD(Vertex[Tri[1]])));
-								if (!Visited[Tri[2]]) Vertices[Tri[2]] = ToF(Ms[2].TransformPosition(ToD(Vertex[Tri[2]])));
-
-								Visited[Tri[0]] = true; Visited[Tri[1]] = true; Visited[Tri[2]] = true;
-							}
-
-							Tris.SetNum(Tdx);
-
-							// Maybe these buffers should be shrunk, but there are unused vertices in the buffer. 
-							for (int i = 0; i < Visited.Num(); i++) if (!Visited[i]) Vertices[i] = FVector3f(0);
-
-							// Copy VertexNormals from the Collection if exists otherwise compute and set it
-							TArray<FVector3f> VertexNormals; VertexNormals.AddUninitialized(Vertex.Num());
-							if (const TManagedArray<FVector3f>* VertexNormal = Collection.FindAttribute<FVector3f>("Normal", FGeometryCollection::VerticesGroup))
-							{
-								for (int32 VertexIdx = 0; VertexIdx < VertexNormals.Num(); ++VertexIdx)
-								{
-									VertexNormals[VertexIdx] = (*VertexNormal)[VertexIdx];
-								}
-							}
-							else
-							{
-								for (int32 VertexIdx = 0; VertexIdx < VertexNormals.Num(); ++VertexIdx)
-								{
-									// TODO: Compute the normal
-									VertexNormals[VertexIdx] = FVector3f(0.f);
-								}
-							}
-
-							// Copy VertexColors from the Collection if exists otherwise set it to IDataflowEnginePlugin::SurfaceColor
-							TArray<FLinearColor> VertexColors; VertexColors.AddUninitialized(Vertex.Num());
-							if (const TManagedArray<FLinearColor>* VertexColorManagedArray = Collection.FindAttribute<FLinearColor>("Color", FGeometryCollection::VerticesGroup))
-							{
-								for (int32 VertexIdx = 0; VertexIdx < VertexColors.Num(); ++VertexIdx)
-								{
-									VertexColors[VertexIdx] = (*VertexColorManagedArray)[VertexIdx];
-								}
-							}
-							else
-							{
-								for (int32 VertexIdx = 0; VertexIdx < VertexColors.Num(); ++VertexIdx)
-								{
-									VertexColors[VertexIdx] = FLinearColor(IDataflowEnginePlugin::SurfaceColor);
-								}
-							}
-
-							// Set the data on the RenderCollection
-							RenderCollection.AddSurface(MoveTemp(Vertices), MoveTemp(Tris), MoveTemp(VertexNormals), MoveTemp(VertexColors));
-						}
+						RenderMeshIndexedGeometryCollection(RenderCollection, State);
+					}
+					else if (bFoundRenderData)
+					{
+						RenderBasicGeometryCollection(RenderCollection, State);
 					}
 				}
 			});
@@ -192,7 +322,9 @@ namespace Dataflow
 								VertexColors[VertexIdx] = FLinearColor(IDataflowEnginePlugin::SurfaceColor);
 							}
 
+							int32 GeometryIndex = RenderCollection.StartGeometryGroup(State.GetGuid().ToString());
 							RenderCollection.AddSurface(MoveTemp(Vertices), MoveTemp(Tris), MoveTemp(VertexNormals), MoveTemp(VertexColors));
+							RenderCollection.EndGeometryGroup(GeometryIndex);
 						}
 					}
 				}
@@ -254,7 +386,9 @@ namespace Dataflow
 						VertexColors[VertexIdx] = FLinearColor(IDataflowEnginePlugin::SurfaceColor);
 					}
 
+					int32 GeometryIndex = RenderCollection.StartGeometryGroup(State.GetGuid().ToString());
 					RenderCollection.AddSurface(MoveTemp(Vertices), MoveTemp(Tris), MoveTemp(VertexNormals), MoveTemp(VertexColors));
+					RenderCollection.EndGeometryGroup(GeometryIndex);
 				}
 			});
 
@@ -301,7 +435,10 @@ namespace Dataflow
 							VertexNormals[3*i+2] = TriangleNormal;
 							VertexColors[i] = FLinearColor(IDataflowEnginePlugin::SurfaceColor);
 						}
+
+						int32 GeometryIndex = RenderCollection.StartGeometryGroup(State.GetGuid().ToString());
 						RenderCollection.AddSurface(MoveTemp(Vertices), MoveTemp(Tris), MoveTemp(VertexNormals), MoveTemp(VertexColors));
+						RenderCollection.EndGeometryGroup(GeometryIndex);
 					}
 				}
 			});
diff --git a/Engine/Plugins/Experimental/Dataflow/Source/DataflowNodes/Public/Dataflow/DataflowCollectionAddScalarVertexPropertyNode.h b/Engine/Plugins/Experimental/Dataflow/Source/DataflowNodes/Public/Dataflow/DataflowCollectionAddScalarVertexPropertyNode.h
--- a/Engine/Plugins/Experimental/Dataflow/Source/DataflowNodes/Public/Dataflow/DataflowCollectionAddScalarVertexPropertyNode.h
+++ b/Engine/Plugins/Experimental/Dataflow/Source/DataflowNodes/Public/Dataflow/DataflowCollectionAddScalarVertexPropertyNode.h
@@ -3,6 +3,7 @@
 #pragma once 
 
 #include "Dataflow/DataflowNode.h"
+#include "GeometryCollection/GeometryCollection.h"
 #include "GeometryCollection/ManagedArrayCollection.h"
 
 #include "DataflowCollectionAddScalarVertexPropertyNode.generated.h"
@@ -13,6 +14,8 @@ struct DATAFLOWNODES_API FDataflowCollectionAddScalarVertexPropertyNode : public
 {
 	GENERATED_USTRUCT_BODY()
 	DATAFLOW_NODE_DEFINE_INTERNAL(FDataflowCollectionAddScalarVertexPropertyNode, "AddScalarVertexProperty", "Collection", "Add a saved scalar property to a collection")
+	DATAFLOW_NODE_RENDER_TYPE(FGeometryCollection::StaticType(), "Collection")
+
 
 public:
 
diff --git a/Engine/Source/Runtime/Experimental/Dataflow/Engine/Private/Dataflow/DataflowEdNode.cpp b/Engine/Source/Runtime/Experimental/Dataflow/Engine/Private/Dataflow/DataflowEdNode.cpp
--- a/Engine/Source/Runtime/Experimental/Dataflow/Engine/Private/Dataflow/DataflowEdNode.cpp
+++ b/Engine/Source/Runtime/Experimental/Dataflow/Engine/Private/Dataflow/DataflowEdNode.cpp
@@ -501,13 +501,11 @@ bool UDataflowEdNode::Render(GeometryCollection::Facades::FRenderingFacade& Rend
 			{
 				if (GetRenderParameters().Num())
 				{
-					int32 GeometryIndex = RenderData.StartGeometryGroup(GetDataflowNodeGuid().ToString());
 					for (Dataflow::FRenderingParameter& Parameter : GetRenderParameters())
 					{
-						Factory->RenderNodeOutput(RenderData, {NodeTarget.Get(), Parameter, *Context});
+						Factory->RenderNodeOutput(RenderData, { GetDataflowNodeGuid(), NodeTarget.Get(), Parameter, *Context});
 						bNeedsRefresh = true;
 					}
-					RenderData.EndGeometryGroup(GeometryIndex);
 				}
 			}
 		}
diff --git a/Engine/Source/Runtime/Experimental/Dataflow/Engine/Public/Dataflow/DataflowRenderingFactory.h b/Engine/Source/Runtime/Experimental/Dataflow/Engine/Public/Dataflow/DataflowRenderingFactory.h
--- a/Engine/Source/Runtime/Experimental/Dataflow/Engine/Public/Dataflow/DataflowRenderingFactory.h
+++ b/Engine/Source/Runtime/Experimental/Dataflow/Engine/Public/Dataflow/DataflowRenderingFactory.h
@@ -14,13 +14,15 @@ namespace Dataflow
 	class FContext;
 
 	struct FGraphRenderingState {
-		FGraphRenderingState(const FDataflowNode* InNode, const FRenderingParameter& InParameters, Dataflow::FContext& InContext)
-			: Node(InNode)
+		FGraphRenderingState(const FGuid InGuid, const FDataflowNode* InNode, const FRenderingParameter& InParameters, Dataflow::FContext& InContext)
+			: NodeGuid(InGuid)
+			, Node(InNode)
 			, RenderType(InParameters.Type)
 			, RenderOutputs(InParameters.Outputs)
 			, Context(InContext)
 		{}
 
+		const FGuid& GetGuid() const { return NodeGuid; }
 		const FName& GetRenderType() const { return RenderType; }
 		const TArray<FName>& GetRenderOutputs() const { return RenderOutputs; }
 
@@ -38,6 +40,7 @@ namespace Dataflow
 		}
 
 	private:
+		const FGuid NodeGuid;
 		const FDataflowNode* Node = nullptr;
 
 		FName RenderType;
