Dataflow : Refresh viewport on tools exit.

#changelist validated
#virtualized

diff --git a/Engine/Plugins/Experimental/Dataflow/Source/DataflowEditor/Private/Dataflow/DataflowEditorModeToolkit.cpp b/Engine/Plugins/Experimental/Dataflow/Source/DataflowEditor/Private/Dataflow/DataflowEditorModeToolkit.cpp
--- a/Engine/Plugins/Experimental/Dataflow/Source/DataflowEditor/Private/Dataflow/DataflowEditorModeToolkit.cpp
+++ b/Engine/Plugins/Experimental/Dataflow/Source/DataflowEditor/Private/Dataflow/DataflowEditorModeToolkit.cpp
@@ -242,6 +242,19 @@ void FDataflowEditorModeToolkit::OnToolStarted(UInteractiveToolManager* Manager,
 
 void FDataflowEditorModeToolkit::OnToolEnded(UInteractiveToolManager* Manager, UInteractiveTool* Tool)
 {
+	auto GetConstructionScene = [&]() 
+	{
+		if (OwningEditorMode.IsValid(/*bEvenIfPendingKill*/ false))
+		{
+			UEdMode* const Mode = OwningEditorMode.Get();
+			if (UDataflowEditorMode* DataflowEdMode = Cast<UDataflowEditorMode>(Mode))
+			{
+				return DataflowEdMode->GetDataflowConstructionScene();
+			}
+		}
+		return (FDataflowConstructionScene* )nullptr;
+	};
+
 	FModeToolkit::OnToolEnded(Manager, Tool);
 
 	ActiveToolName = FText::GetEmpty();
@@ -258,6 +271,16 @@ void FDataflowEditorModeToolkit::OnToolEnded(UInteractiveToolManager* Manager, U
 		CurTool->OnPropertySetsModified.RemoveAll(this);
 		CurTool->OnPropertyModifiedDirectlyByTool.RemoveAll(this);
 	}
+
+
+	if (FDataflowConstructionScene* ConstructionScene = GetConstructionScene())
+	{
+		if (TObjectPtr<UDataflowBaseContent> Content = ConstructionScene->GetDataflowContent())
+		{
+			Content->SetIsDirty(true);
+		}
+	}
+
 }
 
 void FDataflowEditorModeToolkit::SetRestSpaceViewportWidget(TWeakPtr<SDataflowEditorViewport> InRestSpaceViewportWidget)
@@ -265,4 +288,5 @@ void FDataflowEditorModeToolkit::SetRestSpaceViewportWidget(TWeakPtr<SDataflowEd
 	RestSpaceViewportWidget = InRestSpaceViewportWidget;
 }
 
+
 #undef LOCTEXT_NAMESPACE
diff --git a/Engine/Plugins/Experimental/Dataflow/Source/DataflowEditor/Public/Dataflow/DataflowEditorMode.h b/Engine/Plugins/Experimental/Dataflow/Source/DataflowEditor/Public/Dataflow/DataflowEditorMode.h
--- a/Engine/Plugins/Experimental/Dataflow/Source/DataflowEditor/Public/Dataflow/DataflowEditorMode.h
+++ b/Engine/Plugins/Experimental/Dataflow/Source/DataflowEditor/Public/Dataflow/DataflowEditorMode.h
@@ -61,6 +61,7 @@ public:
 
 	/** Set the data flow preview scene to store the dynamic meshes generated by the tools */
 	void SetDataflowConstructionScene(FDataflowConstructionScene* DataflowPreviewScene);
+	FDataflowConstructionScene* GetDataflowConstructionScene() { return PreviewScene; }
 	const FDataflowConstructionScene* GetDataflowConstructionScene() const { return PreviewScene; }
 
 	/** Set the data flow graph editor to create nodes once the tools have ended*/
