Dataflow : Selection updates

diff --git a/Engine/Plugins/Experimental/Dataflow/Source/DataflowEditor/Private/Dataflow/DataflowEditorToolkit.cpp b/Engine/Plugins/Experimental/Dataflow/Source/DataflowEditor/Private/Dataflow/DataflowEditorToolkit.cpp
--- a/Engine/Plugins/Experimental/Dataflow/Source/DataflowEditor/Private/Dataflow/DataflowEditorToolkit.cpp
+++ b/Engine/Plugins/Experimental/Dataflow/Source/DataflowEditor/Private/Dataflow/DataflowEditorToolkit.cpp
@@ -531,6 +531,15 @@ void FDataflowEditorToolkit::OnNodeSelectionChanged(const TSet<UObject*>& InNewS
 		}
 	};
 
+	auto IsControlDown = [&DataflowEditor = GraphEditor]()
+	{
+		if (DataflowEditor)
+		{
+			return DataflowEditor->IsControlDown();
+		}
+		return false;
+	};
+
 	// Despite this function's name, we might not have actually changed which node is selected
 	bool bPrimarySelectionChanged = false;
 
@@ -601,7 +610,6 @@ void FDataflowEditorToolkit::OnNodeSelectionChanged(const TSet<UObject*>& InNewS
 
 				EditorContent->SetPrimarySelectedNode(PrimarySelection);
 			}
-			EditorContent->SetIsDirty(true);
 		}
 	}
 }
diff --git a/Engine/Plugins/Experimental/Dataflow/Source/DataflowEditor/Private/Dataflow/DataflowEditorViewportClient.cpp b/Engine/Plugins/Experimental/Dataflow/Source/DataflowEditor/Private/Dataflow/DataflowEditorViewportClient.cpp
--- a/Engine/Plugins/Experimental/Dataflow/Source/DataflowEditor/Private/Dataflow/DataflowEditorViewportClient.cpp
+++ b/Engine/Plugins/Experimental/Dataflow/Source/DataflowEditor/Private/Dataflow/DataflowEditorViewportClient.cpp
@@ -1,9 +1,13 @@
 // Copyright Epic Games, Inc. All Rights Reserved.
 #include "Dataflow/DataflowEditorViewportClient.h"
 
+#include "AssetEditorModeManager.h"
 #include "Dataflow/DataflowObject.h"
+#include "Dataflow/DataflowEditorMode.h"
 #include "Dataflow/DataflowEditorToolkit.h"
+#include "Dataflow/DataflowEditorCollectionComponent.h"
 #include "Dataflow/DataflowEngineSceneHitProxies.h"
+#include "Dataflow/DataflowGraphEditor.h"
 #include "Dataflow/DataflowPreviewScene.h"
 #include "EditorModeManager.h"
 #include "PreviewScene.h"
@@ -52,32 +56,72 @@ void FDataflowEditorViewportClient::ProcessClick(FSceneView& View, HHitProxy* Hi
 {
 	Super::ProcessClick(View, HitProxy, Key, Event, HitX, HitY);
 
-	USelection* SelectedComponents = ModeTools->GetSelectedComponents();
+	auto EnableToolForSelectedNode = [&](USelection* SelectedComponents)
+	{
+		if (TSharedPtr<FDataflowEditorToolkit> DataflowEditorToolkit = DataflowEditorToolkitPtr.Pin())
+		{
+			if (PreviewScene && PreviewScene->GetDataflowModeManager())
+			{
+				if (UDataflowEditorMode* DataflowMode = Cast<UDataflowEditorMode>(PreviewScene->GetDataflowModeManager()->GetActiveScriptableMode(UDataflowEditorMode::EM_DataflowEditorModeId)))
+				{
+					if (TSharedPtr<SDataflowGraphEditor> GraphEditor = DataflowEditorToolkit->GetDataflowGraphEditor())
+					{
+						if (UEdGraphNode* SelectedNode = GraphEditor->GetSingleSelectedNode())
+						{
+							if (SelectedComponents && SelectedComponents->Num() == 1)
+							{
+								if (UDataflowEditorCollectionComponent* CollectionComponent =
+									Cast< UDataflowEditorCollectionComponent>(SelectedComponents->GetSelectedObject(0)))
+								{
+									if (CollectionComponent->Node == SelectedNode)
+									{
+										// Start the corresponding tool
+										DataflowMode->StartToolForSelectedNode(SelectedNode);
+									}
+								}
+							}
+						}
+					}
+				}
+			}
+		}
+	};
 
-	TArray<UPrimitiveComponent*> PreviouslySelectedComponents;
-	SelectedComponents->GetSelectedObjects<UPrimitiveComponent>(PreviouslySelectedComponents);
+	auto UpdateSelectedComponentInViewport = [&](USelection* SelectedComponents)
+	{
+		TArray<UPrimitiveComponent*> PreviouslySelectedComponents;
+		SelectedComponents->GetSelectedObjects<UPrimitiveComponent>(PreviouslySelectedComponents);
 
-	SelectedComponents->Modify();
-	SelectedComponents->BeginBatchSelectOperation();
+		SelectedComponents->Modify();
+		SelectedComponents->BeginBatchSelectOperation();
 
-	SelectedComponents->DeselectAll();
+		SelectedComponents->DeselectAll();
 
-	if (HitProxy && HitProxy->IsA(HActor::StaticGetType()))
-	{
-		const HActor* ActorProxy = static_cast<HActor*>(HitProxy);
-		if (ActorProxy && ActorProxy->PrimComponent && ActorProxy->Actor)
+		if (HitProxy && HitProxy->IsA(HActor::StaticGetType()))
 		{
-			UPrimitiveComponent* Component = const_cast<UPrimitiveComponent*>(ActorProxy->PrimComponent.Get());
-			SelectedComponents->Select(Component);
-			Component->PushSelectionToProxy();
+			const HActor* ActorProxy = static_cast<HActor*>(HitProxy);
+			if (ActorProxy && ActorProxy->PrimComponent && ActorProxy->Actor)
+			{
+				UPrimitiveComponent* Component = const_cast<UPrimitiveComponent*>(ActorProxy->PrimComponent.Get());
+				SelectedComponents->Select(Component);
+				Component->PushSelectionToProxy();
+			}
 		}
-	}
 
-	SelectedComponents->EndBatchSelectOperation();
+		SelectedComponents->EndBatchSelectOperation();
+
+		for (UPrimitiveComponent* const Component : PreviouslySelectedComponents)
+		{
+			Component->PushSelectionToProxy();
+		}
 
-	for (UPrimitiveComponent* const Component : PreviouslySelectedComponents)
+	};
+	
+	const bool bIsCtrltKeyDown = Viewport->KeyState(EKeys::LeftControl) || Viewport->KeyState(EKeys::RightControl);
+	if (USelection* SelectedComponents = ModeTools->GetSelectedComponents())
 	{
-		Component->PushSelectionToProxy();
+		UpdateSelectedComponentInViewport(SelectedComponents);
+		EnableToolForSelectedNode(SelectedComponents);
 	}
 }
 
diff --git a/Engine/Plugins/Experimental/Dataflow/Source/DataflowEditor/Private/Dataflow/DataflowGraphEditor.cpp b/Engine/Plugins/Experimental/Dataflow/Source/DataflowEditor/Private/Dataflow/DataflowGraphEditor.cpp
--- a/Engine/Plugins/Experimental/Dataflow/Source/DataflowEditor/Private/Dataflow/DataflowGraphEditor.cpp
+++ b/Engine/Plugins/Experimental/Dataflow/Source/DataflowEditor/Private/Dataflow/DataflowGraphEditor.cpp
@@ -204,6 +204,11 @@ FReply SDataflowGraphEditor::OnKeyDown(const FGeometry& MyGeometry, const FKeyEv
 	return SGraphEditor::OnKeyUp(MyGeometry, InKeyEvent);
 }
 
+bool SDataflowGraphEditor::IsControlDown() const
+{
+	return LeftControlKeyDown || RightControlKeyDown;
+}
+
 
 FReply SDataflowGraphEditor::OnKeyUp(const FGeometry& MyGeometry, const FKeyEvent& InKeyEvent)
 {
diff --git a/Engine/Plugins/Experimental/Dataflow/Source/DataflowEditor/Private/Dataflow/DataflowPreviewScene.cpp b/Engine/Plugins/Experimental/Dataflow/Source/DataflowEditor/Private/Dataflow/DataflowPreviewScene.cpp
--- a/Engine/Plugins/Experimental/Dataflow/Source/DataflowEditor/Private/Dataflow/DataflowPreviewScene.cpp
+++ b/Engine/Plugins/Experimental/Dataflow/Source/DataflowEditor/Private/Dataflow/DataflowPreviewScene.cpp
@@ -40,7 +40,7 @@ void FDataflowPreviewScene::AddReferencedObjects(FReferenceCollector& Collector)
 	FAdvancedPreviewScene::AddReferencedObjects(Collector);
 	Collector.AddReferencedObject(DataflowContent);
 	Collector.AddReferencedObject(RootSceneActor);
-	
+
 	GetDataflowContent()->AddContentObjects(Collector);
 }
 
@@ -342,17 +342,6 @@ void FDataflowConstructionScene::UpdateConstructionScene()
 	// Attach a wireframe renderer to the DynamicMeshComponents
 	UpdateWireframeMeshElementsVisualizer();
 
-	// Manage Selection and Tool Interaction
-    if (DataflowModeManager.IsValid())
-    {
-    	USelection* SelectedComponents = DataflowModeManager->GetSelectedComponents();
-    	SelectedComponents->DeselectAll();
-    	for (FRenderElement RenderElement : DynamicMeshComponents)
-    	{
-    		SelectedComponents->Select(RenderElement.Value);
-			RenderElement.Value->PushSelectionToProxy();
-    	}
-    }
 	DataflowContent->SetIsDirty(false);
 }
 
diff --git a/Engine/Plugins/Experimental/Dataflow/Source/DataflowEditor/Private/DataflowEditorTools/DataflowEditorWeightMapPaintTool.cpp b/Engine/Plugins/Experimental/Dataflow/Source/DataflowEditor/Private/DataflowEditorTools/DataflowEditorWeightMapPaintTool.cpp
--- a/Engine/Plugins/Experimental/Dataflow/Source/DataflowEditor/Private/DataflowEditorTools/DataflowEditorWeightMapPaintTool.cpp
+++ b/Engine/Plugins/Experimental/Dataflow/Source/DataflowEditor/Private/DataflowEditorTools/DataflowEditorWeightMapPaintTool.cpp
@@ -449,6 +449,11 @@ void UDataflowEditorWeightMapPaintTool::Shutdown(EToolShutdownType ShutdownType)
 		PolygonSelectionMechanic = nullptr;
 	}
 
+	if (WeightMapNodeToUpdate)
+	{
+		WeightMapNodeToUpdate->Invalidate();
+	}
+
 	UMeshSculptToolBase::Shutdown(ShutdownType);
 }
 
diff --git a/Engine/Plugins/Experimental/Dataflow/Source/DataflowEditor/Public/Dataflow/DataflowEditorToolkit.h b/Engine/Plugins/Experimental/Dataflow/Source/DataflowEditor/Public/Dataflow/DataflowEditorToolkit.h
--- a/Engine/Plugins/Experimental/Dataflow/Source/DataflowEditor/Public/Dataflow/DataflowEditorToolkit.h
+++ b/Engine/Plugins/Experimental/Dataflow/Source/DataflowEditor/Public/Dataflow/DataflowEditorToolkit.h
@@ -54,6 +54,7 @@ public:
 	/** Dataflow Content Access */
 	const TObjectPtr<UDataflowBaseContent> GetDataflowContent() const;
 	TObjectPtr<UDataflowBaseContent> GetDataflowContent();
+	const TSharedPtr<SDataflowGraphEditor> GetDataflowGraphEditor() const { return GraphEditor; }
 
 	// IToolkit interface
 	virtual FName GetToolkitFName() const override;
diff --git a/Engine/Plugins/Experimental/Dataflow/Source/DataflowEditor/Public/Dataflow/DataflowGraphEditor.h b/Engine/Plugins/Experimental/Dataflow/Source/DataflowEditor/Public/Dataflow/DataflowGraphEditor.h
--- a/Engine/Plugins/Experimental/Dataflow/Source/DataflowEditor/Public/Dataflow/DataflowGraphEditor.h
+++ b/Engine/Plugins/Experimental/Dataflow/Source/DataflowEditor/Public/Dataflow/DataflowGraphEditor.h
@@ -62,6 +62,7 @@ public:
 	virtual FReply OnKeyUp(const FGeometry& MyGeometry, const FKeyEvent& InKeyEvent) override;
 	virtual FReply OnDragOver(const FGeometry& MyGeometry, const FDragDropEvent& DragDropEvent) override;
 	virtual FReply OnDrop(const FGeometry& MyGeometry, const FDragDropEvent& DragDropEvent) override;
+	bool IsControlDown() const;
 	//virtual void OnDragEnter(const FGeometry& MyGeometry, const FDragDropEvent& DragDropEvent) override;
 	//virtual void OnDragLeave(const FDragDropEvent& DragDropEvent) override;
 	// end SWidget
