Dataflow : Alt Selection

diff --git a/Engine/Plugins/Experimental/Dataflow/Source/DataflowEditor/Private/Dataflow/DataflowEditorToolkit.cpp b/Engine/Plugins/Experimental/Dataflow/Source/DataflowEditor/Private/Dataflow/DataflowEditorToolkit.cpp
--- a/Engine/Plugins/Experimental/Dataflow/Source/DataflowEditor/Private/Dataflow/DataflowEditorToolkit.cpp
+++ b/Engine/Plugins/Experimental/Dataflow/Source/DataflowEditor/Private/Dataflow/DataflowEditorToolkit.cpp
@@ -3,10 +3,12 @@
 #include "Dataflow/DataflowEditorToolkit.h"
 
 #include "AdvancedPreviewScene.h"
+#include "AssetEditorModeManager.h"
 #include "Animation/Skeleton.h"
 #include "Dataflow/DataflowCore.h"
 #include "Dataflow/DataflowEditor.h"
 #include "Dataflow/DataflowContent.h"
+#include "Dataflow/DataflowEditorCollectionComponent.h"
 #include "Dataflow/DataflowEditorCommands.h"
 #include "Dataflow/DataflowEditorMode.h"
 #include "Dataflow/DataflowEditorModeToolkit.h"
@@ -29,17 +31,20 @@
 #include "EditorViewportCommands.h"
 #include "EdModeInteractiveToolsContext.h"
 #include "Engine/SkeletalMesh.h"
-#include "GraphEditorActions.h"
-#include "Modules/ModuleManager.h"
-#include "PropertyEditorModule.h"
-#include "Styling/SlateStyleRegistry.h"
 #include "Framework/Commands/GenericCommands.h"
-#include "Widgets/Docking/SDockTab.h"
+#include "GameFramework/Actor.h"
+#include "GraphEditorActions.h"
 #include "ISkeletonTree.h"
-#include "AssetEditorModeManager.h"
 #include "ISkeletonEditorModule.h"
 #include "IStructureDetailsView.h"
+#include "Kismet/GameplayStatics.h"
 #include "Misc/MessageDialog.h"
+#include "Modules/ModuleManager.h"
+#include "PropertyEditorModule.h"
+#include "Selection.h"
+#include "Styling/SlateStyleRegistry.h"
+#include "Widgets/Docking/SDockTab.h"
+
 
 #define LOCTEXT_NAMESPACE "DataflowEditorToolkit"
 
@@ -540,6 +545,48 @@ void FDataflowEditorToolkit::OnNodeSelectionChanged(const TSet<UObject*>& InNewS
 		return false;
 	};
 
+	auto SelectComponentsInView = [&](TObjectPtr<UDataflowEdNode> Node)
+	{
+		TArray<AActor*> FoundActors;
+		FDataflowConstructionScene* ConstructionScene = static_cast<FDataflowConstructionScene*>(ObjectScene.Get());
+		if (USelection* SelectedComponents = ConstructionScene->GetDataflowModeManager()->GetSelectedComponents())
+		{
+			SelectedComponents->Modify();
+			SelectedComponents->BeginBatchSelectOperation();
+
+			TArray<TWeakObjectPtr<UObject>> SelectedObjects;
+			const int32 NumSelected = SelectedComponents->GetSelectedObjects(SelectedObjects);
+			for (TWeakObjectPtr<UObject> WeakObject : SelectedObjects)
+			{
+				if (WeakObject.IsValid())
+				{
+					if (UDataflowEditorCollectionComponent* ActorComponent = Cast< UDataflowEditorCollectionComponent>(WeakObject.Get()))
+					{
+						SelectedComponents->Deselect(ActorComponent);
+						ActorComponent->PushSelectionToProxy();
+					}
+				}
+			}
+
+			if (TObjectPtr<AActor> RootActor = static_cast<FDataflowPreviewScene*>(ObjectScene.Get())->GetRootActor())
+			{
+				for (UActorComponent* ActorComponent : RootActor->GetComponents())
+				{
+					if (UDataflowEditorCollectionComponent* Component = Cast<UDataflowEditorCollectionComponent>(ActorComponent))
+					{
+						if (Component->Node == Node)
+						{
+							SelectedComponents->Select(Component);
+							Component->PushSelectionToProxy();
+						}
+					}
+				}
+			}
+			SelectedComponents->EndBatchSelectOperation();
+			//GetDataflowContent()->SetIsDirty(true);
+		}
+	};
+
 	// Despite this function's name, we might not have actually changed which node is selected
 	bool bPrimarySelectionChanged = false;
 
@@ -610,6 +657,11 @@ void FDataflowEditorToolkit::OnNodeSelectionChanged(const TSet<UObject*>& InNewS
 
 				EditorContent->SetPrimarySelectedNode(PrimarySelection);
 			}
+
+			if (GetDataflowGraphEditor()->IsAltDown())
+			{
+				SelectComponentsInView(PrimarySelection);
+			}
 		}
 	}
 }
diff --git a/Engine/Plugins/Experimental/Dataflow/Source/DataflowEditor/Private/Dataflow/DataflowEditorViewportClient.cpp b/Engine/Plugins/Experimental/Dataflow/Source/DataflowEditor/Private/Dataflow/DataflowEditorViewportClient.cpp
--- a/Engine/Plugins/Experimental/Dataflow/Source/DataflowEditor/Private/Dataflow/DataflowEditorViewportClient.cpp
+++ b/Engine/Plugins/Experimental/Dataflow/Source/DataflowEditor/Private/Dataflow/DataflowEditorViewportClient.cpp
@@ -10,8 +10,13 @@
 #include "Dataflow/DataflowGraphEditor.h"
 #include "Dataflow/DataflowPreviewScene.h"
 #include "EditorModeManager.h"
+#include "GraphEditor.h"
 #include "PreviewScene.h"
 #include "Selection.h"
+#include "SGraphPanel.h"
+#include "SNodePanel.h"
+
+
 
 FDataflowEditorViewportClient::FDataflowEditorViewportClient(FEditorModeTools* InModeTools,
                                                              FPreviewScene* InPreviewScene,  const bool bCouldTickScene,
@@ -116,11 +121,32 @@ void FDataflowEditorViewportClient::ProcessClick(FSceneView& View, HHitProxy* Hi
 		}
 
 	};
+
+	auto SelectSingleNodeInGraph = [&](TObjectPtr<const UDataflowEdNode> Node)
+	{
+		if (TSharedPtr<FDataflowEditorToolkit> DataflowEditorToolkit = DataflowEditorToolkitPtr.Pin())
+		{
+			if (TSharedPtr<SDataflowGraphEditor> GraphEditor = DataflowEditorToolkit->GetDataflowGraphEditor())
+			{
+				GraphEditor->GetGraphPanel()->SelectionManager.SelectSingleNode((UObject*)Node.Get());
+			}
+		}
+	};
 	
-	const bool bIsCtrltKeyDown = Viewport->KeyState(EKeys::LeftControl) || Viewport->KeyState(EKeys::RightControl);
 	if (USelection* SelectedComponents = ModeTools->GetSelectedComponents())
 	{
+
 		UpdateSelectedComponentInViewport(SelectedComponents);
+
+		if (bool bIsAltKeyDown = Viewport->KeyState(EKeys::LeftAlt) || Viewport->KeyState(EKeys::RightAlt))
+		{
+			if (UDataflowEditorCollectionComponent* DataflowComponent
+				= SelectedComponents->GetBottom<UDataflowEditorCollectionComponent>())
+			{
+				SelectSingleNodeInGraph(DataflowComponent->Node);
+			}
+		}
+
 		EnableToolForSelectedNode(SelectedComponents);
 	}
 }
diff --git a/Engine/Plugins/Experimental/Dataflow/Source/DataflowEditor/Private/Dataflow/DataflowGraphEditor.cpp b/Engine/Plugins/Experimental/Dataflow/Source/DataflowEditor/Private/Dataflow/DataflowGraphEditor.cpp
--- a/Engine/Plugins/Experimental/Dataflow/Source/DataflowEditor/Private/Dataflow/DataflowGraphEditor.cpp
+++ b/Engine/Plugins/Experimental/Dataflow/Source/DataflowEditor/Private/Dataflow/DataflowGraphEditor.cpp
@@ -197,6 +197,14 @@ FReply SDataflowGraphEditor::OnKeyDown(const FGeometry& MyGeometry, const FKeyEv
 	{
 		RightControlKeyDown = true;
 	}
+	if (InKeyEvent.GetKey() == EKeys::LeftAlt)
+	{
+		LeftAltKeyDown = true;
+	}
+	if (InKeyEvent.GetKey() == EKeys::RightAlt)
+	{
+		RightAltKeyDown = true;
+	}
 	if (InKeyEvent.GetKey() == EKeys::V)
 	{
 		VKeyDown = true;
@@ -209,6 +217,13 @@ bool SDataflowGraphEditor::IsControlDown() const
 	return LeftControlKeyDown || RightControlKeyDown;
 }
 
+bool SDataflowGraphEditor::IsAltDown() const
+{
+	return LeftAltKeyDown || RightAltKeyDown;
+}
+
+
+
 
 FReply SDataflowGraphEditor::OnKeyUp(const FGeometry& MyGeometry, const FKeyEvent& InKeyEvent)
 {
@@ -237,6 +252,14 @@ FReply SDataflowGraphEditor::OnKeyUp(const FGeometry& MyGeometry, const FKeyEven
 	{
 		RightControlKeyDown = false;
 	}
+	if (InKeyEvent.GetKey() == EKeys::LeftAlt)
+	{
+		LeftAltKeyDown = false;
+	}
+	if (InKeyEvent.GetKey() == EKeys::RightAlt)
+	{
+		RightAltKeyDown = false;
+	}
 	if (InKeyEvent.GetKey() == EKeys::V)
 	{
 		VKeyDown = false;
diff --git a/Engine/Plugins/Experimental/Dataflow/Source/DataflowEditor/Public/Dataflow/DataflowGraphEditor.h b/Engine/Plugins/Experimental/Dataflow/Source/DataflowEditor/Public/Dataflow/DataflowGraphEditor.h
--- a/Engine/Plugins/Experimental/Dataflow/Source/DataflowEditor/Public/Dataflow/DataflowGraphEditor.h
+++ b/Engine/Plugins/Experimental/Dataflow/Source/DataflowEditor/Public/Dataflow/DataflowGraphEditor.h
@@ -63,6 +63,7 @@ public:
 	virtual FReply OnDragOver(const FGeometry& MyGeometry, const FDragDropEvent& DragDropEvent) override;
 	virtual FReply OnDrop(const FGeometry& MyGeometry, const FDragDropEvent& DragDropEvent) override;
 	bool IsControlDown() const;
+	bool IsAltDown() const;
 	//virtual void OnDragEnter(const FGeometry& MyGeometry, const FDragDropEvent& DragDropEvent) override;
 	//virtual void OnDragLeave(const FDragDropEvent& DragDropEvent) override;
 	// end SWidget
@@ -161,4 +162,6 @@ private:
 	bool VKeyDown = false;
 	bool LeftControlKeyDown = false;
 	bool RightControlKeyDown = false;
+	bool LeftAltKeyDown = false;
+	bool RightAltKeyDown = false;
 };
