Dataflow : Selection updates

diff --git a/Engine/Plugins/Experimental/Dataflow/Source/DataflowEditor/Private/Dataflow/DataflowEditorToolkit.cpp b/Engine/Plugins/Experimental/Dataflow/Source/DataflowEditor/Private/Dataflow/DataflowEditorToolkit.cpp
--- a/Engine/Plugins/Experimental/Dataflow/Source/DataflowEditor/Private/Dataflow/DataflowEditorToolkit.cpp
+++ b/Engine/Plugins/Experimental/Dataflow/Source/DataflowEditor/Private/Dataflow/DataflowEditorToolkit.cpp
@@ -583,7 +583,6 @@ void FDataflowEditorToolkit::OnNodeSelectionChanged(const TSet<UObject*>& InNewS
 				}
 			}
 			SelectedComponents->EndBatchSelectOperation();
-			//GetDataflowContent()->SetIsDirty(true);
 		}
 	};
 
diff --git a/Engine/Plugins/Experimental/Dataflow/Source/DataflowEditor/Private/Dataflow/DataflowEditorViewportClient.cpp b/Engine/Plugins/Experimental/Dataflow/Source/DataflowEditor/Private/Dataflow/DataflowEditorViewportClient.cpp
--- a/Engine/Plugins/Experimental/Dataflow/Source/DataflowEditor/Private/Dataflow/DataflowEditorViewportClient.cpp
+++ b/Engine/Plugins/Experimental/Dataflow/Source/DataflowEditor/Private/Dataflow/DataflowEditorViewportClient.cpp
@@ -10,6 +10,7 @@
 #include "Dataflow/DataflowGraphEditor.h"
 #include "Dataflow/DataflowPreviewScene.h"
 #include "EditorModeManager.h"
+#include "EdModeInteractiveToolsContext.h"
 #include "GraphEditor.h"
 #include "PreviewScene.h"
 #include "Selection.h"
@@ -133,21 +134,36 @@ void FDataflowEditorViewportClient::ProcessClick(FSceneView& View, HHitProxy* Hi
 		}
 	};
 	
-	if (USelection* SelectedComponents = ModeTools->GetSelectedComponents())
+	auto IsInteractiveToolActive = [&]()
 	{
-
-		UpdateSelectedComponentInViewport(SelectedComponents);
-
-		if (bool bIsAltKeyDown = Viewport->KeyState(EKeys::LeftAlt) || Viewport->KeyState(EKeys::RightAlt))
+		if (UDataflowEditorMode* DataflowMode = Cast<UDataflowEditorMode>(
+			PreviewScene->GetDataflowModeManager()->GetActiveScriptableMode(UDataflowEditorMode::EM_DataflowEditorModeId)))
 		{
-			if (UDataflowEditorCollectionComponent* DataflowComponent
-				= SelectedComponents->GetBottom<UDataflowEditorCollectionComponent>())
+			if (UEditorInteractiveToolsContext* const ToolsContext = DataflowMode->GetInteractiveToolsContext())
 			{
-				SelectSingleNodeInGraph(DataflowComponent->Node);
+				return ToolsContext->HasActiveTool();
 			}
 		}
+		return false;
+	};
+
+	if (!IsInteractiveToolActive())
+	{
+		if (USelection* SelectedComponents = ModeTools->GetSelectedComponents())
+		{
+			UpdateSelectedComponentInViewport(SelectedComponents);
 
-		EnableToolForSelectedNode(SelectedComponents);
+			if (bool bIsAltKeyDown = Viewport->KeyState(EKeys::LeftAlt) || Viewport->KeyState(EKeys::RightAlt))
+			{
+				if (UDataflowEditorCollectionComponent* DataflowComponent
+					= SelectedComponents->GetBottom<UDataflowEditorCollectionComponent>())
+				{
+					SelectSingleNodeInGraph(DataflowComponent->Node);
+				}
+			}
+
+			EnableToolForSelectedNode(SelectedComponents);
+		}
 	}
 }
 
