Dataflow : Rendering
- remove ensure on remapping generation as its optional.

diff --git a/Engine/Plugins/Experimental/Dataflow/Source/DataflowAssetTools/Private/Dataflow/CollectionRenderingPatternUtility.cpp b/Engine/Plugins/Experimental/Dataflow/Source/DataflowAssetTools/Private/Dataflow/CollectionRenderingPatternUtility.cpp
--- a/Engine/Plugins/Experimental/Dataflow/Source/DataflowAssetTools/Private/Dataflow/CollectionRenderingPatternUtility.cpp
+++ b/Engine/Plugins/Experimental/Dataflow/Source/DataflowAssetTools/Private/Dataflow/CollectionRenderingPatternUtility.cpp
@@ -65,7 +65,7 @@ namespace Dataflow
 				DynamicMesh.EnableAttributes();
 
 				// Build Remmaping indices back into the colleciton. 
-				if (ensure(Remapping.Num() < Facade.NumVertices()))
+				if (Remapping.Num() < Facade.NumVertices())
 				{
 					UE::Geometry::FNonManifoldMappingSupport::AttachNonManifoldVertexMappingData(Remapping, DynamicMesh);
 				}
