Fix crash when right clicking a comment node in DataflowGraphEditor

diff --git a/Engine/Plugins/Experimental/Dataflow/Source/DataflowEditor/Private/Dataflow/DataflowGraphEditor.cpp b/Engine/Plugins/Experimental/Dataflow/Source/DataflowEditor/Private/Dataflow/DataflowGraphEditor.cpp
--- a/Engine/Plugins/Experimental/Dataflow/Source/DataflowEditor/Private/Dataflow/DataflowGraphEditor.cpp
+++ b/Engine/Plugins/Experimental/Dataflow/Source/DataflowEditor/Private/Dataflow/DataflowGraphEditor.cpp
@@ -444,20 +444,21 @@ bool SDataflowGraphEditor::CanAddOptionPin() const
 		// Iterate over all nodes, and add the pin
 		for (FGraphPanelSelectionSet::TConstIterator It(SelectedNodes); It; ++It)
 		{
-			const UDataflowEdNode* const EdNode = CastChecked<UDataflowEdNode>(*It);
-
-			if (const TSharedPtr<const FDataflowNode> Node = DataflowGraph->FindBaseNode(EdNode->DataflowNodeGuid))
-			{
-				bCanAddOptionPin = Node->CanAddPin();
-			}
-			else
+			if (const UDataflowEdNode* const EdNode = Cast<UDataflowEdNode>(*It))
 			{
-				bCanAddOptionPin = false;
-			}
+				if (const TSharedPtr<const FDataflowNode> Node = DataflowGraph->FindBaseNode(EdNode->DataflowNodeGuid))
+				{
+					bCanAddOptionPin = Node->CanAddPin();
+				}
+				else
+				{
+					bCanAddOptionPin = false;
+				}
 
-			if (!bCanAddOptionPin)
-			{
-				break;  // One bad node is good enough to return false
+				if (!bCanAddOptionPin)
+				{
+					break;  // One bad node is good enough to return false
+				}
 			}
 		}
 	}
@@ -507,20 +508,21 @@ bool SDataflowGraphEditor::CanRemoveOptionPin() const
 		// Iterate over all nodes, and add the pin
 		for (FGraphPanelSelectionSet::TConstIterator It(SelectedNodes); It; ++It)
 		{
-			const UDataflowEdNode* const EdNode = CastChecked<UDataflowEdNode>(*It);
-
-			if (const TSharedPtr<const FDataflowNode> Node = DataflowGraph->FindBaseNode(EdNode->DataflowNodeGuid))
-			{
-				bCanRemoveOptionPin = Node->CanRemovePin();
-			}
-			else
+			if (const UDataflowEdNode* const EdNode = Cast<UDataflowEdNode>(*It))
 			{
-				bCanRemoveOptionPin = false;
-			}
+				if (const TSharedPtr<const FDataflowNode> Node = DataflowGraph->FindBaseNode(EdNode->DataflowNodeGuid))
+				{
+					bCanRemoveOptionPin = Node->CanRemovePin();
+				}
+				else
+				{
+					bCanRemoveOptionPin = false;
+				}
 
-			if (!bCanRemoveOptionPin)
-			{
-				break;  // One bad node is good enough to return false
+				if (!bCanRemoveOptionPin)
+				{
+					break;  // One bad node is good enough to return false
+				}
 			}
 		}
 	}
