ChaosFlesh : Validate content berfore adding referenced objects.

diff --git a/Engine/Plugins/Experimental/Dataflow/Source/DataflowEditor/Private/Dataflow/DataflowPreviewScene.cpp b/Engine/Plugins/Experimental/Dataflow/Source/DataflowEditor/Private/Dataflow/DataflowPreviewScene.cpp
--- a/Engine/Plugins/Experimental/Dataflow/Source/DataflowEditor/Private/Dataflow/DataflowPreviewScene.cpp
+++ b/Engine/Plugins/Experimental/Dataflow/Source/DataflowEditor/Private/Dataflow/DataflowPreviewScene.cpp
@@ -40,8 +40,10 @@ void FDataflowPreviewScene::AddReferencedObjects(FReferenceCollector& Collector)
 	FAdvancedPreviewScene::AddReferencedObjects(Collector);
 	Collector.AddReferencedObject(DataflowContent);
 	Collector.AddReferencedObject(RootSceneActor);
-
-	GetDataflowContent()->AddContentObjects(Collector);
+	if (GetDataflowContent())
+	{
+		GetDataflowContent()->AddContentObjects(Collector);
+	}
 }
 
 bool FDataflowPreviewScene::IsComponentSelected(const UPrimitiveComponent* InComponent) const
