ChaosFlesh : Add the cached nodes to the ObjectToEdit.

diff --git a/Engine/Plugins/Experimental/Dataflow/Source/DataflowEditor/Private/Dataflow/DataflowEditor.cpp b/Engine/Plugins/Experimental/Dataflow/Source/DataflowEditor/Private/Dataflow/DataflowEditor.cpp
--- a/Engine/Plugins/Experimental/Dataflow/Source/DataflowEditor/Private/Dataflow/DataflowEditor.cpp
+++ b/Engine/Plugins/Experimental/Dataflow/Source/DataflowEditor/Private/Dataflow/DataflowEditor.cpp
@@ -32,6 +32,8 @@ void UDataflowEditor::Initialize(const TArray<TObjectPtr<UObject>>& InObjects)
 
 void UDataflowEditor::InitializeContent(TObjectPtr<UDataflowBaseContent> BaseContent, const TObjectPtr<UObject>& ContentOwner)
 {
+	TArray<TObjectPtr<UObject>> RequiredObjects = { ContentOwner };
+
 	DataflowContent = BaseContent;
 	if(!DataflowContent)
 	{
@@ -59,12 +61,17 @@ void UDataflowEditor::InitializeContent(TObjectPtr<UDataflowBaseContent> BaseCon
 				{
 					DataflowContent = DataflowContextDefinitionHelpers::CreateNewDataflowContext<UDataflowBaseContent>(ContentOwner);
 				}
+
 				DataflowContent->SetDataflowAsset(Private::GetDataflowAssetFrom(ContentOwner));
 				DataflowContent->SetDataflowTerminal(Private::GetDataflowTerminalFrom(ContentOwner));
+				RequiredObjects.Add(Private::GetDataflowAssetFrom(ContentOwner));
 			}
 		}
 	}
 
+	if (!DataflowContent) return;
+	RequiredObjects.Add(DataflowContent);
+
 	if(const TObjectPtr<UDataflowSkeletalContent> SkeletalContent = Cast<UDataflowSkeletalContent>(DataflowContent))
 	{
 		if(!SkeletalContent->GetSkeletalMesh())
@@ -94,9 +101,6 @@ void UDataflowEditor::InitializeContent(TObjectPtr<UDataflowBaseContent> BaseCon
 	
 	// Potentially we could add additional objects to edit here (fields, meshes....)
 	// If these objects have a matching factory we would be able to use geometry tools
-	const TArray<TObjectPtr<UObject>> ArrayObjects = {ContentOwner};
-	UBaseCharacterFXEditor::Initialize(ArrayObjects);
+	UBaseCharacterFXEditor::Initialize(RequiredObjects);
 }
 
-
-
