Dataflow : Evaluation updates
When selecting nodes directly we need to pass the node and out to evaluate. Otherwise, the evaluation defaults to a terminal node, if it is defined.

diff --git a/Engine/Plugins/Experimental/Dataflow/Source/DataflowEditor/Private/Dataflow/DataflowEditorToolkit.cpp b/Engine/Plugins/Experimental/Dataflow/Source/DataflowEditor/Private/Dataflow/DataflowEditorToolkit.cpp
--- a/Engine/Plugins/Experimental/Dataflow/Source/DataflowEditor/Private/Dataflow/DataflowEditorToolkit.cpp
+++ b/Engine/Plugins/Experimental/Dataflow/Source/DataflowEditor/Private/Dataflow/DataflowEditorToolkit.cpp
@@ -702,6 +702,8 @@ void FDataflowEditorToolkit::Tick(float DeltaTime)
 				EditorContent->SetDataflowContext(MakeShared<Dataflow::FEngineContext>(EditorContent->GetDataflowOwner(), EditorContent->GetDataflowAsset(), Dataflow::FTimestamp::Invalid));
 				TimeStamp = Dataflow::FTimestamp::Invalid;
 			}
+
+			// OnTick evaluation only pulls the termnial nodes. The other evaluations can be specific nodes. 
 			FDataflowEditorCommands::EvaluateTerminalNode(*EditorContent->GetDataflowContext().Get(), TimeStamp, EditorContent->GetDataflowAsset(),
 				nullptr, nullptr, EditorContent->GetDataflowOwner(), EditorContent->GetDataflowTerminal());
 			EditorContent->SetLastModifiedTimestamp(TimeStamp);
@@ -733,7 +735,8 @@ TSharedRef<SDataflowGraphEditor> FDataflowEditorToolkit::CreateGraphEditorWidget
 				Dataflow::FTimestamp TimeStamp = Dataflow::FTimestamp::Invalid;
 				
 				FDataflowEditorCommands::EvaluateTerminalNode(*EditorContent->GetDataflowContext().Get(), TimeStamp, EditorContent->GetDataflowAsset(),
-					nullptr, nullptr, EditorContent->GetDataflowOwner(), EditorContent->GetDataflowTerminal());
+					Node, Out, EditorContent->GetDataflowOwner(), EditorContent->GetDataflowTerminal());
+
 				EditorContent->SetLastModifiedTimestamp(TimeStamp);
 			}
 		}
