Dataflow - Enable the comment node in the property editor.

diff --git a/Engine/Plugins/Experimental/Dataflow/Source/DataflowEditor/Private/Dataflow/DataflowEditorCommands.cpp b/Engine/Plugins/Experimental/Dataflow/Source/DataflowEditor/Private/Dataflow/DataflowEditorCommands.cpp
--- a/Engine/Plugins/Experimental/Dataflow/Source/DataflowEditor/Private/Dataflow/DataflowEditorCommands.cpp
+++ b/Engine/Plugins/Experimental/Dataflow/Source/DataflowEditor/Private/Dataflow/DataflowEditorCommands.cpp
@@ -420,8 +420,12 @@ void FDataflowEditorCommands::OnSelectedNodesChanged(TSharedPtr<IStructureDetail
 							PropertiesEditor->SetStructureData(Struct);
 						}
 					}
+					else if (UEdGraphNode_Comment* CommentNode = Cast<UEdGraphNode_Comment>(SelectedObject))
+					{
+						TSharedPtr<FStructOnScope> Struct(new FStructOnScope(UEdGraphNode_Comment::StaticClass(), (uint8*)CommentNode));
+						PropertiesEditor->SetStructureData(Struct);
+					}
 				}
-				
 			}
 		}
 	}
