Dataflow:
- Fixed naming for nodes/commentboxes created by Copy/Cut/Paste and Duplicate

diff --git a/Engine/Plugins/Experimental/Dataflow/Source/DataflowEditor/Private/Dataflow/DataflowSCommentNode.cpp b/Engine/Plugins/Experimental/Dataflow/Source/DataflowEditor/Private/Dataflow/DataflowSCommentNode.cpp
--- a/Engine/Plugins/Experimental/Dataflow/Source/DataflowEditor/Private/Dataflow/DataflowSCommentNode.cpp
+++ b/Engine/Plugins/Experimental/Dataflow/Source/DataflowEditor/Private/Dataflow/DataflowSCommentNode.cpp
@@ -86,12 +86,6 @@ UEdGraphNode* FAssetSchemaAction_Dataflow_DuplicateCommentNode_DataflowEdNode::P
 {
 	FString NewNodeName = CommentNodeToDuplicate->NodeComment;
 
-	// Append "_copy" to selected node's name if it doesn't have it
-	if (!NewNodeName.Contains(CopyPostfix))
-	{
-		NewNodeName.Append(CopyPostfix);
-	}
-
 	UEdGraphNode_Comment* CommentTemplate = NewObject<UEdGraphNode_Comment>(ParentGraph, FName(*NewNodeName));
 
 	CommentTemplate->bCommentBubbleVisible_InDetailsPanel = false;
@@ -143,12 +137,7 @@ TSharedPtr<FAssetSchemaAction_Dataflow_PasteCommentNode_DataflowEdNode> FAssetSc
 //
 UEdGraphNode* FAssetSchemaAction_Dataflow_PasteCommentNode_DataflowEdNode::PerformAction(class UEdGraph* ParentGraph, UEdGraphPin* FromPin, const FVector2D Location, bool bSelectNewNode)
 {
-	// Append "_copy" to selected node's name if it doesn't have it
 	FString NewNodeName = NodeName.ToString();
-	if (!NewNodeName.Contains(CopyPostfix))
-	{
-		NewNodeName.Append(CopyPostfix);
-	}
 
 	UEdGraphNode_Comment* CommentTemplate = NewObject<UEdGraphNode_Comment>(ParentGraph, FName(*NewNodeName));
 
diff --git a/Engine/Plugins/Experimental/Dataflow/Source/DataflowEditor/Private/Dataflow/DataflowSNode.cpp b/Engine/Plugins/Experimental/Dataflow/Source/DataflowEditor/Private/Dataflow/DataflowSNode.cpp
--- a/Engine/Plugins/Experimental/Dataflow/Source/DataflowEditor/Private/Dataflow/DataflowSNode.cpp
+++ b/Engine/Plugins/Experimental/Dataflow/Source/DataflowEditor/Private/Dataflow/DataflowSNode.cpp
@@ -187,6 +187,32 @@ TSharedPtr<FAssetSchemaAction_Dataflow_CreateNode_DataflowEdNode> FAssetSchemaAc
 	return TSharedPtr<FAssetSchemaAction_Dataflow_CreateNode_DataflowEdNode>(nullptr);
 }
 
+static FName GetNodeUniqueName(UDataflow* Dataflow, FString NodeBaseName)
+{
+	FString Left, Right;
+	int32 NameIndex = 1;
+
+	// Check if NodeBaseName already ends with "_dd"
+	if (NodeBaseName.Split(TEXT("_"), &Left, &Right, ESearchCase::IgnoreCase, ESearchDir::FromEnd))
+	{
+		if (Right.IsNumeric())
+		{
+			NameIndex = FCString::Atoi(*Right);
+
+			NodeBaseName = Left;
+		}
+	}
+
+	FName NodeUniqueName{ NodeBaseName };
+	while (Dataflow->GetDataflow()->FindBaseNode(FName(NodeUniqueName)) != nullptr)
+	{
+		NodeUniqueName = FName(NodeBaseName + FString::Printf(TEXT("_%02d"), NameIndex));
+		NameIndex++;
+	}
+
+	return NodeUniqueName;
+}
+
 void SDataflowEdNode::CopyDataflowNodeSettings(TSharedPtr<FDataflowNode> SourceDataflowNode, TSharedPtr<FDataflowNode> TargetDataflowNode)
 {
 	using namespace UE::Transaction;
@@ -305,22 +331,10 @@ UEdGraphNode* FAssetSchemaAction_Dataflow_DuplicateNode_DataflowEdNode::PerformA
 {
 	if (UDataflow* Dataflow = Cast<UDataflow>(ParentGraph))
 	{
-		// Append "_copy" to selected node's name if it doesn't have it
 		FString NodeToDuplicateName = DataflowNodeToDuplicate->GetName().ToString();
-		if (!NodeToDuplicateName.Contains("_copy"))
-		{
-			NodeToDuplicateName.Append("_copy");
-		}
 
 		// Check if that is unique, if not then make it unique with an index postfix
-		const FString NodeBaseName = NodeToDuplicateName;
-		FName NodeUniqueName{ NodeBaseName };
-		int32 NameIndex = 0;
-		while (Dataflow->GetDataflow()->FindBaseNode(FName(NodeUniqueName)) != nullptr)
-		{
-			NodeUniqueName = FName(NodeBaseName + FString::Printf(TEXT("_%d"), NameIndex));
-			NameIndex++;
-		}
+		FName NodeUniqueName = GetNodeUniqueName(Dataflow, NodeToDuplicateName);
 
 		return CreateNode(Dataflow, FromPin, Location, bSelectNewNode, NodeUniqueName, NodeTypeName, DataflowNodeToDuplicate, /*bCopySettings=*/true);
 	}
@@ -398,22 +412,10 @@ UEdGraphNode* FAssetSchemaAction_Dataflow_PasteNode_DataflowEdNode::PerformActio
 {
 	if (UDataflow* Dataflow = Cast<UDataflow>(ParentGraph))
 	{
-		// Append "_copy" to selected node's name if it doesn't have it
 		FString NodeToDuplicateName = NodeName.ToString();
-		if (!NodeToDuplicateName.Contains("_copy"))
-		{
-			NodeToDuplicateName.Append("_copy");
-		}
 
 		// Check if that is unique, if not then make it unique with an index postfix
-		const FString NodeBaseName = NodeToDuplicateName;
-		FName NodeUniqueName{ NodeBaseName };
-		int32 NameIndex = 0;
-		while (Dataflow->GetDataflow()->FindBaseNode(FName(NodeUniqueName)) != nullptr)
-		{
-			NodeUniqueName = FName(NodeBaseName + FString::Printf(TEXT("_%d"), NameIndex));
-			NameIndex++;
-		}
+		FName NodeUniqueName = GetNodeUniqueName(Dataflow, NodeToDuplicateName);
 
 		return CreateNodeFromPaste(Dataflow, FromPin, Location, bSelectNewNode, NodeUniqueName, NodeTypeName, NodeProperties);
 	}
