Dataflow:
- Fixed message for Copy/Cut/Paste and duplicate comment boxes
- For Copy/Cut/Duplicate all comment box properties get set correctly
- Bug fixes

diff --git a/Engine/Plugins/Experimental/Dataflow/Source/DataflowEditor/Private/Dataflow/DataflowEditorCommands.cpp b/Engine/Plugins/Experimental/Dataflow/Source/DataflowEditor/Private/Dataflow/DataflowEditorCommands.cpp
--- a/Engine/Plugins/Experimental/Dataflow/Source/DataflowEditor/Private/Dataflow/DataflowEditorCommands.cpp
+++ b/Engine/Plugins/Experimental/Dataflow/Source/DataflowEditor/Private/Dataflow/DataflowEditorCommands.cpp
@@ -570,19 +570,6 @@ void FDataflowEditorCommands::DuplicateNodes(UDataflow* Graph, const TSharedPtr<
 						}
 					}
 				}
-
-				// Display message stating that nodes were duplicated
-				const int32 NumDuplicatedNodes = SelectedEdNodes.Num();
-				FText MessageFormat;
-				if (NumDuplicatedNodes == 1)
-				{
-					MessageFormat = LOCTEXT("DataflowDuplicatedNodesSingleNode", "{0} node was duplicated");
-				}
-				else
-				{
-					MessageFormat = LOCTEXT("DataflowDuplicatedNodesMultipleNodes", "{0} nodes were duplicated");
-				}
-				ShowNotificationMessage(FText::Format(MessageFormat, NumDuplicatedNodes), SNotificationItem::CS_Success);
 			}
 
 			// Process Comment nodes
@@ -635,6 +622,38 @@ void FDataflowEditorCommands::DuplicateNodes(UDataflow* Graph, const TSharedPtr<
 					}
 				}
 			}
+
+			// Display message stating that nodes were duplicated
+			const int32 NumDuplicatedNodes = SelectedEdNodes.Num();
+			if (NumDuplicatedNodes > 0)
+			{
+				FText MessageFormat;
+				if (NumDuplicatedNodes == 1)
+				{
+					MessageFormat = LOCTEXT("DataflowDuplicatedNodesSingleNode", "{0} node was duplicated");
+				}
+				else
+				{
+					MessageFormat = LOCTEXT("DataflowDuplicatedNodesMultipleNodes", "{0} nodes were duplicated");
+				}
+				ShowNotificationMessage(FText::Format(MessageFormat, NumDuplicatedNodes), SNotificationItem::CS_Success);
+			}
+
+			// Display message stating that comment boxe(s) were duplicated
+			const int32 NumDuplicatedComments = SelectedEdCommentNodes.Num();
+			if (NumDuplicatedComments > 0)
+			{
+				FText MessageFormat;
+				if (NumDuplicatedComments == 1)
+				{
+					MessageFormat = LOCTEXT("DataflowDuplicatedNodesToClipboardSingleComment", "{0} comment was duplicated");
+				}
+				else
+				{
+					MessageFormat = LOCTEXT("DataflowDuplicatedNodesToClipboardMultipleComments", "{0} comments were duplicated");
+				}
+				ShowNotificationMessage(FText::Format(MessageFormat, NumDuplicatedComments), SNotificationItem::CS_Success);
+			}
 		}
 	}
 }
@@ -716,6 +735,7 @@ void FDataflowEditorCommands::CopyNodes(UDataflow* InGraph, const TSharedPtr<SDa
 						CommentNodeData.Position.X = CommentEdNode->NodePosX;
 						CommentNodeData.Position.Y = CommentEdNode->NodePosY;
 						CommentNodeData.Color = CommentEdNode->CommentColor;
+						CommentNodeData.FontSize = CommentEdNode->FontSize;
 
 						CopyPasteContent.CommentNodeData.Add(CommentNodeData);
 					}
@@ -759,16 +779,35 @@ void FDataflowEditorCommands::CopyNodes(UDataflow* InGraph, const TSharedPtr<SDa
 
 				// Display message stating that nodes were copied to clipboard
 				const int32 NumCopiedNodes = CopyPasteContent.NodeData.Num();
-				FText MessageFormat;
-				if (NumCopiedNodes == 1)
+				if (NumCopiedNodes > 0)
 				{
-					MessageFormat = LOCTEXT("DataflowCopiedNodesToClipboardSingleNode", "{0} node was copied to clipboard");
+					FText MessageFormat;
+					if (NumCopiedNodes == 1)
+					{
+						MessageFormat = LOCTEXT("DataflowCopiedNodesToClipboardSingleNode", "{0} node was copied to clipboard");
+					}
+					else
+					{
+						MessageFormat = LOCTEXT("DataflowCopiedNodesToClipboardMultipleNodes", "{0} nodes were copied to clipboard");
+					}
+					ShowNotificationMessage(FText::Format(MessageFormat, NumCopiedNodes), SNotificationItem::CS_Success);
 				}
-				else
+
+				// Display message stating that comment boxe(s) were copied to clipboard
+				const int32 NumCopiedComments = CopyPasteContent.CommentNodeData.Num();
+				if (NumCopiedComments > 0)
 				{
-					MessageFormat = LOCTEXT("DataflowCopiedNodesToClipboardMultipleNodes", "{0} nodes were copied to clipboard");
+					FText MessageFormat;
+					if (NumCopiedComments == 1)
+					{
+						MessageFormat = LOCTEXT("DataflowCopiedNodesToClipboardSingleComment", "{0} comment was copied to clipboard");
+					}
+					else
+					{
+						MessageFormat = LOCTEXT("DataflowCopiedNodesToClipboardMultipleComments", "{0} comments were copied to clipboard");
+					}
+					ShowNotificationMessage(FText::Format(MessageFormat, NumCopiedComments), SNotificationItem::CS_Success);
 				}
-				ShowNotificationMessage(FText::Format(MessageFormat, NumCopiedNodes), SNotificationItem::CS_Success);
 			}
 		}
 	}
@@ -831,7 +870,7 @@ void FDataflowEditorCommands::PasteNodes(UDataflow* Graph, const TSharedPtr<SDat
 
 		// Paste Comment nodes
 		if (CopyPasteContent.CommentNodeData.Num() > 0)
-		{ 
+		{
 			const TSharedPtr<SGraphEditor>& InGraphEditor = (TSharedPtr<SGraphEditor>)DataflowGraphEditor;
 
 			int32 Idx = 0;
@@ -853,6 +892,7 @@ void FDataflowEditorCommands::PasteNodes(UDataflow* Graph, const TSharedPtr<SDat
 					PasteCommentNodeAction->Size.X = CommentNodeData.Size.X + 50; // Make it longer, because the nodes are longer after copying ('_copy' in their name)
 					PasteCommentNodeAction->Size.Y = CommentNodeData.Size.Y + 30;
 					PasteCommentNodeAction->Color = CommentNodeData.Color;
+					PasteCommentNodeAction->FontSize = CommentNodeData.FontSize;
 
 					FVector2D CommentNodeLocation(CommentNodeData.Position);
 					FVector2D NewLocation = CommentNodeLocation + AppliedTranslation;
@@ -893,7 +933,7 @@ void FDataflowEditorCommands::PasteNodes(UDataflow* Graph, const TSharedPtr<SDat
 							FDataflowOutput* OutputConnection = DataflowNodeFrom->FindOutput(OutputputName);
 
 							DataflowGraph->Connect(OutputConnection, InputConnection);
-							
+
 							if (UEdGraphPin* OutputPin = GetPin(EdNodeMap[NodeOut], EEdGraphPinDirection::EGPD_Output, OutputputName))
 							{
 								if (UEdGraphPin* InputPin = GetPin(EdNodeMap[NodeIn], EEdGraphPinDirection::EGPD_Input, InputputName))
@@ -925,16 +965,35 @@ void FDataflowEditorCommands::PasteNodes(UDataflow* Graph, const TSharedPtr<SDat
 
 		// Display message stating that nodes were pasted from clipboard
 		const int32 NumPastedNodes = CopyPasteContent.NodeData.Num();
-		FText MessageFormat;
-		if (NumPastedNodes == 1)
+		if (NumPastedNodes > 0)
 		{
-			MessageFormat = LOCTEXT("DataflowPastedNodesFromClipboardSingleNode", "{0} node was pasted from clipboard");
+			FText MessageFormat;
+			if (NumPastedNodes == 1)
+			{
+				MessageFormat = LOCTEXT("DataflowPastedNodesFromClipboardSingleNode", "{0} node was pasted from clipboard");
+			}
+			else
+			{
+				MessageFormat = LOCTEXT("DataflowPastedNodesFromClipboardMultipleNodes", "{0} nodes were pasted from clipboard");
+			}
+			ShowNotificationMessage(FText::Format(MessageFormat, NumPastedNodes), SNotificationItem::CS_Success);
 		}
-		else
+
+		// Display message stating that comment boxe(s) were pasted to clipboard
+		const int32 NumPastedComments = CopyPasteContent.CommentNodeData.Num();
+		if (NumPastedComments > 0)
 		{
-			MessageFormat = LOCTEXT("DataflowPastedNodesFromClipboardMultipleNodes", "{0} nodes were pasted from clipboard");
+			FText MessageFormat;
+			if (NumPastedComments == 1)
+			{
+				MessageFormat = LOCTEXT("DataflowPastedNodesToClipboardSingleComment", "{0} comment was pasted from clipboard");
+			}
+			else
+			{
+				MessageFormat = LOCTEXT("DataflowPastedNodesToClipboardMultipleComments", "{0} comments were pasted from clipboard");
+			}
+			ShowNotificationMessage(FText::Format(MessageFormat, NumPastedComments), SNotificationItem::CS_Success);
 		}
-		ShowNotificationMessage(FText::Format(MessageFormat, NumPastedNodes), SNotificationItem::CS_Success);
 	}
 }
 
diff --git a/Engine/Plugins/Experimental/Dataflow/Source/DataflowEditor/Private/Dataflow/DataflowSCommentNode.cpp b/Engine/Plugins/Experimental/Dataflow/Source/DataflowEditor/Private/Dataflow/DataflowSCommentNode.cpp
--- a/Engine/Plugins/Experimental/Dataflow/Source/DataflowEditor/Private/Dataflow/DataflowSCommentNode.cpp
+++ b/Engine/Plugins/Experimental/Dataflow/Source/DataflowEditor/Private/Dataflow/DataflowSCommentNode.cpp
@@ -113,6 +113,7 @@ UEdGraphNode* FAssetSchemaAction_Dataflow_DuplicateCommentNode_DataflowEdNode::P
 	CommentTemplate->NodeHeight = CommentNodeToDuplicate->NodeHeight;
 	CommentTemplate->SnapToGrid(GetDefault<UEditorStyleSettings>()->GridSnapSize);
 	CommentTemplate->CommentColor = CommentNodeToDuplicate->CommentColor;
+	CommentTemplate->FontSize = CommentNodeToDuplicate->FontSize;
 
 	CommentTemplate->NodeComment = NewNodeName;
 
@@ -164,6 +165,7 @@ UEdGraphNode* FAssetSchemaAction_Dataflow_PasteCommentNode_DataflowEdNode::Perfo
 	CommentTemplate->NodeHeight = Size.Y;
 	CommentTemplate->SnapToGrid(GetDefault<UEditorStyleSettings>()->GridSnapSize);
 	CommentTemplate->CommentColor = Color;
+	CommentTemplate->FontSize = FontSize;
 
 	CommentTemplate->NodeComment = NewNodeName;
 
diff --git a/Engine/Plugins/Experimental/Dataflow/Source/DataflowEditor/Public/Dataflow/AssetDefinition_DataflowAsset.h b/Engine/Plugins/Experimental/Dataflow/Source/DataflowEditor/Public/Dataflow/AssetDefinition_DataflowAsset.h
--- a/Engine/Plugins/Experimental/Dataflow/Source/DataflowEditor/Public/Dataflow/AssetDefinition_DataflowAsset.h
+++ b/Engine/Plugins/Experimental/Dataflow/Source/DataflowEditor/Public/Dataflow/AssetDefinition_DataflowAsset.h
@@ -59,6 +59,9 @@ struct DATAFLOWEDITOR_API FDataflowCommentNodeData
 
 	UPROPERTY()
 	FVector2D Position = FVector2D::ZeroVector;
+
+	UPROPERTY()
+	int32 FontSize = 18;
 };
 
 USTRUCT()
diff --git a/Engine/Plugins/Experimental/Dataflow/Source/DataflowEditor/Public/Dataflow/DataflowSCommentNode.h b/Engine/Plugins/Experimental/Dataflow/Source/DataflowEditor/Public/Dataflow/DataflowSCommentNode.h
--- a/Engine/Plugins/Experimental/Dataflow/Source/DataflowEditor/Public/Dataflow/DataflowSCommentNode.h
+++ b/Engine/Plugins/Experimental/Dataflow/Source/DataflowEditor/Public/Dataflow/DataflowSCommentNode.h
@@ -97,4 +97,5 @@ public:
 	FName NodeName;
 	FVector2D Size;
 	FLinearColor Color;
+	int32 FontSize;
 };
