Dataflow: highlight "Deprecated" or "Experimental" nodes

diff --git a/Engine/Plugins/Experimental/Dataflow/Source/DataflowEditor/Private/Dataflow/DataflowSNode.cpp b/Engine/Plugins/Experimental/Dataflow/Source/DataflowEditor/Private/Dataflow/DataflowSNode.cpp
--- a/Engine/Plugins/Experimental/Dataflow/Source/DataflowEditor/Private/Dataflow/DataflowSNode.cpp
+++ b/Engine/Plugins/Experimental/Dataflow/Source/DataflowEditor/Private/Dataflow/DataflowSNode.cpp
@@ -131,6 +131,29 @@ TArray<FOverlayWidgetInfo> SDataflowEdNode::GetOverlayWidgets(bool bSelected, co
 	return Widgets;
 }
 
+
+void SDataflowEdNode::UpdateErrorInfo()
+{
+	if (DataflowGraphNode)
+	{
+		if (const TSharedPtr<FDataflowNode> DataflowNode = DataflowGraphNode->GetDataflowNode())
+		{
+			if (DataflowNode->IsExperimental())
+			{
+				ErrorMsg = FString(TEXT("Experimental"));
+				ErrorColor = FAppStyle::GetColor("ErrorReporting.WarningBackgroundColor");
+			} 
+			else if (DataflowNode->IsDeprecated())
+			{
+				ErrorMsg = FString(TEXT("Deprecated"));
+				ErrorColor = FAppStyle::GetColor("ErrorReporting.WarningBackgroundColor");
+			}
+		}
+	}
+	
+}
+
+
 FReply SDataflowEdNode::OnMouseButtonDoubleClick(const FGeometry& InMyGeometry, const FPointerEvent& InMouseEvent)
 {
 	if (GraphNode)
diff --git a/Engine/Plugins/Experimental/Dataflow/Source/DataflowEditor/Public/Dataflow/DataflowSNode.h b/Engine/Plugins/Experimental/Dataflow/Source/DataflowEditor/Public/Dataflow/DataflowSNode.h
--- a/Engine/Plugins/Experimental/Dataflow/Source/DataflowEditor/Public/Dataflow/DataflowSNode.h
+++ b/Engine/Plugins/Experimental/Dataflow/Source/DataflowEditor/Public/Dataflow/DataflowSNode.h
@@ -35,6 +35,7 @@ public:
 	// SGraphNode interface
 	virtual FReply OnMouseButtonDoubleClick(const FGeometry& InMyGeometry, const FPointerEvent& InMouseEvent) override;
 	TArray<FOverlayWidgetInfo> GetOverlayWidgets(bool bSelected, const FVector2D& WidgetSize) const;
+	virtual void UpdateErrorInfo() override;
 
 	static void CopyDataflowNodeSettings(TSharedPtr<FDataflowNode> SourceDataflowNode, TSharedPtr<FDataflowNode> TargetDataflowNode);
 
diff --git a/Engine/Source/Runtime/Experimental/Dataflow/Core/Public/Dataflow/DataflowNode.h b/Engine/Source/Runtime/Experimental/Dataflow/Core/Public/Dataflow/DataflowNode.h
--- a/Engine/Source/Runtime/Experimental/Dataflow/Core/Public/Dataflow/DataflowNode.h
+++ b/Engine/Source/Runtime/Experimental/Dataflow/Core/Public/Dataflow/DataflowNode.h
@@ -85,6 +85,9 @@ struct FDataflowNode
 	UE_DEPRECATED(5.4, "FDataflowNode::CopyNodeProperties is deprecated.")
 	DATAFLOWCORE_API void CopyNodeProperties(const TSharedPtr<FDataflowNode> CopyFromDataflowNode);
 
+	virtual bool IsDeprecated() { return false; }
+	virtual bool IsExperimental() { return false; }
+
 	//
 	// Connections
 	//
