Dataflow : Restructure Viewports to support separate Construction and Simulation Clients

diff --git a/Engine/Plugins/Experimental/Dataflow/Source/DataflowEditor/Private/Dataflow/CollectionSpreadSheetWidget.cpp b/Engine/Plugins/Experimental/Dataflow/Source/DataflowEditor/Private/Dataflow/CollectionSpreadSheetWidget.cpp
--- a/Engine/Plugins/Experimental/Dataflow/Source/DataflowEditor/Private/Dataflow/CollectionSpreadSheetWidget.cpp
+++ b/Engine/Plugins/Experimental/Dataflow/Source/DataflowEditor/Private/Dataflow/CollectionSpreadSheetWidget.cpp
@@ -1,6 +1,6 @@
 // Copyright Epic Games, Inc. All Rights Reserved.
 
-#include "Dataflow/CollectionSpreadSheetWidget.h"
+#include "Dataflow/DataflowCollectionSpreadSheetWidget.h"
 #include "Widgets/Input/SButton.h"
 #include "Styling/StarshipCoreStyle.h"
 #include "Widgets/Layout/SScrollBox.h"
diff --git a/Engine/Plugins/Experimental/Dataflow/Source/DataflowEditor/Private/Dataflow/DataflowConstructionViewportClient.cpp b/Engine/Plugins/Experimental/Dataflow/Source/DataflowEditor/Private/Dataflow/DataflowConstructionViewportClient.cpp
new file mode 100644
--- /dev/null
+++ b/Engine/Plugins/Experimental/Dataflow/Source/DataflowEditor/Private/Dataflow/DataflowConstructionViewportClient.cpp
@@ -0,0 +1,217 @@
+// Copyright Epic Games, Inc. All Rights Reserved.
+#include "Dataflow/DataflowConstructionViewportClient.h"
+
+#include "AssetEditorModeManager.h"
+#include "Dataflow/DataflowObject.h"
+#include "Dataflow/DataflowEditorMode.h"
+#include "Dataflow/DataflowEditorToolkit.h"
+#include "Dataflow/DataflowEditorCollectionComponent.h"
+#include "Dataflow/DataflowEngineSceneHitProxies.h"
+#include "Dataflow/DataflowGraphEditor.h"
+#include "Dataflow/DataflowEditorScenes.h"
+#include "EditorModeManager.h"
+#include "EdModeInteractiveToolsContext.h"
+#include "GraphEditor.h"
+#include "PreviewScene.h"
+#include "Selection.h"
+#include "SGraphPanel.h"
+#include "SNodePanel.h"
+
+
+
+FDataflowConstructionViewportClient::FDataflowConstructionViewportClient(FEditorModeTools* InModeTools,
+                                                             FPreviewScene* InPreviewScene,  const bool bCouldTickScene,
+                                                             const TWeakPtr<SEditorViewport> InEditorViewportWidget)
+	: FEditorViewportClient(InModeTools, InPreviewScene, InEditorViewportWidget)
+{
+	// We want our near clip plane to be quite close so that we can zoom in further.
+	OverrideNearClipPlane(KINDA_SMALL_NUMBER);
+
+	EngineShowFlags.SetSelectionOutline(true);
+	EngineShowFlags.EnableAdvancedFeatures();
+
+	PreviewScene = static_cast<FDataflowPreviewScene*>(InPreviewScene);
+	bEnableSceneTicking = bCouldTickScene;
+}
+
+void FDataflowConstructionViewportClient::SetDataflowEditorToolkit(TWeakPtr<FDataflowEditorToolkit> InDataflowEditorToolkitPtr)
+{
+	DataflowEditorToolkitPtr = InDataflowEditorToolkitPtr;
+}
+
+void FDataflowConstructionViewportClient::SetToolCommandList(TWeakPtr<FUICommandList> InToolCommandList)
+{
+	ToolCommandList = InToolCommandList;
+}
+
+//const UInputBehaviorSet* FDataflowConstructionViewportClient::GetInputBehaviors() const
+//{
+//	return BehaviorSet;
+//}
+
+void FDataflowConstructionViewportClient::Tick(float DeltaSeconds)
+{
+	Super::Tick(DeltaSeconds);
+	if (PreviewScene)
+	{
+		PreviewScene->TickDataflowScene(DeltaSeconds);
+	}
+}
+
+void FDataflowConstructionViewportClient::ProcessClick(FSceneView& View, HHitProxy* HitProxy, FKey Key, EInputEvent Event, uint32 HitX, uint32 HitY)
+{
+	Super::ProcessClick(View, HitProxy, Key, Event, HitX, HitY);
+
+	auto EnableToolForSelectedNode = [&](USelection* SelectedComponents)
+	{
+		if (TSharedPtr<FDataflowEditorToolkit> DataflowEditorToolkit = DataflowEditorToolkitPtr.Pin())
+		{
+			if (PreviewScene && PreviewScene->GetDataflowModeManager())
+			{
+				if (UDataflowEditorMode* DataflowMode = Cast<UDataflowEditorMode>(PreviewScene->GetDataflowModeManager()->GetActiveScriptableMode(UDataflowEditorMode::EM_DataflowEditorModeId)))
+				{
+					if (TSharedPtr<SDataflowGraphEditor> GraphEditor = DataflowEditorToolkit->GetDataflowGraphEditor())
+					{
+						if (UEdGraphNode* SelectedNode = GraphEditor->GetSingleSelectedNode())
+						{
+							if (SelectedComponents && SelectedComponents->Num() == 1)
+							{
+								if (UDataflowEditorCollectionComponent* CollectionComponent =
+									Cast< UDataflowEditorCollectionComponent>(SelectedComponents->GetSelectedObject(0)))
+								{
+									if (CollectionComponent->Node == SelectedNode)
+									{
+										// Start the corresponding tool
+										DataflowMode->StartToolForSelectedNode(SelectedNode);
+									}
+								}
+							}
+						}
+					}
+				}
+			}
+		}
+	};
+
+	auto UpdateSelectedComponentInViewport = [&](USelection* SelectedComponents)
+	{
+		TArray<UPrimitiveComponent*> PreviouslySelectedComponents;
+		SelectedComponents->GetSelectedObjects<UPrimitiveComponent>(PreviouslySelectedComponents);
+
+		SelectedComponents->Modify();
+		SelectedComponents->BeginBatchSelectOperation();
+
+		SelectedComponents->DeselectAll();
+
+		if (HitProxy && HitProxy->IsA(HActor::StaticGetType()))
+		{
+			const HActor* ActorProxy = static_cast<HActor*>(HitProxy);
+			if (ActorProxy && ActorProxy->PrimComponent && ActorProxy->Actor)
+			{
+				UPrimitiveComponent* Component = const_cast<UPrimitiveComponent*>(ActorProxy->PrimComponent.Get());
+				SelectedComponents->Select(Component);
+				Component->PushSelectionToProxy();
+			}
+		}
+
+		SelectedComponents->EndBatchSelectOperation();
+
+		for (UPrimitiveComponent* const Component : PreviouslySelectedComponents)
+		{
+			Component->PushSelectionToProxy();
+		}
+
+	};
+
+	auto SelectSingleNodeInGraph = [&](TObjectPtr<const UDataflowEdNode> Node)
+	{
+		if (TSharedPtr<FDataflowEditorToolkit> DataflowEditorToolkit = DataflowEditorToolkitPtr.Pin())
+		{
+			if (TSharedPtr<SDataflowGraphEditor> GraphEditor = DataflowEditorToolkit->GetDataflowGraphEditor())
+			{
+				GraphEditor->GetGraphPanel()->SelectionManager.SelectSingleNode((UObject*)Node.Get());
+			}
+		}
+	};
+	
+	auto IsInteractiveToolActive = [&]()
+	{
+		if (UDataflowEditorMode* DataflowMode = Cast<UDataflowEditorMode>(
+			PreviewScene->GetDataflowModeManager()->GetActiveScriptableMode(UDataflowEditorMode::EM_DataflowEditorModeId)))
+		{
+			if (UEditorInteractiveToolsContext* const ToolsContext = DataflowMode->GetInteractiveToolsContext())
+			{
+				return ToolsContext->HasActiveTool();
+			}
+		}
+		return false;
+	};
+
+	if (!IsInteractiveToolActive())
+	{
+		if (USelection* SelectedComponents = ModeTools->GetSelectedComponents())
+		{
+			UpdateSelectedComponentInViewport(SelectedComponents);
+
+			if (bool bIsAltKeyDown = Viewport->KeyState(EKeys::LeftAlt) || Viewport->KeyState(EKeys::RightAlt))
+			{
+				if (UDataflowEditorCollectionComponent* DataflowComponent
+					= SelectedComponents->GetBottom<UDataflowEditorCollectionComponent>())
+				{
+					SelectSingleNodeInGraph(DataflowComponent->Node);
+				}
+			}
+
+			EnableToolForSelectedNode(SelectedComponents);
+		}
+	}
+}
+
+void FDataflowConstructionViewportClient::SetConstructionViewMode(Dataflow::EDataflowPatternVertexType InViewMode)
+{
+	// @todo(Dataflow) : Add support for Sim2D
+	//const bool bSwitching2D3D = (ConstructionViewMode == Dataflow::EDataflowPatternVertexType::Sim2D) != (InViewMode == Dataflow::EDataflowPatternVertexType::Sim2D);
+	//if (bSwitching2D3D)
+	//{
+	//	Swap(SavedInactiveViewTransform, ViewTransformPerspective);
+	//}
+
+	ConstructionViewMode = Dataflow::EDataflowPatternVertexType::Sim3D;
+
+	
+	//if (ConstructionViewMode == EDataflowPatternVertexType::Sim2D)
+	//{
+	//	for (UInputBehavior* const Behavior : BehaviorsFor2DMode)
+	//	{
+	//		BehaviorSet->Add(Behavior);
+	//	}
+	//
+	//	const double AbsZ = FMath::Abs(ViewTransformPerspective.GetLocation().Z);
+	//	constexpr double CameraFarPlaneWorldZ = -10.0;
+	//	constexpr double CameraNearPlaneProportionZ = 0.8;
+	//	OverrideFarClipPlane(static_cast<float>(AbsZ - CameraFarPlaneWorldZ));
+	//	OverrideNearClipPlane(static_cast<float>(AbsZ * (1.0 - CameraNearPlaneProportionZ)));
+	//}
+	//else
+	//{
+	OverrideFarClipPlane(0);
+	OverrideNearClipPlane(UE_KINDA_SMALL_NUMBER);
+	//}
+
+	//ModeTools->GetInteractiveToolsContext()->InputRouter->DeregisterSource(this);
+	//ModeTools->GetInteractiveToolsContext()->InputRouter->RegisterSource(this);
+
+}
+
+
+Dataflow::EDataflowPatternVertexType FDataflowConstructionViewportClient::GetConstructionViewMode() const
+{
+	return Dataflow::EDataflowPatternVertexType::Sim3D;// ConstructionViewMode;
+}
+
+void FDataflowConstructionViewportClient::AddReferencedObjects(FReferenceCollector& Collector)
+{
+	Super::AddReferencedObjects(Collector);
+	Collector.AddReferencedObject(BehaviorSet);
+}
+
diff --git a/Engine/Plugins/Experimental/Dataflow/Source/DataflowEditor/Private/Dataflow/DataflowConstructionViewportToolbar.cpp b/Engine/Plugins/Experimental/Dataflow/Source/DataflowEditor/Private/Dataflow/DataflowConstructionViewportToolbar.cpp
new file mode 100644
--- /dev/null
+++ b/Engine/Plugins/Experimental/Dataflow/Source/DataflowEditor/Private/Dataflow/DataflowConstructionViewportToolbar.cpp
@@ -0,0 +1,10 @@
+// Copyright Epic Games, Inc. All Rights Reserved.
+#include "Dataflow/DataflowConstructionViewportToolbar.h"
+#include "Dataflow/DataflowEditorViewport.h"
+
+void SDataflowConstructionViewportSelectionToolBar::Construct(const FArguments& InArgs, TSharedPtr<SDataflowEditorViewport> InDataflowViewport)
+{
+	EditorViewport = InDataflowViewport;
+	
+	SCommonEditorViewportToolbarBase::Construct(SCommonEditorViewportToolbarBase::FArguments(), InDataflowViewport);
+}
\ No newline at end of file
diff --git a/Engine/Plugins/Experimental/Dataflow/Source/DataflowEditor/Private/Dataflow/DataflowEditorMode.cpp b/Engine/Plugins/Experimental/Dataflow/Source/DataflowEditor/Private/Dataflow/DataflowEditorMode.cpp
--- a/Engine/Plugins/Experimental/Dataflow/Source/DataflowEditor/Private/Dataflow/DataflowEditorMode.cpp
+++ b/Engine/Plugins/Experimental/Dataflow/Source/DataflowEditor/Private/Dataflow/DataflowEditorMode.cpp
@@ -10,13 +10,14 @@
 #include "DataflowEditorTools/DataflowEditorWeightMapPaintTool.h"
 #include "Dataflow/DataflowComponentToolTarget.h"
 #include "Dataflow/DataflowCollectionAddScalarVertexPropertyNode.h"
+#include "Dataflow/DataflowConstructionViewportClient.h"
 #include "Dataflow/DataflowEditor.h"
 #include "Dataflow/DataflowContent.h"
 #include "Dataflow/DataflowEditorCommands.h"
 #include "Dataflow/DataflowEditorModeToolkit.h"
-#include "Dataflow/DataflowEditorViewportClient.h"
 #include "Dataflow/DataflowGraphEditor.h"
-#include "Dataflow/DataflowPreviewScene.h"
+#include "Dataflow/DataflowEditorScenes.h"
+#include "Dataflow/DataflowSimulationViewportClient.h"
 #include "Dataflow/DataflowSNode.h"
 #include "Dataflow/DataflowToolTarget.h"
 #include "EditorModeManager.h"
@@ -213,9 +214,9 @@ bool UDataflowEditorMode::ShouldToolStartBeAllowed(const FString& ToolIdentifier
 		return false;
 	}
 
-	if (PreviewScene && PreviewScene->GetDataflowModeManager() && PreviewScene->GetDataflowModeManager()->GetInteractiveToolsContext())
+	if (ConstructionScene && ConstructionScene->GetDataflowModeManager() && ConstructionScene->GetDataflowModeManager()->GetInteractiveToolsContext())
 	{
-		if (PreviewScene->GetDataflowModeManager()->GetInteractiveToolsContext()->HasActiveTool())
+		if (ConstructionScene->GetDataflowModeManager()->GetInteractiveToolsContext()->HasActiveTool())
 		{
 			return false;
 		}
@@ -251,7 +252,7 @@ void UDataflowEditorMode::OnToolEnded(UInteractiveToolManager* Manager, UInterac
 	}
 	else
 	{
-		PreviewScene->ResetConstructionScene();
+		ConstructionScene->ResetConstructionScene();
 	}
 
 	if (TSharedPtr<SDataflowGraphEditor> GraphEditor = DataflowGraphEditor.Pin())
@@ -291,17 +292,17 @@ void UDataflowEditorMode::Exit()
 {
 	UActorComponent::MarkRenderStateDirtyEvent.RemoveAll(this);
 
-	PreviewScene->ResetConstructionScene();
-	PreviewScene = nullptr;
+	ConstructionScene->ResetConstructionScene();
+	ConstructionScene = nullptr;
 
 	Super::Exit();
 }
 
-void UDataflowEditorMode::SetDataflowConstructionScene(FDataflowConstructionScene* InPreviewScene)
+void UDataflowEditorMode::SetDataflowConstructionScene(FDataflowConstructionScene* InConstructionScene)
 {
-	PreviewScene = InPreviewScene;
+	ConstructionScene = InConstructionScene;
 
-	UEditorInteractiveToolsContext* const PreviewToolsContext = PreviewScene->GetDataflowModeManager()->GetInteractiveToolsContext();
+	UEditorInteractiveToolsContext* const PreviewToolsContext = ConstructionScene->GetDataflowModeManager()->GetInteractiveToolsContext();
 	UInteractiveToolManager* const PreviewToolManager = PreviewToolsContext->ToolManager;
 	//#todo(brice): Make sure AddToolTargetFactories has been called. 
 	//PreviewToolsContext->TargetManager->AddTargetFactory(NewObject<UClothComponentToolTargetFactory>(PreviewToolManager));
@@ -322,7 +323,7 @@ void UDataflowEditorMode::SetDataflowConstructionScene(FDataflowConstructionScen
 void UDataflowEditorMode::CreateToolTargets(const TArray<TObjectPtr<UObject>>& AssetsIn)
 {
 	ToolTargets.Reset();
-	if (TObjectPtr<UDataflowBaseContent> EditorContent = PreviewScene->GetDataflowContent())
+	if (TObjectPtr<UDataflowBaseContent> EditorContent = ConstructionScene->GetDataflowContent())
 	{
 		if (UToolTarget* Target = GetInteractiveToolsContext()->TargetManager->BuildTarget(EditorContent, GetToolTargetRequirements()))
 		{
@@ -350,7 +351,7 @@ bool UDataflowEditorMode::IsComponentSelected(const UPrimitiveComponent* InCompo
 
 void UDataflowEditorMode::RefocusRestSpaceViewportClient()
 {
-	TSharedPtr<FDataflowEditorViewportClient, ESPMode::ThreadSafe> PinnedVC = ConstructionViewportClient.Pin();
+	TSharedPtr<FDataflowConstructionViewportClient, ESPMode::ThreadSafe> PinnedVC = ConstructionViewportClient.Pin();
 	if (PinnedVC.IsValid())
 	{
 		// This will happen in FocusViewportOnBox anyways; do it now to get a consistent end result
@@ -380,7 +381,7 @@ void UDataflowEditorMode::RefocusRestSpaceViewportClient()
 void UDataflowEditorMode::FirstTimeFocusRestSpaceViewport()
 {
 	// If this is the first time seeing a valid 2D or 3D mesh, refocus the camera on it.
-	const bool bIsValid = (PreviewScene->HasRenderableGeometry());
+	const bool bIsValid = (ConstructionScene->HasRenderableGeometry());
 	const bool bIs2D = ConstructionViewMode == Dataflow::EDataflowPatternVertexType::Sim3D;
 
 	if (bIsValid)
@@ -407,11 +408,11 @@ void UDataflowEditorMode::InitializeTargets(const TArray<TObjectPtr<UObject>>& O
 	// @todo(brice) : What are the ToolTargets storing?
 	// ... for(ToolTarget& : ToolTargets){
 	// ... UE::ToolTarget::GetDynamicMeshCopy(Target)
-	// ... UE::ToolTarget::GetMaterialSet(Target).Materials for PreviewScene->AddDynamicMeshComponent
+	// ... UE::ToolTarget::GetMaterialSet(Target).Materials for ConstructionScene->AddDynamicMeshComponent
 	// ... }
 
 	// @todo(michael) : do we need to update the construction scene?
-	PreviewScene->UpdateConstructionScene();
+	ConstructionScene->UpdateConstructionScene();
 }
 
 void UDataflowEditorMode::ModeTick(float DeltaTime)
@@ -459,13 +460,13 @@ void UDataflowEditorMode::RestSpaceViewportResized(FViewport* RestspaceViewport,
 
 FBox UDataflowEditorMode::SceneBoundingBox() const
 {
-	return PreviewScene->GetBoundingBox();
+	return ConstructionScene->GetBoundingBox();
 }
 
 FBox UDataflowEditorMode::SelectionBoundingBox() const
 {
 	// If the selection is on the GetBoundingBox is automatically computing the selection one
-	FBox Bounds = PreviewScene->GetBoundingBox();
+	FBox Bounds = ConstructionScene->GetBoundingBox();
 	if (Bounds.IsValid)
 	{
 		return Bounds;
@@ -495,9 +496,9 @@ void UDataflowEditorMode::SetConstructionViewMode(Dataflow::EDataflowPatternVert
 	}
 
 	ConstructionViewMode = InMode;
-	PreviewScene->UpdateConstructionScene();
+	ConstructionScene->UpdateConstructionScene();
 
-	const TSharedPtr<FDataflowEditorViewportClient> VC = ConstructionViewportClient.Pin();
+	const TSharedPtr<FDataflowConstructionViewportClient> VC = ConstructionViewportClient.Pin();
 	if (VC.IsValid())
 	{
 		VC->SetConstructionViewMode(ConstructionViewMode);
@@ -553,7 +554,7 @@ void UDataflowEditorMode::ToggleConstructionViewWireframe()
 {
 	check(false);
 	bConstructionViewWireframe = !bConstructionViewWireframe;
-	PreviewScene->UpdateConstructionScene();
+	ConstructionScene->UpdateConstructionScene();
 }
 
 bool UDataflowEditorMode::CanSetConstructionViewWireframeActive() const
@@ -571,11 +572,11 @@ bool UDataflowEditorMode::CanSetConstructionViewWireframeActive() const
 	return DataflowToolBuilder->CanSetConstructionViewWireframeActive();
 }
 
-void UDataflowEditorMode::SetRestSpaceViewportClient(TWeakPtr<FDataflowEditorViewportClient, ESPMode::ThreadSafe> InViewportClient)
+void UDataflowEditorMode::SetRestSpaceViewportClient(TWeakPtr<FDataflowConstructionViewportClient, ESPMode::ThreadSafe> InViewportClient)
 {
 	ConstructionViewportClient = InViewportClient;
 
-	TSharedPtr<FDataflowEditorViewportClient> VC = ConstructionViewportClient.Pin();
+	TSharedPtr<FDataflowConstructionViewportClient> VC = ConstructionViewportClient.Pin();
 	if (VC.IsValid())
 	{
 		VC->SetConstructionViewMode(ConstructionViewMode);
@@ -590,9 +591,9 @@ void UDataflowEditorMode::SetRestSpaceViewportClient(TWeakPtr<FDataflowEditorVie
 
 void UDataflowEditorMode::InitializeContextObject()
 {
-	check(PreviewScene);
+	check(ConstructionScene);
 
-	if (TObjectPtr<UDataflowBaseContent> DataflowContent = PreviewScene->GetDataflowContent())
+	if (TObjectPtr<UDataflowBaseContent> DataflowContent = ConstructionScene->GetDataflowContent())
 	{
 		UEditorInteractiveToolsContext* const RestSpaceToolsContext = GetInteractiveToolsContext();
 
@@ -690,7 +691,7 @@ UEdGraphNode* UDataflowEditorMode::CreateNewNode(const FName& NewNodeTypeName)
 		return nullptr;
 	}
 
-	if (TObjectPtr<UDataflowBaseContent> EditorContent = PreviewScene->GetDataflowContent())
+	if (TObjectPtr<UDataflowBaseContent> EditorContent = ConstructionScene->GetDataflowContent())
 	{
 		if (TObjectPtr<UDataflow> DataflowGraph = EditorContent->GetDataflowAsset())
 		{
@@ -707,7 +708,7 @@ UEdGraphNode* UDataflowEditorMode::CreateNewNode(const FName& NewNodeTypeName)
 
 UEdGraphNode* UDataflowEditorMode::CreateAndConnectNewNode(const FName& NewNodeTypeName, UEdGraphNode& UpstreamNode, const FName& ConnectionTypeName, const FName& NewNodeConnectionName)
 {
-	if (TObjectPtr<UDataflowBaseContent> EditorContent = PreviewScene->GetDataflowContent())
+	if (TObjectPtr<UDataflowBaseContent> EditorContent = ConstructionScene->GetDataflowContent())
 	{
 		if (TObjectPtr<UDataflow> DataflowGraph = EditorContent->GetDataflowAsset())
 		{
diff --git a/Engine/Plugins/Experimental/Dataflow/Source/DataflowEditor/Private/Dataflow/DataflowEditorModeToolkit.cpp b/Engine/Plugins/Experimental/Dataflow/Source/DataflowEditor/Private/Dataflow/DataflowEditorModeToolkit.cpp
--- a/Engine/Plugins/Experimental/Dataflow/Source/DataflowEditor/Private/Dataflow/DataflowEditorModeToolkit.cpp
+++ b/Engine/Plugins/Experimental/Dataflow/Source/DataflowEditor/Private/Dataflow/DataflowEditorModeToolkit.cpp
@@ -7,7 +7,7 @@
 #include "Dataflow/DataflowEditorMode.h"
 #include "Dataflow/DataflowEditorStyle.h"
 #include "Dataflow/DataflowEditorViewport.h"
-#include "Dataflow/DataflowPreviewScene.h"
+#include "Dataflow/DataflowEditorScenes.h"
 #include "EdModeInteractiveToolsContext.h"
 #include "Framework/Application/SlateApplication.h"
 #include "InteractiveToolManager.h"
@@ -174,7 +174,7 @@ SBaseCharacterFXEditorViewport* FDataflowEditorModeToolkit::GetViewportWidgetFor
 					//}
 					//else
 					//{
-					if (const TSharedPtr<SDataflowEditorViewport> Widget = RestSpaceViewportWidget.Pin())
+					if (const TSharedPtr<SDataflowEditorViewport> Widget = ConstructionViewportWidget.Pin())
 					{
 						return Widget.Get();
 					}
@@ -283,9 +283,9 @@ void FDataflowEditorModeToolkit::OnToolEnded(UInteractiveToolManager* Manager, U
 
 }
 
-void FDataflowEditorModeToolkit::SetRestSpaceViewportWidget(TWeakPtr<SDataflowEditorViewport> InRestSpaceViewportWidget)
+void FDataflowEditorModeToolkit::SetConstructionViewportWidget(TWeakPtr<SDataflowEditorViewport> InConstructionViewportWidget)
 {
-	RestSpaceViewportWidget = InRestSpaceViewportWidget;
+	ConstructionViewportWidget = InConstructionViewportWidget;
 }
 
 
diff --git a/Engine/Plugins/Experimental/Dataflow/Source/DataflowEditor/Private/Dataflow/DataflowEditorScenes.cpp b/Engine/Plugins/Experimental/Dataflow/Source/DataflowEditor/Private/Dataflow/DataflowEditorScenes.cpp
new file mode 100644
--- /dev/null
+++ b/Engine/Plugins/Experimental/Dataflow/Source/DataflowEditor/Private/Dataflow/DataflowEditorScenes.cpp
@@ -0,0 +1,384 @@
+// Copyright Epic Games, Inc. All Rights Reserved.
+
+#include "Dataflow/DataflowEditorScenes.h"
+
+#include "Animation/AnimSingleNodeInstance.h"
+#include "AssetEditorModeManager.h"
+#include "Components/DynamicMeshComponent.h"
+#include "Dataflow/CollectionRenderingPatternUtility.h"
+#include "Dataflow/DataflowEditor.h"
+#include "Dataflow/DataflowEditorCollectionComponent.h"
+#include "Dataflow/DataflowContent.h"
+#include "Dataflow/DataflowEditorStyle.h"
+#include "Dataflow/DataflowEditorUtil.h"
+#include "Drawing/MeshElementsVisualizer.h"
+#include "Elements/Framework/EngineElementsLibrary.h"
+#include "InteractiveTool.h"
+#include "Selection.h"
+
+#define LOCTEXT_NAMESPACE "FDataflowPreviewScene"
+
+
+bool bDataflowShowFloorDefault = false;
+FAutoConsoleVariableRef CVARDataflowShowFloorDefault(TEXT("p.Dataflow.Editor.ShowFloor"), bDataflowShowFloorDefault, TEXT("Show the floor in the dataflow editor[def:false]"));
+
+
+FDataflowPreviewScene::FDataflowPreviewScene(FPreviewScene::ConstructionValues ConstructionValues, TObjectPtr<UDataflowBaseContent> InEditorContent) 
+	: FAdvancedPreviewScene(ConstructionValues), DataflowContent(InEditorContent)
+{
+	check(DataflowContent);
+	SetFloorVisibility(bDataflowShowFloorDefault, true);
+
+	RootSceneActor = GetWorld()->SpawnActor<AActor>(AActor::StaticClass());
+}
+
+FDataflowPreviewScene::~FDataflowPreviewScene()
+{}
+
+void FDataflowPreviewScene::AddReferencedObjects(FReferenceCollector& Collector)
+{
+	FAdvancedPreviewScene::AddReferencedObjects(Collector);
+	Collector.AddReferencedObject(DataflowContent);
+	Collector.AddReferencedObject(RootSceneActor);
+	if (GetDataflowContent())
+	{
+		GetDataflowContent()->AddContentObjects(Collector);
+	}
+}
+
+bool FDataflowPreviewScene::IsComponentSelected(const UPrimitiveComponent* InComponent) const
+{
+	if(DataflowModeManager.IsValid())
+	{
+		if (const UTypedElementSelectionSet* const TypedElementSelectionSet = DataflowModeManager->GetEditorSelectionSet())
+		{
+			if (const FTypedElementHandle ComponentElement = UEngineElementsLibrary::AcquireEditorComponentElementHandle(InComponent))
+			{
+				const bool bElementSelected = TypedElementSelectionSet->IsElementSelected(ComponentElement, FTypedElementIsSelectedOptions());
+				return bElementSelected;
+			}
+		}
+	}
+	return false;
+}
+
+FBox FDataflowPreviewScene::GetBoundingBox() const
+{
+	FBox SceneBounds(ForceInitToZero);
+	if(DataflowModeManager.IsValid())
+	{
+		USelection* const SelectedComponents = DataflowModeManager->GetSelectedComponents();
+
+		TArray<TWeakObjectPtr<UObject>> SelectedObjects;
+		const int32 NumSelected = SelectedComponents->GetSelectedObjects(SelectedObjects);
+		
+		if(NumSelected > 0)
+		{
+			for(const TWeakObjectPtr<UObject> SelectedObject : SelectedObjects)
+			{
+				if(const UPrimitiveComponent* SelectedComponent = Cast<UPrimitiveComponent>(SelectedObject))
+				{
+					SceneBounds += SelectedComponent->Bounds.GetBox();
+				}
+			}
+		}
+		else
+		{
+			SceneBounds += RootSceneActor->GetComponentsBoundingBox(true);
+		}
+	}
+	return SceneBounds;
+}
+
+FDataflowConstructionScene::FDataflowConstructionScene(FPreviewScene::ConstructionValues ConstructionValues, TObjectPtr<UDataflowBaseContent> InEditorContent) 
+	: FDataflowPreviewScene(ConstructionValues,InEditorContent)
+{}
+
+FDataflowConstructionScene::~FDataflowConstructionScene()
+{
+	ResetDynamicMeshComponents();
+}
+
+void FDataflowConstructionScene::AddReferencedObjects(FReferenceCollector& Collector)
+{
+	FDataflowPreviewScene::AddReferencedObjects(Collector);
+
+	Collector.AddReferencedObjects(DynamicMeshComponents);
+	Collector.AddReferencedObjects(WireframeElements);
+}
+
+FORCEINLINE Dataflow::FTimestamp LatestTimestamp(const UDataflow* Dataflow, const ::Dataflow::FContext* Context)
+{
+	if (Dataflow && Context)
+	{
+		return FMath::Max(Dataflow->GetRenderingTimestamp().Value, Context->GetTimestamp().Value);
+	}
+	return ::Dataflow::FTimestamp::Invalid;
+}
+
+void FDataflowConstructionScene::TickDataflowScene(const float DeltaSeconds)
+{
+	if (const TSharedPtr<Dataflow::FContext> DataflowContext = DataflowContent->GetDataflowContext())
+	{
+		if (const UDataflow* Dataflow = DataflowContent->GetDataflowAsset())
+		{
+			const Dataflow::FTimestamp SystemTimestamp = LatestTimestamp(Dataflow, DataflowContext.Get());
+			if (SystemTimestamp >= DataflowContent->GetLastModifiedTimestamp() || DataflowContent->IsDirty())
+			{
+				DataflowContent->SetLastModifiedTimestamp(SystemTimestamp.Value + 1);
+
+				if(DataflowContent->IsDirty())
+				{
+					UpdateConstructionScene();
+				}
+			}
+		}
+	}
+	for (TObjectPtr<UInteractiveToolPropertySet>& Propset : PropertyObjectsToTick)
+	{
+		if (Propset)
+		{
+			if (Propset->IsPropertySetEnabled())
+			{
+				Propset->CheckAndUpdateWatched();
+			}
+			else
+			{
+				Propset->SilentUpdateWatched();
+			}
+		}
+	}
+
+	for (FRenderWireElement Elem : WireframeElements)
+	{
+		Elem.Value->OnTick(DeltaSeconds);
+	}
+}
+
+void FDataflowConstructionScene::UpdateDynamicMeshComponents()
+{
+	using namespace UE::Geometry;//FDynamicMesh3
+
+	// The preview scene for the construction view will be
+	// cleared and rebuilt from scratch. This will genrate a 
+	// list of UPrimitiveComponents for rendering.
+	ResetDynamicMeshComponents();
+
+	if (DataflowContent)
+	{
+		const TObjectPtr<UDataflow>& DataflowAsset = DataflowContent->GetDataflowAsset();
+		const TSharedPtr<Dataflow::FEngineContext>& DataflowContext = DataflowContent->GetDataflowContext();
+		if(DataflowAsset && DataflowContext)
+		{
+			for (TObjectPtr<const UDataflowEdNode> Target : DataflowAsset->GetRenderTargets())
+			{
+				if (Target)
+				{
+					TSharedPtr<FManagedArrayCollection> RenderCollection(new FManagedArrayCollection);
+					GeometryCollection::Facades::FRenderingFacade Facade(*RenderCollection);
+					Facade.DefineSchema();
+
+					Target->Render(Facade, DataflowContext);
+
+					int32 NumGeometry = Facade.NumGeometry();
+					for (int32 MeshIndex = 0; MeshIndex < NumGeometry; MeshIndex++)
+					{
+						FDynamicMesh3 DynamicMesh;
+						Dataflow::Conversion::RenderingFacadeToDynamicMesh(Facade, MeshIndex, DynamicMesh);
+
+						if (DynamicMesh.VertexCount())
+						{
+							if (Target == DataflowContent->GetPrimarySelectedNode())
+							{
+								DataflowContent->SetPrimaryRenderCollection(RenderCollection);
+							}
+
+							AddDynamicMeshComponent({Target, MeshIndex }, MoveTemp(DynamicMesh), {});
+						}
+					}
+				}
+			}
+		}
+	}
+}
+
+void FDataflowConstructionScene::ResetDynamicMeshComponents()
+{
+	USelection* SelectedComponents = DataflowModeManager->GetSelectedComponents();
+	for (FRenderElement RenderElement : DynamicMeshComponents)
+	{
+		TObjectPtr<UDynamicMeshComponent>& DynamicMeshComponent = RenderElement.Value;
+
+		DynamicMeshComponent->SelectionOverrideDelegate.Unbind();
+		if (SelectedComponents->IsSelected(DynamicMeshComponent))
+		{
+			SelectedComponents->Deselect(DynamicMeshComponent);
+			DynamicMeshComponent->PushSelectionToProxy();
+		}
+		RemoveComponent(DynamicMeshComponent);
+	}
+	DynamicMeshComponents.Reset();
+}
+
+TObjectPtr<UDynamicMeshComponent>& FDataflowConstructionScene::AddDynamicMeshComponent(FDataflowRenderKey InKey, UE::Geometry::FDynamicMesh3&& DynamicMesh, const TArray<UMaterialInterface*>& MaterialSet)
+{
+	TObjectPtr<UDataflowEditorCollectionComponent> DynamicMeshComponent = NewObject<UDataflowEditorCollectionComponent>(RootSceneActor);
+	DynamicMeshComponent->MeshIndex = InKey.Value;
+	DynamicMeshComponent->Node = InKey.Key;;
+	DynamicMeshComponent->SetMesh(MoveTemp(DynamicMesh));
+	
+	// @todo(Material) This is just to have a material, we should transfer the materials from the assets if they have them. 
+	if (DataflowContent && DataflowContent->GetDataflowAsset() && DataflowContent->GetDataflowAsset()->Material)
+	{
+		DynamicMeshComponent->ConfigureMaterialSet({ DataflowContent->GetDataflowAsset()->Material });
+	}
+	else
+	{
+		DynamicMeshComponent->SetOverrideRenderMaterial(FDataflowEditorStyle::Get().VertexMaterial);
+		DynamicMeshComponent->SetShadowsEnabled(false);
+	}
+	//else if (FDataflowEditorStyle::Get().DefaultMaterial)
+	//{
+	//	DynamicMeshComponent->ConfigureMaterialSet({ FDataflowEditorStyle::Get().DefaultMaterial });
+	//}
+	//else
+	//{
+	//	DynamicMeshComponent->ValidateMaterialSlots(true, false);
+	//}
+
+	DynamicMeshComponent->SelectionOverrideDelegate = UPrimitiveComponent::FSelectionOverride::CreateRaw(this, &FDataflowPreviewScene::IsComponentSelected);
+	DynamicMeshComponent->UpdateBounds();
+
+	AddComponent(DynamicMeshComponent, DynamicMeshComponent->GetRelativeTransform());	
+	DynamicMeshComponents.Emplace(InKey, DynamicMeshComponent);
+	return DynamicMeshComponents[InKey];
+}
+
+void FDataflowConstructionScene::AddWireframeMeshElementsVisualizer()
+{
+	ensure(WireframeElements.Num()==0);
+	for(FRenderElement Elem : DynamicMeshComponents)
+	{
+		TObjectPtr<UDynamicMeshComponent> DynamicMeshComponent = Elem.Value;
+
+		// Set up the wireframe display of the rest space mesh.
+
+		TObjectPtr<UMeshElementsVisualizer> WireframeDraw = NewObject<UMeshElementsVisualizer>(RootSceneActor);
+		WireframeElements.Add(DynamicMeshComponent, WireframeDraw);
+
+		WireframeDraw->CreateInWorld(GetWorld(), FTransform::Identity);
+		WireframeDraw->Settings->DepthBias = 2.0;
+		WireframeDraw->Settings->bAdjustDepthBiasUsingMeshSize = false;
+		WireframeDraw->Settings->bShowWireframe = true;
+		WireframeDraw->Settings->bShowBorders = true;
+		WireframeDraw->Settings->bShowUVSeams = false;
+		WireframeDraw->WireframeComponent->BoundaryEdgeThickness = 2;
+
+		WireframeDraw->SetMeshAccessFunction([DynamicMeshComponent](UMeshElementsVisualizer::ProcessDynamicMeshFunc ProcessFunc)
+		{
+			ProcessFunc(*DynamicMeshComponent->GetMesh());
+		});
+
+		for (FRenderElement RenderElement : DynamicMeshComponents)
+		{
+			RenderElement.Value->OnMeshChanged.Add(FSimpleMulticastDelegate::FDelegate::CreateLambda([WireframeDraw,this]()
+			{
+				WireframeDraw->NotifyMeshChanged();
+			}));
+
+			const bool bRestSpaceMeshVisible = RenderElement.Value->GetVisibleFlag();
+			WireframeDraw->Settings->bVisible = bRestSpaceMeshVisible && bConstructionViewWireframe;
+		}
+		PropertyObjectsToTick.Add(WireframeDraw->Settings);
+	}
+}
+
+void FDataflowConstructionScene::ResetWireframeMeshElementsVisualizer()
+{
+	for (FRenderWireElement Elem : WireframeElements)
+	{
+		Elem.Value->Disconnect();
+	}
+	WireframeElements.Empty();
+}
+
+void FDataflowConstructionScene::UpdateWireframeMeshElementsVisualizer()
+{
+	ResetWireframeMeshElementsVisualizer();
+	AddWireframeMeshElementsVisualizer();
+}
+
+bool FDataflowConstructionScene::HasRenderableGeometry()
+{
+	for (FRenderElement RenderElement : DynamicMeshComponents)
+	{
+		if (RenderElement.Value->GetMesh()->TriangleCount() > 0)
+		{
+			return true;
+		}
+	}
+	return false;
+}
+
+void FDataflowConstructionScene::ResetConstructionScene()
+{
+	// Some objects, like the UMeshElementsVisualizer and Settings Objects
+	// are not part of a tool, so they won't get ticked.This member holds
+	// ticked objects that get rebuilt on Update
+	PropertyObjectsToTick.Empty();
+
+	ResetWireframeMeshElementsVisualizer();
+
+	ResetDynamicMeshComponents();
+}
+
+void FDataflowConstructionScene::UpdateConstructionScene()
+{
+	ResetConstructionScene();
+
+	// The preview scene for the construction view will be
+	// cleared and rebuilt from scratch. This will genrate a 
+	// list of UPrimitiveComponents for rendering.
+	UpdateDynamicMeshComponents();
+	
+	// Attach a wireframe renderer to the DynamicMeshComponents
+	UpdateWireframeMeshElementsVisualizer();
+
+	DataflowContent->SetIsDirty(false);
+}
+
+FDataflowSimulationScene::FDataflowSimulationScene(FPreviewScene::ConstructionValues ConstructionValues, TObjectPtr<UDataflowBaseContent> InEditorContent) 
+	: FDataflowPreviewScene(ConstructionValues,InEditorContent)
+{
+	DataflowContent->RegisterWorldContent(this, RootSceneActor);
+
+	TInlineComponentArray<UPrimitiveComponent*> PrimComponents;
+	RootSceneActor->GetComponents(PrimComponents);
+
+	for(UPrimitiveComponent* PrimComponent : PrimComponents)
+	{
+		PrimComponent->SelectionOverrideDelegate =
+			UPrimitiveComponent::FSelectionOverride::CreateRaw(this, &FDataflowPreviewScene::IsComponentSelected);
+	}
+}
+
+FDataflowSimulationScene::~FDataflowSimulationScene()
+{
+	TInlineComponentArray<UPrimitiveComponent*> PrimComponents;
+	RootSceneActor->GetComponents(PrimComponents);
+
+	for(UPrimitiveComponent* PrimComponent : PrimComponents)
+	{
+		PrimComponent->SelectionOverrideDelegate.Unbind();
+	}
+	
+	DataflowContent->UnregisterWorldContent(this);
+}
+
+void FDataflowSimulationScene::TickDataflowScene(const float DeltaSeconds)
+{
+	GetWorld()->Tick(ELevelTick::LEVELTICK_All, DeltaSeconds);
+}
+
+#undef LOCTEXT_NAMESPACE
+
diff --git a/Engine/Plugins/Experimental/Dataflow/Source/DataflowEditor/Private/Dataflow/DataflowEditorToolkit.cpp b/Engine/Plugins/Experimental/Dataflow/Source/DataflowEditor/Private/Dataflow/DataflowEditorToolkit.cpp
--- a/Engine/Plugins/Experimental/Dataflow/Source/DataflowEditor/Private/Dataflow/DataflowEditorToolkit.cpp
+++ b/Engine/Plugins/Experimental/Dataflow/Source/DataflowEditor/Private/Dataflow/DataflowEditorToolkit.cpp
@@ -8,6 +8,7 @@
 #include "Dataflow/DataflowCore.h"
 #include "Dataflow/DataflowEditor.h"
 #include "Dataflow/DataflowContent.h"
+#include "Dataflow/DataflowConstructionViewportClient.h"
 #include "Dataflow/DataflowEditorCollectionComponent.h"
 #include "Dataflow/DataflowEditorCommands.h"
 #include "Dataflow/DataflowEditorMode.h"
@@ -15,14 +16,14 @@
 #include "Dataflow/DataflowEditorModule.h"
 #include "Dataflow/DataflowEditorModeUILayer.h"
 #include "Dataflow/DataflowEditorViewport.h"
-#include "Dataflow/DataflowEditorViewportClient.h"
 #include "Dataflow/DataflowEdNode.h"
 #include "Dataflow/DataflowGraphEditor.h"
 #include "Dataflow/DataflowNodeFactory.h"
 #include "Dataflow/DataflowObject.h"
 #include "Dataflow/DataflowObjectInterface.h"
-#include "Dataflow/DataflowPreviewScene.h"
+#include "Dataflow/DataflowEditorScenes.h"
 #include "Dataflow/DataflowSchema.h"
+#include "Dataflow/DataflowSimulationViewportClient.h"
 #include "DynamicMeshBuilder.h"
 #include "EditorModeManager.h"
 #include "EditorStyleSet.h"
@@ -349,7 +350,7 @@ void FDataflowEditorToolkit::PostInitAssetEditor()
 	//@todo(brice) : can we remove this coupling of the viewport and mode
 	UDataflowEditorMode* const DataflowMode = CastChecked<UDataflowEditorMode>(EditorModeManager->GetActiveScriptableMode(UDataflowEditorMode::EM_DataflowEditorModeId));
 	const TWeakPtr<FViewportClient> WeakViewportClient(ViewportClient);
-	DataflowMode->SetRestSpaceViewportClient(StaticCastWeakPtr<FDataflowEditorViewportClient>(WeakViewportClient));
+	DataflowMode->SetRestSpaceViewportClient(StaticCastWeakPtr<FDataflowConstructionViewportClient>(WeakViewportClient));
 }
 
 void FDataflowEditorToolkit::InitializeEdMode(UBaseCharacterFXEditorMode* EdMode)
@@ -371,7 +372,7 @@ void FDataflowEditorToolkit::InitializeEdMode(UBaseCharacterFXEditorMode* EdMode
 	if (TSharedPtr<FModeToolkit> ModeToolkit = DataflowMode->GetToolkit().Pin())
 	{
 		FDataflowEditorModeToolkit* DataflowModeToolkit = static_cast<FDataflowEditorModeToolkit*>(ModeToolkit.Get());
-		DataflowModeToolkit->SetRestSpaceViewportWidget(DataflowEditorViewport);
+		DataflowModeToolkit->SetConstructionViewportWidget(DataflowEditorViewport);
 
 		FName ParentToolbarName;
 		const FName ToolBarName = GetToolMenuToolbarName(ParentToolbarName);
@@ -431,7 +432,7 @@ AssetEditorViewportFactoryFunction FDataflowEditorToolkit::GetViewportDelegate()
 	AssetEditorViewportFactoryFunction TempViewportDelegate = [this](FAssetEditorViewportConstructionArgs InArgs)
 	{
 		return SAssignNew(DataflowEditorViewport, SDataflowEditorViewport, InArgs)
-			.ViewportClient(StaticCastSharedPtr<FDataflowEditorViewportClient>(ViewportClient));
+			.ViewportClient(StaticCastSharedPtr<FDataflowConstructionViewportClient>(ViewportClient));
 	};
 
 	return TempViewportDelegate;
@@ -447,7 +448,7 @@ TSharedPtr<FEditorViewportClient> FDataflowEditorToolkit::CreateEditorViewportCl
 	// namely ViewportType.
 	// Instead, we do viewport client adjustment in PostInitAssetEditor().
 	check(EditorModeManager.IsValid());
-	TSharedPtr<FDataflowEditorViewportClient> LocalConstructionClient = MakeShared<FDataflowEditorViewportClient>(
+	TSharedPtr<FDataflowConstructionViewportClient> LocalConstructionClient = MakeShared<FDataflowConstructionViewportClient>(
 	EditorModeManager.Get(), ObjectScene.Get(), true);
 	LocalConstructionClient->SetDataflowEditorToolkit(StaticCastSharedRef<FDataflowEditorToolkit>(
 		const_cast<FDataflowEditorToolkit*>(this)->AsShared()));
@@ -457,7 +458,7 @@ TSharedPtr<FEditorViewportClient> FDataflowEditorToolkit::CreateEditorViewportCl
 void FDataflowEditorToolkit::CreateSimulationViewportClient()
 {
 	SimulationTabContent = MakeShareable(new FEditorViewportTabContent());
-	SimulationViewportClient = MakeShared<FDataflowEditorViewportClient>(SimulationModeManager.Get(),
+	SimulationViewportClient = MakeShared<FDataflowSimulationViewportClient>(SimulationModeManager.Get(),
 		SimulationScene.Get(), false);
 	
 	SimulationViewportClient->SetDataflowEditorToolkit(StaticCastSharedRef<FDataflowEditorToolkit>(this->AsShared()));
@@ -465,7 +466,7 @@ void FDataflowEditorToolkit::CreateSimulationViewportClient()
 	SimulationViewportDelegate = [this](FAssetEditorViewportConstructionArgs InArgs)
 	{
 		return SAssignNew(DataflowSimulationViewport, SDataflowEditorViewport, InArgs)
-			.ViewportClient(StaticCastSharedPtr<FDataflowEditorViewportClient>(SimulationViewportClient));
+			.ViewportClient(StaticCastSharedPtr<FDataflowSimulationViewportClient>(SimulationViewportClient));
 	};
 }
 
diff --git a/Engine/Plugins/Experimental/Dataflow/Source/DataflowEditor/Private/Dataflow/DataflowEditorViewport.cpp b/Engine/Plugins/Experimental/Dataflow/Source/DataflowEditor/Private/Dataflow/DataflowEditorViewport.cpp
--- a/Engine/Plugins/Experimental/Dataflow/Source/DataflowEditor/Private/Dataflow/DataflowEditorViewport.cpp
+++ b/Engine/Plugins/Experimental/Dataflow/Source/DataflowEditor/Private/Dataflow/DataflowEditorViewport.cpp
@@ -4,12 +4,13 @@
 
 #include "Dataflow/DataflowActor.h"
 #include "Dataflow/DataflowEditorMode.h"
-#include "Dataflow/DataflowEditorViewportClient.h"
+#include "Dataflow/DataflowConstructionViewportClient.h"
+#include "Dataflow/DataflowSimulationViewportClient.h"
 #include "Dataflow/DataflowEditorToolkit.h"
 #include "EditorModeManager.h"
 #include "Dataflow/DataflowContent.h"
-#include "Dataflow/DataflowEditorViewportToolbar.h"
-#include "Dataflow/DataflowPreviewScene.h"
+#include "Dataflow/DataflowConstructionViewportToolbar.h"
+#include "Dataflow/DataflowEditorScenes.h"
 #include "Dataflow/DataflowSimulationPanel.h"
 
 #define LOCTEXT_NAMESPACE "SDataflowEditorViewport"
@@ -26,9 +27,10 @@ void SDataflowEditorViewport::Construct(const FArguments& InArgs, const FAssetEd
 	SAssetEditorViewport::Construct(ParentArgs, InViewportConstructionArgs);
 	Client->VisibilityDelegate.BindSP(this, &SDataflowEditorViewport::IsVisible);
 
+	/*
 	if(static_cast<FDataflowPreviewScene*>(Client->GetPreviewScene())->CanRunSimulation())
 	{
-		TSharedPtr<FDataflowEditorViewportClient> DataflowClient = StaticCastSharedPtr<FDataflowEditorViewportClient>(Client);
+		TSharedPtr<FDataflowSimulationViewportClient> DataflowClient = StaticCastSharedPtr<FDataflowSimulationViewportClient>(Client);
 		TWeakPtr<FDataflowSimulationScene> SimulationScene = DataflowClient->GetDataflowEditorToolkit().Pin()->GetSimulationScene();
             
 		ViewportOverlay->AddSlot()
@@ -52,11 +54,12 @@ void SDataflowEditorViewport::Construct(const FArguments& InArgs, const FAssetEd
 			]
 		];
 	}
+	*/
 }
 
 TSharedPtr<SWidget> SDataflowEditorViewport::MakeViewportToolbar()
 {
-	return SNew(SDataflowViewportSelectionToolBar, SharedThis(this));
+	return SNew(SDataflowConstructionViewportSelectionToolBar, SharedThis(this));
 }
 
 void SDataflowEditorViewport::OnFocusViewportToSelection()
diff --git a/Engine/Plugins/Experimental/Dataflow/Source/DataflowEditor/Private/Dataflow/DataflowEditorViewportClient.cpp b/Engine/Plugins/Experimental/Dataflow/Source/DataflowEditor/Private/Dataflow/DataflowEditorViewportClient.cpp
deleted file mode 100644
diff --git a/Engine/Plugins/Experimental/Dataflow/Source/DataflowEditor/Private/Dataflow/DataflowEditorViewportToolbar.cpp b/Engine/Plugins/Experimental/Dataflow/Source/DataflowEditor/Private/Dataflow/DataflowEditorViewportToolbar.cpp
deleted file mode 100644
diff --git a/Engine/Plugins/Experimental/Dataflow/Source/DataflowEditor/Private/Dataflow/DataflowPreviewScene.cpp b/Engine/Plugins/Experimental/Dataflow/Source/DataflowEditor/Private/Dataflow/DataflowPreviewScene.cpp
deleted file mode 100644
diff --git a/Engine/Plugins/Experimental/Dataflow/Source/DataflowEditor/Private/Dataflow/DataflowSimulationPanel.cpp b/Engine/Plugins/Experimental/Dataflow/Source/DataflowEditor/Private/Dataflow/DataflowSimulationPanel.cpp
--- a/Engine/Plugins/Experimental/Dataflow/Source/DataflowEditor/Private/Dataflow/DataflowSimulationPanel.cpp
+++ b/Engine/Plugins/Experimental/Dataflow/Source/DataflowEditor/Private/Dataflow/DataflowSimulationPanel.cpp
@@ -1,7 +1,7 @@
 ﻿// Copyright Epic Games, Inc. All Rights Reserved.
 
 #include "Dataflow/DataflowSimulationPanel.h"
-#include "Dataflow/DataflowPreviewScene.h"
+#include "Dataflow/DataflowEditorScenes.h"
 #include "Widgets/SBoxPanel.h"
 #include "SScrubControlPanel.h"
 #include "Animation/AnimSingleNodeInstance.h"
diff --git a/Engine/Plugins/Experimental/Dataflow/Source/DataflowEditor/Private/Dataflow/DataflowSimulationViewportClient.cpp b/Engine/Plugins/Experimental/Dataflow/Source/DataflowEditor/Private/Dataflow/DataflowSimulationViewportClient.cpp
new file mode 100644
--- /dev/null
+++ b/Engine/Plugins/Experimental/Dataflow/Source/DataflowEditor/Private/Dataflow/DataflowSimulationViewportClient.cpp
@@ -0,0 +1,72 @@
+// Copyright Epic Games, Inc. All Rights Reserved.
+#include "Dataflow/DataflowSimulationViewportClient.h"
+
+#include "AssetEditorModeManager.h"
+#include "Dataflow/DataflowObject.h"
+#include "Dataflow/DataflowEditorMode.h"
+#include "Dataflow/DataflowEditorToolkit.h"
+#include "Dataflow/DataflowEditorCollectionComponent.h"
+#include "Dataflow/DataflowEngineSceneHitProxies.h"
+#include "Dataflow/DataflowGraphEditor.h"
+#include "Dataflow/DataflowEditorScenes.h"
+#include "EditorModeManager.h"
+#include "EdModeInteractiveToolsContext.h"
+#include "GraphEditor.h"
+#include "PreviewScene.h"
+#include "Selection.h"
+#include "SGraphPanel.h"
+#include "SNodePanel.h"
+
+
+
+FDataflowSimulationViewportClient::FDataflowSimulationViewportClient(FEditorModeTools* InModeTools,
+                                                             FPreviewScene* InPreviewScene,  const bool bCouldTickScene,
+                                                             const TWeakPtr<SEditorViewport> InEditorViewportWidget)
+	: FEditorViewportClient(InModeTools, InPreviewScene, InEditorViewportWidget)
+{
+	// We want our near clip plane to be quite close so that we can zoom in further.
+	OverrideNearClipPlane(KINDA_SMALL_NUMBER);
+
+	EngineShowFlags.SetSelectionOutline(true);
+	EngineShowFlags.EnableAdvancedFeatures();
+
+	PreviewScene = static_cast<FDataflowPreviewScene*>(InPreviewScene);
+	bEnableSceneTicking = bCouldTickScene;
+}
+
+void FDataflowSimulationViewportClient::SetDataflowEditorToolkit(TWeakPtr<FDataflowEditorToolkit> InDataflowEditorToolkitPtr)
+{
+	DataflowEditorToolkitPtr = InDataflowEditorToolkitPtr;
+}
+
+void FDataflowSimulationViewportClient::SetToolCommandList(TWeakPtr<FUICommandList> InToolCommandList)
+{
+	ToolCommandList = InToolCommandList;
+}
+
+//const UInputBehaviorSet* FDataflowSimulationViewportClient::GetInputBehaviors() const
+//{
+//	return BehaviorSet;
+//}
+
+void FDataflowSimulationViewportClient::Tick(float DeltaSeconds)
+{
+	Super::Tick(DeltaSeconds);
+	if (PreviewScene)
+	{
+		PreviewScene->TickDataflowScene(DeltaSeconds);
+	}
+}
+
+void FDataflowSimulationViewportClient::ProcessClick(FSceneView& View, HHitProxy* HitProxy, FKey Key, EInputEvent Event, uint32 HitX, uint32 HitY)
+{
+	Super::ProcessClick(View, HitProxy, Key, Event, HitX, HitY);
+}
+
+
+void FDataflowSimulationViewportClient::AddReferencedObjects(FReferenceCollector& Collector)
+{
+	Super::AddReferencedObjects(Collector);
+	Collector.AddReferencedObject(BehaviorSet);
+}
+
diff --git a/Engine/Plugins/Experimental/Dataflow/Source/DataflowEditor/Public/Dataflow/CollectionSpreadSheetWidget.h b/Engine/Plugins/Experimental/Dataflow/Source/DataflowEditor/Public/Dataflow/CollectionSpreadSheetWidget.h
deleted file mode 100644
diff --git a/Engine/Plugins/Experimental/Dataflow/Source/DataflowEditor/Public/Dataflow/DataflowCollectionSpreadSheet.h b/Engine/Plugins/Experimental/Dataflow/Source/DataflowEditor/Public/Dataflow/DataflowCollectionSpreadSheet.h
--- a/Engine/Plugins/Experimental/Dataflow/Source/DataflowEditor/Public/Dataflow/DataflowCollectionSpreadSheet.h
+++ b/Engine/Plugins/Experimental/Dataflow/Source/DataflowEditor/Public/Dataflow/DataflowCollectionSpreadSheet.h
@@ -5,7 +5,7 @@
 #include "Widgets/SCompoundWidget.h"
 #include "Widgets/Views/SListView.h"
 #include "Dataflow/DataflowEdNode.h"
-#include "Dataflow/CollectionSpreadSheetWidget.h"
+#include "Dataflow/DataflowCollectionSpreadSheetWidget.h"
 #include "Dataflow/DataflowSelection.h"
 #include "Dataflow/DataflowView.h"
 
diff --git a/Engine/Plugins/Experimental/Dataflow/Source/DataflowEditor/Public/Dataflow/DataflowCollectionSpreadSheetWidget.h b/Engine/Plugins/Experimental/Dataflow/Source/DataflowEditor/Public/Dataflow/DataflowCollectionSpreadSheetWidget.h
new file mode 100644
--- /dev/null
+++ b/Engine/Plugins/Experimental/Dataflow/Source/DataflowEditor/Public/Dataflow/DataflowCollectionSpreadSheetWidget.h
@@ -0,0 +1,171 @@
+// Copyright Epic Games, Inc. All Rights Reserved.
+
+#pragma once
+
+#include "CoreMinimal.h"
+#include "Dataflow/DataflowSelection.h"
+#include "DetailLayoutBuilder.h"
+#include "Widgets/Views/STableRow.h"
+#include "Widgets/Text/STextBlock.h"
+#include "Widgets/Input/SComboBox.h"
+#include "GeometryCollection/ManagedArrayCollection.h"
+
+/**
+* Struct to hold OutputType/Selection data for the outputs
+* Data is stored in a map using the OutputName as key: TMap<FString, SCollectionInfo>
+*/
+struct SCollectionInfo
+{
+//	FString OutputType;
+	const FManagedArrayCollection Collection;
+};
+
+
+/** 
+* Header
+* 1st column possible values: Transform Index/Face Index/Vertex Index based on the OutputType
+* 2nd column is SelectionStatus
+*/
+struct FCollectionSpreadSheetHeader
+{
+	static const FName IndexColumnName;
+
+	TArray<FName> ColumnNames;
+};
+
+
+/** 
+* Representing a row in the table 
+* Index/SelectionStatus
+*/
+struct FCollectionSpreadSheetItem
+{
+	TArray<FString> Values;
+};
+
+
+/** 
+* 
+*/
+class SCollectionSpreadSheetRow : public SMultiColumnTableRow<TSharedPtr<const FCollectionSpreadSheetItem>>
+{
+public:
+	SLATE_BEGIN_ARGS(SCollectionSpreadSheetRow) 
+	{}
+		SLATE_ARGUMENT(TSharedPtr<const FCollectionSpreadSheetHeader>, Header)
+		SLATE_ARGUMENT(TSharedPtr<const FCollectionSpreadSheetItem>, Item)
+	SLATE_END_ARGS()
+
+	void Construct(const FArguments& InArgs, TSharedRef<STableViewBase> OwnerTableView, const TSharedPtr<const FCollectionSpreadSheetHeader>& InHeader, const TSharedPtr<const FCollectionSpreadSheetItem>& InItem);
+
+	virtual TSharedRef<SWidget> GenerateWidgetForColumn(const FName& ColumnName) override;
+
+private:
+	TSharedPtr<const FCollectionSpreadSheetHeader> Header;
+	TSharedPtr<const FCollectionSpreadSheetItem> Item;
+};
+
+/** 
+* 2xn grid to display Collection data
+*/
+class SCollectionSpreadSheet : public SCompoundWidget
+{
+public:
+	SLATE_BEGIN_ARGS(SCollectionSpreadSheet) {}
+		SLATE_ARGUMENT(FName, SelectedOutput)
+	SLATE_END_ARGS()
+
+	void Construct(const FArguments& InArgs);
+
+	TMap<FString, SCollectionInfo>& GetCollectionInfoMap() { return CollectionInfoMap; }
+
+	const FName& GetSelectedOutput() const;
+	void SetSelectedOutput(const FName& InSelectedOutput);
+	
+	const FName& GetSelectedGroup() const;
+	void SetSelectedGroup(const FName& InSelectedGroup);
+
+	int32 GetNumItems() const { return NumItems; }
+	void SetNumItems(int32 InNumItems) { NumItems = InNumItems; }
+
+	TSharedRef<ITableRow> GenerateRow(TSharedPtr<const FCollectionSpreadSheetItem> InItem, const TSharedRef<STableViewBase>& OwnerTable);
+
+	TSharedPtr<SScrollBar> SpreadSheetVerticalScrollBar;
+
+	TMap<FString, int32> AttrTypeWidthMap;
+
+private:
+	FName SelectedOutput = FName(TEXT(""));
+	FName SelectedGroup = FName(TEXT(""));
+	TMap<FString, SCollectionInfo> CollectionInfoMap;
+
+	TSharedPtr<SListView<TSharedPtr<const FCollectionSpreadSheetItem>>> ListView;
+//	TSharedPtr<STreeView<TSharedPtr<const FCollectionSpreadSheetItem>>> ListView;
+	TArray<TSharedPtr<const FCollectionSpreadSheetItem>> ListItems;
+
+	TSharedPtr<FCollectionSpreadSheetHeader> Header;
+	TSharedPtr<SHeaderRow> HeaderRowWidget;
+
+	int32 NumItems = 0;
+	void RegenerateHeader();
+	void RepopulateListView();
+};
+
+/**
+* Widget for the CollectionSpreadSheet panel
+*/
+class SCollectionSpreadSheetWidget : public SCompoundWidget
+{
+public:
+	SLATE_BEGIN_ARGS(SCollectionSpreadSheetWidget) {}
+	SLATE_END_ARGS()
+
+	/**
+	* Slate construction function
+	*
+	* @param InArgs - Arguments passed from Slate
+	*
+	*/
+	void Construct(const FArguments& InArgs);
+
+	void SetData(const FString& InNodeName);
+	void RefreshWidget();
+	TSharedPtr<SCollectionSpreadSheet> GetCollectionTable() { return CollectionTable; }
+	void SetStatusText();
+
+	void UpdateCollectionGroups(const FName& InOutputName);
+
+	/** Gets a multicast delegate which is called whenever the PinnedDown button clicked */
+	DECLARE_MULTICAST_DELEGATE_OneParam(FOnPinnedDownChanged, const bool);
+	FOnPinnedDownChanged& GetOnPinnedDownChangedDelegate() { return OnPinnedDownChangedDelegate; }
+
+	/** Gets a multicast delegate which is called whenever the LockRefresh button clicked */
+	DECLARE_MULTICAST_DELEGATE_OneParam(FOnRefreshLockedChanged, const bool);
+	FOnRefreshLockedChanged& GetOnRefreshLockedChangedDelegate() { return OnRefreshLockedChangedDelegate; }
+
+	// UI callbacks
+	void NodeOutputsComboBoxSelectionChanged(FName InSelectedOutput, ESelectInfo::Type InSelectInfo);
+	void CollectionGroupsComboBoxSelectionChanged(FName InSelectedOutput, ESelectInfo::Type InSelectInfo);
+	FText GetNoOutputText();
+	FText GetNoGroupText();
+
+private:
+	TSharedPtr<STextBlock> NodeNameTextBlock;
+	TSharedPtr<SComboBox<FName>> NodeOutputsComboBox;
+	TSharedPtr<STextBlock> NodeOutputsComboBoxLabel;
+	TSharedPtr<SComboBox<FName>> CollectionGroupsComboBox;
+	TSharedPtr<STextBlock> CollectionGroupsComboBoxLabel;
+	TSharedPtr<SCollectionSpreadSheet> CollectionTable;
+	TSharedPtr<STextBlock> StatusTextBlock;
+
+	FString NodeName;
+	TArray<FName> NodeOutputs;
+	TArray<FName> CollectionGroups;
+	bool bIsPinnedDown = false;
+	bool bIsRefreshLocked = false;
+
+protected:
+	FOnPinnedDownChanged OnPinnedDownChangedDelegate;
+	FOnRefreshLockedChanged OnRefreshLockedChangedDelegate;
+
+};
\ No newline at end of file
diff --git a/Engine/Plugins/Experimental/Dataflow/Source/DataflowEditor/Public/Dataflow/DataflowConstructionViewportClient.h b/Engine/Plugins/Experimental/Dataflow/Source/DataflowEditor/Public/Dataflow/DataflowConstructionViewportClient.h
new file mode 100644
--- /dev/null
+++ b/Engine/Plugins/Experimental/Dataflow/Source/DataflowEditor/Public/Dataflow/DataflowConstructionViewportClient.h
@@ -0,0 +1,69 @@
+// Copyright Epic Games, Inc. All Rights Reserved.
+#pragma once
+
+#include "CoreMinimal.h"
+#include "EditorViewportClient.h"
+#include "Dataflow/DataflowNodeParameters.h"
+#include "Dataflow/DataflowComponentSelectionState.h"
+#include "Dataflow/DataflowContent.h"
+#include "InputBehaviorSet.h"
+
+class FDataflowEditorToolkit;
+class ADataflowActor;
+class UTransformProxy;
+class UCombinedTransformGizmo;
+class FTransformGizmoDataBinder;
+class FDataflowPreviewScene;
+class UInputBehaviorSet;
+
+class DATAFLOWEDITOR_API FDataflowConstructionViewportClient : public FEditorViewportClient //, public IInputBehaviorSource
+{
+public:
+	using Super = FEditorViewportClient;
+
+	FDataflowConstructionViewportClient(FEditorModeTools* InModeTools, FPreviewScene* InPreviewScene,  const bool bCouldTickScene,
+								  const TWeakPtr<SEditorViewport> InEditorViewportWidget = nullptr);
+
+	void SetConstructionViewMode(Dataflow::EDataflowPatternVertexType InViewMode);
+	Dataflow::EDataflowPatternVertexType GetConstructionViewMode() const;
+
+	// IInputBehaviorSource
+	// virtual const UInputBehaviorSet* GetInputBehaviors() const override;
+
+	/** Set the data flow toolkit used to create the client*/
+	void SetDataflowEditorToolkit(TWeakPtr<FDataflowEditorToolkit> DataflowToolkit);
+	
+	/** Get the data flow toolkit  */
+	const TWeakPtr<FDataflowEditorToolkit>& GetDataflowEditorToolkit() const { return DataflowEditorToolkitPtr; }
+
+	/** Set the tool command list */
+	void SetToolCommandList(TWeakPtr<FUICommandList> ToolCommandList);
+
+	// FGCObject Interface
+	virtual void AddReferencedObjects(FReferenceCollector& Collector) override;
+	virtual FString GetReferencerName() const override { return TEXT("FDataflowConstructionViewportClient"); }
+
+private:
+
+	// FEditorViewportClient interface
+	virtual void Tick(float DeltaSeconds) override;
+	virtual void ProcessClick(FSceneView& View, HHitProxy* HitProxy, FKey Key, EInputEvent Event, uint32 HitX, uint32 HitY) override;
+
+	/** Toolkit used to create the viewport client */
+	TWeakPtr<FDataflowEditorToolkit> DataflowEditorToolkitPtr = nullptr;
+
+	/** Dataflow preview scene from the toolkit */
+	FDataflowPreviewScene* PreviewScene = nullptr;
+
+	// @todo(brice) : Is this needed?
+	TWeakPtr<FUICommandList> ToolCommandList;
+
+	/** Construction view mode */
+	Dataflow::EDataflowPatternVertexType ConstructionViewMode = Dataflow::EDataflowPatternVertexType::Sim3D;
+
+	/** Behavior set for the behavior UI */
+	TObjectPtr<UInputBehaviorSet> BehaviorSet;
+	
+	/** Flag to enable scene ticking from the client */
+	bool bEnableSceneTicking = false;
+};
diff --git a/Engine/Plugins/Experimental/Dataflow/Source/DataflowEditor/Public/Dataflow/DataflowConstructionViewportToolbar.h b/Engine/Plugins/Experimental/Dataflow/Source/DataflowEditor/Public/Dataflow/DataflowConstructionViewportToolbar.h
new file mode 100644
--- /dev/null
+++ b/Engine/Plugins/Experimental/Dataflow/Source/DataflowEditor/Public/Dataflow/DataflowConstructionViewportToolbar.h
@@ -0,0 +1,23 @@
+// Copyright Epic Games, Inc. All Rights Reserved.
+#pragma once
+
+#include "CoreMinimal.h"
+#include "SCommonEditorViewportToolbarBase.h"
+
+// ----------------------------------------------------------------------------------
+class SDataflowEditorViewport;
+
+/** Base toolbar for the dataflow. Should be extended to add more features */
+class SDataflowConstructionViewportSelectionToolBar : public SCommonEditorViewportToolbarBase
+{
+public:
+	SLATE_BEGIN_ARGS(SDataflowConstructionViewportSelectionToolBar) {}
+	SLATE_END_ARGS()
+
+	/** Constructs this widget with the given parameters */
+	void Construct(const FArguments& InArgs, TSharedPtr<SDataflowEditorViewport> InDataflowViewport);
+	
+private:
+	/** Reference to the parent viewport */
+	TWeakPtr<SDataflowEditorViewport> EditorViewport;
+};
\ No newline at end of file
diff --git a/Engine/Plugins/Experimental/Dataflow/Source/DataflowEditor/Public/Dataflow/DataflowEditorMode.h b/Engine/Plugins/Experimental/Dataflow/Source/DataflowEditor/Public/Dataflow/DataflowEditorMode.h
--- a/Engine/Plugins/Experimental/Dataflow/Source/DataflowEditor/Public/Dataflow/DataflowEditorMode.h
+++ b/Engine/Plugins/Experimental/Dataflow/Source/DataflowEditor/Public/Dataflow/DataflowEditorMode.h
@@ -9,7 +9,7 @@
 class UDataflowComponent;
 class FDataflowConstructionScene;
 class SDataflowGraphEditor;
-class FDataflowEditorViewportClient;
+class FDataflowConstructionViewportClient;
 class IDataflowEditorToolBuilder;
 struct FManagedArrayCollection;
 class UEdGraphNode;
@@ -59,10 +59,10 @@ public:
 		return bConstructionViewWireframe;
 	}
 
-	/** Set the data flow preview scene to store the dynamic meshes generated by the tools */
+	/** Set the construction scene to store the dynamic meshes generated by the tools */
 	void SetDataflowConstructionScene(FDataflowConstructionScene* DataflowPreviewScene);
-	FDataflowConstructionScene* GetDataflowConstructionScene() { return PreviewScene; }
-	const FDataflowConstructionScene* GetDataflowConstructionScene() const { return PreviewScene; }
+	FDataflowConstructionScene* GetDataflowConstructionScene() { return ConstructionScene; }
+	const FDataflowConstructionScene* GetDataflowConstructionScene() const { return ConstructionScene; }
 
 	/** Set the data flow graph editor to create nodes once the tools have ended*/
 	void SetDataflowGraphEditor(TSharedPtr<SDataflowGraphEditor> DataflowGraphEditor);
@@ -106,7 +106,7 @@ private:
 
 	void RegisterAddNodeCommand(TSharedPtr<FUICommandInfo> AddNodeCommand, const FName& NewNodeType, TSharedPtr<FUICommandInfo> StartToolCommand);
 
-	void SetRestSpaceViewportClient(TWeakPtr<FDataflowEditorViewportClient, ESPMode::ThreadSafe> ViewportClient);
+	void SetRestSpaceViewportClient(TWeakPtr<FDataflowConstructionViewportClient, ESPMode::ThreadSafe> ViewportClient);
 	void RefocusRestSpaceViewportClient();
 	void FirstTimeFocusRestSpaceViewport();
 	bool IsComponentSelected(const UPrimitiveComponent* InComponent);
@@ -147,7 +147,7 @@ private:
 	void DeleteContextObject();
 
 	/** Dataflow preview scene from the toolkit */
-	FDataflowConstructionScene* PreviewScene = nullptr;
+	FDataflowConstructionScene* ConstructionScene = nullptr;
 
 	/** Correspondence between node types and commands to launch tools */
 	TMap<FName, TSharedPtr<const FUICommandInfo>> NodeTypeToToolCommandMap;
@@ -169,7 +169,7 @@ private:
 	UPROPERTY()
 	TObjectPtr<UEditorInteractiveToolsContext> ActiveToolsContext = nullptr;
 
-	TWeakPtr<FDataflowEditorViewportClient, ESPMode::ThreadSafe> ConstructionViewportClient;
+	TWeakPtr<FDataflowConstructionViewportClient, ESPMode::ThreadSafe> ConstructionViewportClient;
 
 	// The first time we get a valid mesh, refocus the camera on it
 	bool bFirstValid2DMesh = true;
diff --git a/Engine/Plugins/Experimental/Dataflow/Source/DataflowEditor/Public/Dataflow/DataflowEditorModeToolkit.h b/Engine/Plugins/Experimental/Dataflow/Source/DataflowEditor/Public/Dataflow/DataflowEditorModeToolkit.h
--- a/Engine/Plugins/Experimental/Dataflow/Source/DataflowEditor/Public/Dataflow/DataflowEditorModeToolkit.h
+++ b/Engine/Plugins/Experimental/Dataflow/Source/DataflowEditor/Public/Dataflow/DataflowEditorModeToolkit.h
@@ -35,17 +35,17 @@ public:
 	virtual FName GetToolkitFName() const override;
 	virtual FText GetBaseToolkitName() const override;
 
-	void SetRestSpaceViewportWidget(TWeakPtr<SDataflowEditorViewport>);
+	void SetConstructionViewportWidget(TWeakPtr<SDataflowEditorViewport>);
 	//void SetPreviewViewportWidget(TWeakPtr<SChaosClothAssetEditor3DViewport>);
 
 private:
 
 	// Get the viewport widget associated with the given manager
-	// TODO: This should not be necessary any more as we do not run tools in the Cloth Preview Viewport (JIRA UE-201248)
+	// TODO: This should not be necessary any more as we do not run tools in the Simulation Viewport (JIRA UE-201248)
 	SBaseCharacterFXEditorViewport* GetViewportWidgetForManager(UInteractiveToolManager* Manager);
 
 	UEditorInteractiveToolsContext* GetCurrentToolsContext();
 
-	TWeakPtr<SDataflowEditorViewport> RestSpaceViewportWidget;
+	TWeakPtr<SDataflowEditorViewport> ConstructionViewportWidget;
 	//TWeakPtr<SChaosClothAssetEditor3DViewport> PreviewViewportWidget;
 };
diff --git a/Engine/Plugins/Experimental/Dataflow/Source/DataflowEditor/Public/Dataflow/DataflowEditorScenes.h b/Engine/Plugins/Experimental/Dataflow/Source/DataflowEditor/Public/Dataflow/DataflowEditorScenes.h
new file mode 100644
--- /dev/null
+++ b/Engine/Plugins/Experimental/Dataflow/Source/DataflowEditor/Public/Dataflow/DataflowEditorScenes.h
@@ -0,0 +1,164 @@
+// Copyright Epic Games, Inc. All Rights Reserved.
+
+#pragma once
+
+#include "AdvancedPreviewScene.h"
+
+class UDataflowBaseContent;
+class USkeletalMeshComponent;
+class UDynamicMeshComponent;
+class UDataflowComponent;
+class UAnimationAsset;
+class UAnimSingleNodeInstance;
+class UMaterialInterface;
+class FAssetEditorModeManager;
+class ADataflowActor;
+class UMeshElementsVisualizer;
+class UInteractiveToolPropertySet;
+class UDataflowEdNode;
+
+namespace UE
+{
+	namespace Geometry
+	{
+		class FDynamicMesh3;
+	}
+}
+
+/**
+ * Dataflow preview scene
+ * @brief the scene is holding all the objects that will be
+ * visible and potentially editable within the viewport
+ */
+class DATAFLOWEDITOR_API FDataflowPreviewScene : public FAdvancedPreviewScene
+{
+public:
+
+	FDataflowPreviewScene(FPreviewScene::ConstructionValues ConstructionValues, TObjectPtr<UDataflowBaseContent> InDataflowContent);
+	virtual ~FDataflowPreviewScene();
+	
+	// FGCObject interface
+	virtual void AddReferencedObjects(FReferenceCollector& Collector) override;
+
+	/** Dataflow content accessors */
+	TObjectPtr<UDataflowBaseContent> GetDataflowContent() { return DataflowContent; }
+	const TObjectPtr<UDataflowBaseContent> GetDataflowContent() const { return DataflowContent; }
+
+	/** Root scene actor accessors */
+	TObjectPtr<AActor> GetRootActor() { return RootSceneActor; }
+	const TObjectPtr<AActor> GetRootActor() const { return RootSceneActor; }
+
+	/** Dataflow mode manager accessors */
+	TSharedPtr<FAssetEditorModeManager>& GetDataflowModeManager() { return DataflowModeManager; }
+	const TSharedPtr<FAssetEditorModeManager>& GetDataflowModeManager() const { return DataflowModeManager; }
+	
+	/** Build the scene bounding box */
+	FBox GetBoundingBox() const;
+
+	/** Tick data flow scene */
+	virtual void TickDataflowScene(const float DeltaSeconds) {}
+
+	/** Check if a primitive component is selected */
+	bool IsComponentSelected(const UPrimitiveComponent* InComponent) const;
+
+	/** Check if the preview scene can run simulation */
+	virtual bool CanRunSimulation() const {return false;}
+	
+protected:
+	
+	/** Root scene actor */
+	TObjectPtr<AActor> RootSceneActor = nullptr;
+	
+	/** Data flows assets datas used to build the scene components */
+	TObjectPtr<UDataflowBaseContent> DataflowContent = nullptr;
+
+	/** Mode Manager for selection */
+	TSharedPtr<FAssetEditorModeManager> DataflowModeManager;
+};
+
+/**
+ * Dataflow construction scene holding all the dynamic meshes used/displayed in the dataflow graph
+ */
+class DATAFLOWEDITOR_API FDataflowConstructionScene : public FDataflowPreviewScene
+{
+	/*(const void* ObjectPointer, const int32 ObjectIndex)*/
+	typedef TPair<TObjectPtr<const UDataflowEdNode>, int32> FDataflowRenderKey;
+
+public:
+
+	FDataflowConstructionScene(FPreviewScene::ConstructionValues ConstructionValues, TObjectPtr<UDataflowBaseContent> InEditorContent);
+	virtual ~FDataflowConstructionScene();
+
+	// FGCObject interface
+	virtual void AddReferencedObjects(FReferenceCollector& Collector) override;
+	
+	/** Tick data flow scene */
+	virtual void TickDataflowScene(const float DeltaSeconds) override;
+
+	/** Check if the preview scene can run simulation */
+	virtual bool CanRunSimulation() const override {return false;}
+	
+	/** Is there ANYTHING to render */
+	bool HasRenderableGeometry();
+
+	/** Close the construction scene */
+	void ResetConstructionScene();
+
+	/** Update the construction scene */
+	void UpdateConstructionScene();
+	
+private:
+
+	/** Reset the scene mesh visualizer */
+	void ResetWireframeMeshElementsVisualizer();
+
+	/** Update the scene mesh visualizer */
+	void UpdateWireframeMeshElementsVisualizer();
+
+	/** Add a mesh visualizer to the scene */
+	void AddWireframeMeshElementsVisualizer();
+
+	/** Reset all the dynamic mesh components */
+    void ResetDynamicMeshComponents();
+	
+	/** Build the dynamic mesh components from the dataflow graph render targets*/
+	void UpdateDynamicMeshComponents();
+	
+	/** Add a dynamic mesh component to the scene */
+	TObjectPtr<UDynamicMeshComponent>& AddDynamicMeshComponent(
+		FDataflowRenderKey InKey, UE::Geometry::FDynamicMesh3&& DynamicMesh,  const TArray<UMaterialInterface*>& MaterialSet);
+	
+	/** Dynamic mesh components generated by the tools for edit/manipulation*/
+	typedef TPair<FDataflowRenderKey, TObjectPtr<UDynamicMeshComponent>> FRenderElement;
+	TMap<FDataflowRenderKey, TObjectPtr<UDynamicMeshComponent>> DynamicMeshComponents;
+
+	/** Rest space wireframe. They have to get ticked to be able to respond to setting changes. [ticked] */
+	typedef TPair<TObjectPtr<UDynamicMeshComponent>, TObjectPtr<UMeshElementsVisualizer>> FRenderWireElement;
+	TMap<TObjectPtr<UDynamicMeshComponent>,TObjectPtr<UMeshElementsVisualizer> > WireframeElements;
+
+	/** Property objects (visible or not) that get ticked. [ticked] */
+	TArray<TObjectPtr<UInteractiveToolPropertySet>> PropertyObjectsToTick;
+
+	/** Show the wireframe on render */ 
+	bool bConstructionViewWireframe = true;
+};
+
+
+/**
+ * Dataflow simulation scene holding all the dataflow content components
+ */
+class DATAFLOWEDITOR_API FDataflowSimulationScene : public FDataflowPreviewScene
+{
+public:
+
+	FDataflowSimulationScene(FPreviewScene::ConstructionValues ConstructionValues, TObjectPtr<UDataflowBaseContent> InEditorContent);
+	virtual ~FDataflowSimulationScene();
+
+	/** Tick data flow scene */
+	virtual void TickDataflowScene(const float DeltaSeconds) override;
+	
+	/** Check if the preview scene can run simulation */
+	virtual bool CanRunSimulation() const { return true; }
+};
+
+
diff --git a/Engine/Plugins/Experimental/Dataflow/Source/DataflowEditor/Public/Dataflow/DataflowEditorToolkit.h b/Engine/Plugins/Experimental/Dataflow/Source/DataflowEditor/Public/Dataflow/DataflowEditorToolkit.h
--- a/Engine/Plugins/Experimental/Dataflow/Source/DataflowEditor/Public/Dataflow/DataflowEditorToolkit.h
+++ b/Engine/Plugins/Experimental/Dataflow/Source/DataflowEditor/Public/Dataflow/DataflowEditorToolkit.h
@@ -21,6 +21,7 @@ class UDataflow;
 class USkeletalMesh;
 class SDataflowGraphEditor;
 class FDataflowConstructionScene;
+class FDataflowSimulationViewportClient;
 class UDataflowBaseContent;
 class FDataflowSimulationScene;
 
@@ -187,5 +188,5 @@ private:
 	AssetEditorViewportFactoryFunction SimulationViewportDelegate;
 
 	/** Simulation Viewport client */
-	TSharedPtr<FDataflowEditorViewportClient> SimulationViewportClient;
+	TSharedPtr<FDataflowSimulationViewportClient> SimulationViewportClient;
 };
diff --git a/Engine/Plugins/Experimental/Dataflow/Source/DataflowEditor/Public/Dataflow/DataflowEditorViewport.h b/Engine/Plugins/Experimental/Dataflow/Source/DataflowEditor/Public/Dataflow/DataflowEditorViewport.h
--- a/Engine/Plugins/Experimental/Dataflow/Source/DataflowEditor/Public/Dataflow/DataflowEditorViewport.h
+++ b/Engine/Plugins/Experimental/Dataflow/Source/DataflowEditor/Public/Dataflow/DataflowEditorViewport.h
@@ -8,7 +8,7 @@
 class ADataflowActor;
 class FAdvancedPreviewScene;
 class UDataflowEditorMode;
-class FDataflowEditorViewportClient;
+class FDataflowConstructionViewportClient;
 
 // ----------------------------------------------------------------------------------
 
@@ -16,7 +16,7 @@ class SDataflowEditorViewport : public SBaseCharacterFXEditorViewport, public IC
 {
 public:
 	SLATE_BEGIN_ARGS(SDataflowEditorViewport) {}
-	SLATE_ARGUMENT(TSharedPtr<FDataflowEditorViewportClient>, ViewportClient)
+	SLATE_ARGUMENT(TSharedPtr<FEditorViewportClient>, ViewportClient)
 		SLATE_END_ARGS()
 
 	SDataflowEditorViewport();
@@ -40,3 +40,4 @@ private:
 
 	UDataflowEditorMode* GetEdMode() const;
 };
+
diff --git a/Engine/Plugins/Experimental/Dataflow/Source/DataflowEditor/Public/Dataflow/DataflowEditorViewportClient.h b/Engine/Plugins/Experimental/Dataflow/Source/DataflowEditor/Public/Dataflow/DataflowEditorViewportClient.h
deleted file mode 100644
diff --git a/Engine/Plugins/Experimental/Dataflow/Source/DataflowEditor/Public/Dataflow/DataflowEditorViewportToolbar.h b/Engine/Plugins/Experimental/Dataflow/Source/DataflowEditor/Public/Dataflow/DataflowEditorViewportToolbar.h
deleted file mode 100644
diff --git a/Engine/Plugins/Experimental/Dataflow/Source/DataflowEditor/Public/Dataflow/DataflowPreviewScene.h b/Engine/Plugins/Experimental/Dataflow/Source/DataflowEditor/Public/Dataflow/DataflowPreviewScene.h
deleted file mode 100644
diff --git a/Engine/Plugins/Experimental/Dataflow/Source/DataflowEditor/Public/Dataflow/DataflowSimulationPanel.h b/Engine/Plugins/Experimental/Dataflow/Source/DataflowEditor/Public/Dataflow/DataflowSimulationPanel.h
--- a/Engine/Plugins/Experimental/Dataflow/Source/DataflowEditor/Public/Dataflow/DataflowSimulationPanel.h
+++ b/Engine/Plugins/Experimental/Dataflow/Source/DataflowEditor/Public/Dataflow/DataflowSimulationPanel.h
@@ -1,8 +1,8 @@
-﻿// Copyright Epic Games, Inc. All Rights Reserved.
+// Copyright Epic Games, Inc. All Rights Reserved.
 
 #pragma once
 
-#include "DataflowPreviewScene.h"
+#include "DataflowEditorScenes.h"
 #include "Widgets/SCompoundWidget.h"
 #include "ITransportControl.h"  // EPlaybackMode::Type
 
diff --git a/Engine/Plugins/Experimental/Dataflow/Source/DataflowEditor/Public/Dataflow/DataflowSimulationViewportClient.h b/Engine/Plugins/Experimental/Dataflow/Source/DataflowEditor/Public/Dataflow/DataflowSimulationViewportClient.h
new file mode 100644
--- /dev/null
+++ b/Engine/Plugins/Experimental/Dataflow/Source/DataflowEditor/Public/Dataflow/DataflowSimulationViewportClient.h
@@ -0,0 +1,62 @@
+// Copyright Epic Games, Inc. All Rights Reserved.
+#pragma once
+
+#include "CoreMinimal.h"
+#include "EditorViewportClient.h"
+#include "Dataflow/DataflowNodeParameters.h"
+#include "Dataflow/DataflowComponentSelectionState.h"
+#include "Dataflow/DataflowContent.h"
+#include "InputBehaviorSet.h"
+
+class FDataflowEditorToolkit;
+class ADataflowActor;
+class UTransformProxy;
+class UCombinedTransformGizmo;
+class FTransformGizmoDataBinder;
+class FDataflowPreviewScene;
+class UInputBehaviorSet;
+
+class DATAFLOWEDITOR_API FDataflowSimulationViewportClient : public FEditorViewportClient //, public IInputBehaviorSource
+{
+//@todo(brice) : Add BehaviorUI support
+
+public:
+	using Super = FEditorViewportClient;
+
+	FDataflowSimulationViewportClient(FEditorModeTools* InModeTools, FPreviewScene* InPreviewScene,  const bool bCouldTickScene,
+								  const TWeakPtr<SEditorViewport> InEditorViewportWidget = nullptr);
+
+	/** Set the data flow toolkit used to create the client*/
+	void SetDataflowEditorToolkit(TWeakPtr<FDataflowEditorToolkit> DataflowToolkit);
+	
+	/** Get the data flow toolkit  */
+	const TWeakPtr<FDataflowEditorToolkit>& GetDataflowEditorToolkit() const { return DataflowEditorToolkitPtr; }
+
+	/** Set the tool command list */
+	void SetToolCommandList(TWeakPtr<FUICommandList> ToolCommandList);
+
+	// FGCObject Interface
+	virtual void AddReferencedObjects(FReferenceCollector& Collector) override;
+	virtual FString GetReferencerName() const override { return TEXT("FDataflowSimulationViewportClient"); }
+
+private:
+
+	// FEditorViewportClient interface
+	virtual void Tick(float DeltaSeconds) override;
+	virtual void ProcessClick(FSceneView& View, HHitProxy* HitProxy, FKey Key, EInputEvent Event, uint32 HitX, uint32 HitY) override;
+
+	/** Toolkit used to create the viewport client */
+	TWeakPtr<FDataflowEditorToolkit> DataflowEditorToolkitPtr = nullptr;
+
+	/** Dataflow preview scene from the toolkit */
+	FDataflowPreviewScene* PreviewScene = nullptr;
+
+	// @todo(brice) : Is this needed?
+	TWeakPtr<FUICommandList> ToolCommandList;
+
+	/** Behavior set for the behavior UI */
+	TObjectPtr<UInputBehaviorSet> BehaviorSet;
+	
+	/** Flag to enable scene ticking from the client */
+	bool bEnableSceneTicking = false;
+};
