Dataflow : Viewports for Construction and Simulation

diff --git a/Engine/Plugins/Experimental/Dataflow/Source/DataflowEditor/Private/Dataflow/DataflowConstructionViewport.cpp b/Engine/Plugins/Experimental/Dataflow/Source/DataflowEditor/Private/Dataflow/DataflowConstructionViewport.cpp
new file mode 100644
--- /dev/null
+++ b/Engine/Plugins/Experimental/Dataflow/Source/DataflowEditor/Private/Dataflow/DataflowConstructionViewport.cpp
@@ -0,0 +1,93 @@
+// Copyright Epic Games, Inc. All Rights Reserved.
+
+#include "Dataflow/DataflowConstructionViewport.h"
+
+#include "Dataflow/DataflowActor.h"
+#include "Dataflow/DataflowEditorMode.h"
+#include "Dataflow/DataflowConstructionViewportClient.h"
+#include "Dataflow/DataflowEditorToolkit.h"
+#include "EditorModeManager.h"
+#include "Dataflow/DataflowContent.h"
+#include "Dataflow/DataflowConstructionViewportToolbar.h"
+#include "Dataflow/DataflowEditorScenes.h"
+#include "Dataflow/DataflowSimulationPanel.h"
+
+#define LOCTEXT_NAMESPACE "SDataflowConstructionViewport"
+
+
+SDataflowConstructionViewport::SDataflowConstructionViewport()
+{
+}
+
+void SDataflowConstructionViewport::Construct(const FArguments& InArgs, const FAssetEditorViewportConstructionArgs& InViewportConstructionArgs)
+{
+	SAssetEditorViewport::FArguments ParentArgs;
+	ParentArgs._EditorViewportClient = InArgs._ViewportClient;
+	SAssetEditorViewport::Construct(ParentArgs, InViewportConstructionArgs);
+	Client->VisibilityDelegate.BindSP(this, &SDataflowConstructionViewport::IsVisible);
+}
+
+TSharedPtr<SWidget> SDataflowConstructionViewport::MakeViewportToolbar()
+{
+	return SNew(SDataflowConstructionViewportSelectionToolBar, SharedThis(this));
+}
+
+void SDataflowConstructionViewport::OnFocusViewportToSelection()
+{
+	if(const FDataflowPreviewScene* PreviewScene = static_cast<FDataflowPreviewScene*>(Client->GetPreviewScene()))
+	{
+		const FBox SceneBoundingBox = PreviewScene->GetBoundingBox();
+		Client->FocusViewportOnBox(SceneBoundingBox);
+	}
+}
+
+UDataflowEditorMode* SDataflowConstructionViewport::GetEdMode() const
+{
+	if (const FEditorModeTools* const EditorModeTools = Client->GetModeTools())
+	{
+		if (UDataflowEditorMode* const DataflowEdMode = Cast<UDataflowEditorMode>(EditorModeTools->GetActiveScriptableMode(UDataflowEditorMode::EM_DataflowEditorModeId)))
+		{
+			return DataflowEdMode;
+		}
+	}
+	return nullptr;
+}
+
+void SDataflowConstructionViewport::BindCommands()
+{
+	SAssetEditorViewport::BindCommands();
+}
+
+bool SDataflowConstructionViewport::IsVisible() const
+{
+	// Intentionally not calling SEditorViewport::IsVisible because it will return false if our simulation is more than 250ms.
+	return ViewportWidget.IsValid();
+}
+
+TSharedRef<class SEditorViewport> SDataflowConstructionViewport::GetViewportWidget()
+{
+	return SharedThis(this);
+}
+
+TSharedPtr<FExtender> SDataflowConstructionViewport::GetExtenders() const
+{
+	TSharedPtr<FExtender> Result(MakeShareable(new FExtender));
+	return Result;
+}
+
+void SDataflowConstructionViewport::OnFloatingButtonClicked()
+{
+}
+
+float SDataflowConstructionViewport:: GetViewMinInput() const
+{
+	return static_cast<FDataflowPreviewScene*>(Client->GetPreviewScene())->GetDataflowContent()->GetSimulationRange()[0];
+}
+
+float SDataflowConstructionViewport::GetViewMaxInput() const
+{
+	return static_cast<FDataflowPreviewScene*>(Client->GetPreviewScene())->GetDataflowContent()->GetSimulationRange()[1];
+}
+
+
+#undef LOCTEXT_NAMESPACE
diff --git a/Engine/Plugins/Experimental/Dataflow/Source/DataflowEditor/Private/Dataflow/DataflowConstructionViewportToolbar.cpp b/Engine/Plugins/Experimental/Dataflow/Source/DataflowEditor/Private/Dataflow/DataflowConstructionViewportToolbar.cpp
--- a/Engine/Plugins/Experimental/Dataflow/Source/DataflowEditor/Private/Dataflow/DataflowConstructionViewportToolbar.cpp
+++ b/Engine/Plugins/Experimental/Dataflow/Source/DataflowEditor/Private/Dataflow/DataflowConstructionViewportToolbar.cpp
@@ -1,8 +1,8 @@
 // Copyright Epic Games, Inc. All Rights Reserved.
 #include "Dataflow/DataflowConstructionViewportToolbar.h"
-#include "Dataflow/DataflowEditorViewport.h"
+#include "Dataflow/DataflowConstructionViewport.h"
 
-void SDataflowConstructionViewportSelectionToolBar::Construct(const FArguments& InArgs, TSharedPtr<SDataflowEditorViewport> InDataflowViewport)
+void SDataflowConstructionViewportSelectionToolBar::Construct(const FArguments& InArgs, TSharedPtr<SDataflowConstructionViewport> InDataflowViewport)
 {
 	EditorViewport = InDataflowViewport;
 	
diff --git a/Engine/Plugins/Experimental/Dataflow/Source/DataflowEditor/Private/Dataflow/DataflowEditorMode.cpp b/Engine/Plugins/Experimental/Dataflow/Source/DataflowEditor/Private/Dataflow/DataflowEditorMode.cpp
--- a/Engine/Plugins/Experimental/Dataflow/Source/DataflowEditor/Private/Dataflow/DataflowEditorMode.cpp
+++ b/Engine/Plugins/Experimental/Dataflow/Source/DataflowEditor/Private/Dataflow/DataflowEditorMode.cpp
@@ -295,6 +295,10 @@ void UDataflowEditorMode::Exit()
 	ConstructionScene->ResetConstructionScene();
 	ConstructionScene = nullptr;
 
+	SimulationScene->ResetSimulationScene();
+	SimulationScene = nullptr;
+
+
 	Super::Exit();
 }
 
@@ -320,6 +324,11 @@ void UDataflowEditorMode::SetDataflowConstructionScene(FDataflowConstructionScen
 	PreviewToolManager->OnToolEnded.AddSP(DataflowModeToolkit, &FDataflowEditorModeToolkit::OnToolEnded);
 }
 
+void UDataflowEditorMode::SetDataflowSimlationScene(FDataflowSimulationScene* InSimulationScene)
+{
+	SimulationScene = InSimulationScene;
+}
+
 void UDataflowEditorMode::CreateToolTargets(const TArray<TObjectPtr<UObject>>& AssetsIn)
 {
 	ToolTargets.Reset();
@@ -349,7 +358,7 @@ bool UDataflowEditorMode::IsComponentSelected(const UPrimitiveComponent* InCompo
 	return false;
 }
 
-void UDataflowEditorMode::RefocusRestSpaceViewportClient()
+void UDataflowEditorMode::RefocusConstructionViewportClient()
 {
 	TSharedPtr<FDataflowConstructionViewportClient, ESPMode::ThreadSafe> PinnedVC = ConstructionViewportClient.Pin();
 	if (PinnedVC.IsValid())
@@ -378,7 +387,25 @@ void UDataflowEditorMode::RefocusRestSpaceViewportClient()
 	}
 }
 
-void UDataflowEditorMode::FirstTimeFocusRestSpaceViewport()
+void UDataflowEditorMode::RefocusSimulationViewportClient()
+{
+	TSharedPtr<FDataflowSimulationViewportClient, ESPMode::ThreadSafe> PinnedVC = SimulationViewportClient.Pin();
+	if (PinnedVC.IsValid())
+	{
+		// This will happen in FocusViewportOnBox anyways; do it now to get a consistent end result
+		PinnedVC->ToggleOrbitCamera(false);
+
+		const FBox SceneBounds = SceneBoundingBox();
+
+		// 3D space
+		PinnedVC->SetInitialViewTransform(ELevelViewportType::LVT_Perspective, FVector(0, 150, 200), FRotator(0, 0, 0), DEFAULT_ORTHOZOOM);
+
+		constexpr bool bInstant = true;
+		PinnedVC->FocusViewportOnBox(SceneBounds, bInstant);
+	}
+}
+
+void UDataflowEditorMode::FirstTimeFocusConstructionViewport()
 {
 	// If this is the first time seeing a valid 2D or 3D mesh, refocus the camera on it.
 	const bool bIsValid = (ConstructionScene->HasRenderableGeometry());
@@ -389,16 +416,27 @@ void UDataflowEditorMode::FirstTimeFocusRestSpaceViewport()
 		if (bIs2D && bFirstValid2DMesh)
 		{
 			bFirstValid2DMesh = false;
-			RefocusRestSpaceViewportClient();
+			RefocusConstructionViewportClient();
 		}
 		else if (!bIs2D && bFirstValid3DMesh)
 		{
 			bFirstValid3DMesh = false;
-			RefocusRestSpaceViewportClient();
+			RefocusConstructionViewportClient();
 		}
 	}
 }
 
+void UDataflowEditorMode::FirstTimeFocusSimulationViewport()
+{
+	// If this is the first time seeing a valid 2D or 3D mesh, refocus the camera on it.
+	const bool bIsValid = (SimulationScene->HasRenderableGeometry());
+
+	if (bIsValid)
+	{
+		RefocusSimulationViewportClient();
+	}
+}
+
 void UDataflowEditorMode::InitializeTargets(const TArray<TObjectPtr<UObject>>& ObjectsToEdit)
 {
 	UBaseCharacterFXEditorMode::InitializeTargets(ObjectsToEdit);
@@ -447,17 +485,29 @@ void UDataflowEditorMode::ModeTick(float DeltaTime)
 	}
 }
 
-void UDataflowEditorMode::RestSpaceViewportResized(FViewport* RestspaceViewport, uint32 /*Unused*/)
+void UDataflowEditorMode::ConstructionViewportResized(FViewport* ConstructionViewport, uint32 /*Unused*/)
+{
+	// We'd like to call RefocusConstructionViewportClient() when the viewport is first created, however in Ortho mode the
+	// viewport needs to have non-zero size for FocusViewportOnBox() to work properly. So we wait until the viewport is resized here.
+	if (bShouldFocusConstructionView && ConstructionViewport && ConstructionViewport->GetSizeXY().X > 0 && ConstructionViewport->GetSizeXY().Y > 0)
+	{
+		RefocusConstructionViewportClient();
+		bShouldFocusConstructionView = false;
+	}
+}
+
+void UDataflowEditorMode::SimulationViewportResized(FViewport* SimulationViewport, uint32 /*Unused*/)
 {
-	// We'd like to call RefocusRestSpaceViewportClient() when the viewport is first created, however in Ortho mode the
+	// We'd like to call RefocusConstructionViewportClient() when the viewport is first created, however in Ortho mode the
 	// viewport needs to have non-zero size for FocusViewportOnBox() to work properly. So we wait until the viewport is resized here.
-	if (bShouldFocusRestSpaceView && RestspaceViewport && RestspaceViewport->GetSizeXY().X > 0 && RestspaceViewport->GetSizeXY().Y > 0)
+	if (bShouldFocusSimulationView && SimulationViewport && SimulationViewport->GetSizeXY().X > 0 && SimulationViewport->GetSizeXY().Y > 0)
 	{
-		RefocusRestSpaceViewportClient();
-		bShouldFocusRestSpaceView = false;
+		RefocusSimulationViewportClient();
+		bShouldFocusSimulationView = false;
 	}
 }
 
+
 FBox UDataflowEditorMode::SceneBoundingBox() const
 {
 	return ConstructionScene->GetBoundingBox();
@@ -505,7 +555,7 @@ void UDataflowEditorMode::SetConstructionViewMode(Dataflow::EDataflowPatternVert
 	}
 
 	// If we are switching to a mode with a valid mesh for the first time, focus the camera on it
-	FirstTimeFocusRestSpaceViewport();
+	FirstTimeFocusConstructionViewport();
 
 	if (bEndedActiveTool)
 	{
@@ -572,7 +622,7 @@ bool UDataflowEditorMode::CanSetConstructionViewWireframeActive() const
 	return DataflowToolBuilder->CanSetConstructionViewWireframeActive();
 }
 
-void UDataflowEditorMode::SetRestSpaceViewportClient(TWeakPtr<FDataflowConstructionViewportClient, ESPMode::ThreadSafe> InViewportClient)
+void UDataflowEditorMode::SetConstructionViewportClient(TWeakPtr<FDataflowConstructionViewportClient, ESPMode::ThreadSafe> InViewportClient)
 {
 	ConstructionViewportClient = InViewportClient;
 
@@ -584,7 +634,21 @@ void UDataflowEditorMode::SetRestSpaceViewportClient(TWeakPtr<FDataflowConstruct
 
 		if (VC->Viewport)
 		{
-			VC->Viewport->ViewportResizedEvent.AddUObject(this, &UDataflowEditorMode::RestSpaceViewportResized);
+			VC->Viewport->ViewportResizedEvent.AddUObject(this, &UDataflowEditorMode::ConstructionViewportResized);
+		}
+	}
+}
+
+void UDataflowEditorMode::SetSimulationViewportClient(TWeakPtr<FDataflowSimulationViewportClient, ESPMode::ThreadSafe> InViewportClient)
+{
+	SimulationViewportClient = InViewportClient;
+
+	TSharedPtr<FDataflowSimulationViewportClient> VC = SimulationViewportClient.Pin();
+	if (VC.IsValid())
+	{
+		if (VC->Viewport)
+		{
+			VC->Viewport->ViewportResizedEvent.AddUObject(this, &UDataflowEditorMode::SimulationViewportResized);
 		}
 	}
 }
@@ -595,13 +659,13 @@ void UDataflowEditorMode::InitializeContextObject()
 
 	if (TObjectPtr<UDataflowBaseContent> DataflowContent = ConstructionScene->GetDataflowContent())
 	{
-		UEditorInteractiveToolsContext* const RestSpaceToolsContext = GetInteractiveToolsContext();
+		UEditorInteractiveToolsContext* const ConstructionToolsContext = GetInteractiveToolsContext();
 
-		UDataflowContextObject* ContextObject = RestSpaceToolsContext->ContextObjectStore->FindContext<UDataflowContextObject>();
+		UDataflowContextObject* ContextObject = ConstructionToolsContext->ContextObjectStore->FindContext<UDataflowContextObject>();
 		if (!ContextObject)
 		{
 			ContextObject = DataflowContent;
-			RestSpaceToolsContext->ContextObjectStore->AddContextObject(ContextObject);
+			ConstructionToolsContext->ContextObjectStore->AddContextObject(ContextObject);
 		}
 
 		check(ContextObject);
@@ -612,10 +676,10 @@ void UDataflowEditorMode::InitializeContextObject()
 
 void UDataflowEditorMode::DeleteContextObject()
 {
-	UEditorInteractiveToolsContext* const RestSpaceToolsContext = GetInteractiveToolsContext();
-	if (UDataflowContextObject* ContextObject = RestSpaceToolsContext->ContextObjectStore->FindContext<UDataflowContextObject>())
+	UEditorInteractiveToolsContext* const ConstructionToolsContext = GetInteractiveToolsContext();
+	if (UDataflowContextObject* ContextObject = ConstructionToolsContext->ContextObjectStore->FindContext<UDataflowContextObject>())
 	{
-		RestSpaceToolsContext->ContextObjectStore->RemoveContextObject(ContextObject);
+		ConstructionToolsContext->ContextObjectStore->RemoveContextObject(ContextObject);
 	}
 }
 
diff --git a/Engine/Plugins/Experimental/Dataflow/Source/DataflowEditor/Private/Dataflow/DataflowEditorModeToolkit.cpp b/Engine/Plugins/Experimental/Dataflow/Source/DataflowEditor/Private/Dataflow/DataflowEditorModeToolkit.cpp
--- a/Engine/Plugins/Experimental/Dataflow/Source/DataflowEditor/Private/Dataflow/DataflowEditorModeToolkit.cpp
+++ b/Engine/Plugins/Experimental/Dataflow/Source/DataflowEditor/Private/Dataflow/DataflowEditorModeToolkit.cpp
@@ -6,7 +6,7 @@
 #include "Dataflow/DataflowEditorCommands.h"
 #include "Dataflow/DataflowEditorMode.h"
 #include "Dataflow/DataflowEditorStyle.h"
-#include "Dataflow/DataflowEditorViewport.h"
+#include "Dataflow/DataflowConstructionViewport.h"
 #include "Dataflow/DataflowEditorScenes.h"
 #include "EdModeInteractiveToolsContext.h"
 #include "Framework/Application/SlateApplication.h"
@@ -174,7 +174,7 @@ SBaseCharacterFXEditorViewport* FDataflowEditorModeToolkit::GetViewportWidgetFor
 					//}
 					//else
 					//{
-					if (const TSharedPtr<SDataflowEditorViewport> Widget = ConstructionViewportWidget.Pin())
+					if (const TSharedPtr<SDataflowConstructionViewport> Widget = ConstructionViewportWidget.Pin())
 					{
 						return Widget.Get();
 					}
@@ -283,10 +283,16 @@ void FDataflowEditorModeToolkit::OnToolEnded(UInteractiveToolManager* Manager, U
 
 }
 
-void FDataflowEditorModeToolkit::SetConstructionViewportWidget(TWeakPtr<SDataflowEditorViewport> InConstructionViewportWidget)
+void FDataflowEditorModeToolkit::SetConstructionViewportWidget(TWeakPtr<SDataflowConstructionViewport> InConstructionViewportWidget)
 {
 	ConstructionViewportWidget = InConstructionViewportWidget;
 }
 
+void FDataflowEditorModeToolkit::SetSimulationViewportWidget(TWeakPtr<SDataflowSimulationViewport> InSimulationViewportWidget)
+{
+	SimulationViewportWidget = InSimulationViewportWidget;
+}
+
+
 
 #undef LOCTEXT_NAMESPACE
diff --git a/Engine/Plugins/Experimental/Dataflow/Source/DataflowEditor/Private/Dataflow/DataflowEditorToolkit.cpp b/Engine/Plugins/Experimental/Dataflow/Source/DataflowEditor/Private/Dataflow/DataflowEditorToolkit.cpp
--- a/Engine/Plugins/Experimental/Dataflow/Source/DataflowEditor/Private/Dataflow/DataflowEditorToolkit.cpp
+++ b/Engine/Plugins/Experimental/Dataflow/Source/DataflowEditor/Private/Dataflow/DataflowEditorToolkit.cpp
@@ -15,7 +15,7 @@
 #include "Dataflow/DataflowEditorModeToolkit.h"
 #include "Dataflow/DataflowEditorModule.h"
 #include "Dataflow/DataflowEditorModeUILayer.h"
-#include "Dataflow/DataflowEditorViewport.h"
+#include "Dataflow/DataflowConstructionViewport.h"
 #include "Dataflow/DataflowEdNode.h"
 #include "Dataflow/DataflowGraphEditor.h"
 #include "Dataflow/DataflowNodeFactory.h"
@@ -345,12 +345,14 @@ void FDataflowEditorToolkit::PostInitAssetEditor()
 	SetCommonViewportClientOptions(ViewportClient.Get());
 	SetCommonViewportClientOptions(SimulationViewportClient.Get());
 
-	//@todo(brice) : FChaosClothAssetEditorToolkit::PostInitAssetEditor() has alot more going on. 
+	//@todo(Dataflow) : FChaosClothAssetEditorToolkit::PostInitAssetEditor() has alot more going on. 
 
 	//@todo(brice) : can we remove this coupling of the viewport and mode
 	UDataflowEditorMode* const DataflowMode = CastChecked<UDataflowEditorMode>(EditorModeManager->GetActiveScriptableMode(UDataflowEditorMode::EM_DataflowEditorModeId));
-	const TWeakPtr<FViewportClient> WeakViewportClient(ViewportClient);
-	DataflowMode->SetRestSpaceViewportClient(StaticCastWeakPtr<FDataflowConstructionViewportClient>(WeakViewportClient));
+	const TWeakPtr<FViewportClient> WeakConstructionViewportClient(ViewportClient);
+	DataflowMode->SetConstructionViewportClient(StaticCastWeakPtr<FDataflowConstructionViewportClient>(WeakConstructionViewportClient));
+	const TWeakPtr<FViewportClient> WeakSimulationViewportClient(SimulationViewportClient);
+	DataflowMode->SetSimulationViewportClient(StaticCastWeakPtr<FDataflowSimulationViewportClient>(WeakSimulationViewportClient));
 }
 
 void FDataflowEditorToolkit::InitializeEdMode(UBaseCharacterFXEditorMode* EdMode)
@@ -372,7 +374,8 @@ void FDataflowEditorToolkit::InitializeEdMode(UBaseCharacterFXEditorMode* EdMode
 	if (TSharedPtr<FModeToolkit> ModeToolkit = DataflowMode->GetToolkit().Pin())
 	{
 		FDataflowEditorModeToolkit* DataflowModeToolkit = static_cast<FDataflowEditorModeToolkit*>(ModeToolkit.Get());
-		DataflowModeToolkit->SetConstructionViewportWidget(DataflowEditorViewport);
+		DataflowModeToolkit->SetConstructionViewportWidget(DataflowConstructionViewport);
+		DataflowModeToolkit->SetSimulationViewportWidget(DataflowSimulationViewport);
 
 		FName ParentToolbarName;
 		const FName ToolBarName = GetToolMenuToolbarName(ParentToolbarName);
@@ -431,8 +434,20 @@ AssetEditorViewportFactoryFunction FDataflowEditorToolkit::GetViewportDelegate()
 {
 	AssetEditorViewportFactoryFunction TempViewportDelegate = [this](FAssetEditorViewportConstructionArgs InArgs)
 	{
-		return SAssignNew(DataflowEditorViewport, SDataflowEditorViewport, InArgs)
+		TSharedRef<SDataflowConstructionViewport> Viewport = SAssignNew(DataflowConstructionViewport, SDataflowConstructionViewport, InArgs)
 			.ViewportClient(StaticCastSharedPtr<FDataflowConstructionViewportClient>(ViewportClient));
+
+		if (UDataflowEditorMode* const DataflowMode = Cast<UDataflowEditorMode>(EditorModeManager->GetActiveScriptableMode(UDataflowEditorMode::EM_DataflowEditorModeId)))
+		{
+			if (TSharedPtr<FModeToolkit> ModeToolkit = DataflowMode->GetToolkit().Pin())
+			{
+				if (FDataflowEditorModeToolkit* DataflowModeToolkit = static_cast<FDataflowEditorModeToolkit*>(ModeToolkit.Get()))
+				{
+					DataflowModeToolkit->SetConstructionViewportWidget(DataflowConstructionViewport);
+				}
+			}
+		}
+		return Viewport;
 	};
 
 	return TempViewportDelegate;
@@ -465,8 +480,20 @@ void FDataflowEditorToolkit::CreateSimulationViewportClient()
 
 	SimulationViewportDelegate = [this](FAssetEditorViewportConstructionArgs InArgs)
 	{
-		return SAssignNew(DataflowSimulationViewport, SDataflowEditorViewport, InArgs)
+		TSharedRef<SDataflowSimulationViewport> Viewport = SAssignNew(DataflowSimulationViewport, SDataflowSimulationViewport, InArgs)
 			.ViewportClient(StaticCastSharedPtr<FDataflowSimulationViewportClient>(SimulationViewportClient));
+
+		if (UDataflowEditorMode* const DataflowMode = Cast<UDataflowEditorMode>(EditorModeManager->GetActiveScriptableMode(UDataflowEditorMode::EM_DataflowEditorModeId)))
+		{
+			if (TSharedPtr<FModeToolkit> ModeToolkit = DataflowMode->GetToolkit().Pin())
+			{
+				if (FDataflowEditorModeToolkit* DataflowModeToolkit = static_cast<FDataflowEditorModeToolkit*>(ModeToolkit.Get()))
+				{
+					DataflowModeToolkit->SetSimulationViewportWidget(DataflowSimulationViewport);
+				}
+			}
+		}
+		return Viewport;
 	};
 }
 
diff --git a/Engine/Plugins/Experimental/Dataflow/Source/DataflowEditor/Private/Dataflow/DataflowEditorViewport.cpp b/Engine/Plugins/Experimental/Dataflow/Source/DataflowEditor/Private/Dataflow/DataflowEditorViewport.cpp
deleted file mode 100644
diff --git a/Engine/Plugins/Experimental/Dataflow/Source/DataflowEditor/Private/Dataflow/DataflowSimulationViewport.cpp b/Engine/Plugins/Experimental/Dataflow/Source/DataflowEditor/Private/Dataflow/DataflowSimulationViewport.cpp
new file mode 100644
--- /dev/null
+++ b/Engine/Plugins/Experimental/Dataflow/Source/DataflowEditor/Private/Dataflow/DataflowSimulationViewport.cpp
@@ -0,0 +1,114 @@
+// Copyright Epic Games, Inc. All Rights Reserved.
+
+#include "Dataflow/DataflowSimulationViewport.h"
+
+#include "Dataflow/DataflowActor.h"
+#include "Dataflow/DataflowEditorMode.h"
+#include "Dataflow/DataflowSimulationViewportClient.h"
+#include "Dataflow/DataflowEditorToolkit.h"
+#include "EditorModeManager.h"
+#include "Dataflow/DataflowContent.h"
+#include "Dataflow/DataflowEditorScenes.h"
+#include "Dataflow/DataflowSimulationPanel.h"
+
+#define LOCTEXT_NAMESPACE "SDataflowSimulationViewport"
+
+
+SDataflowSimulationViewport::SDataflowSimulationViewport()
+{
+}
+
+void SDataflowSimulationViewport::Construct(const FArguments& InArgs, const FAssetEditorViewportConstructionArgs& InViewportConstructionArgs)
+{
+	SAssetEditorViewport::FArguments ParentArgs;
+	ParentArgs._EditorViewportClient = InArgs._ViewportClient;
+	SAssetEditorViewport::Construct(ParentArgs, InViewportConstructionArgs);
+	Client->VisibilityDelegate.BindSP(this, &SDataflowSimulationViewport::IsVisible);
+
+	if(static_cast<FDataflowSimulationScene*>(Client->GetPreviewScene())->CanRunSimulation())
+	{
+		TSharedPtr<FDataflowSimulationViewportClient> DataflowClient = StaticCastSharedPtr<FDataflowSimulationViewportClient>(Client);
+		TWeakPtr<FDataflowSimulationScene> SimulationScene = DataflowClient->GetDataflowEditorToolkit().Pin()->GetSimulationScene();
+            
+		ViewportOverlay->AddSlot()
+		[
+			SNew(SHorizontalBox)
+			+ SHorizontalBox::Slot()
+			.HAlign(HAlign_Fill)
+			.VAlign(VAlign_Bottom)
+			.FillWidth(1)
+			.Padding(10.0f, 0.0f)
+			[
+				SNew(SBorder)
+				.BorderImage(FAppStyle::Get().GetBrush("EditorViewport.OverlayBrush"))
+				.Visibility(EVisibility::Visible)
+				.Padding(10.0f, 2.0f)
+				[
+					SNew(SDataflowSimulationPanel, SimulationScene)
+					.ViewInputMin(this, &SDataflowSimulationViewport::GetViewMinInput)
+					.ViewInputMax(this, &SDataflowSimulationViewport::GetViewMaxInput)
+				]
+			]
+		];
+	}
+}
+
+void SDataflowSimulationViewport::OnFocusViewportToSelection()
+{
+	if(const FDataflowPreviewScene* PreviewScene = static_cast<FDataflowPreviewScene*>(Client->GetPreviewScene()))
+	{
+		const FBox SceneBoundingBox = PreviewScene->GetBoundingBox();
+		Client->FocusViewportOnBox(SceneBoundingBox);
+	}
+}
+
+UDataflowEditorMode* SDataflowSimulationViewport::GetEdMode() const
+{
+	if (const FEditorModeTools* const EditorModeTools = Client->GetModeTools())
+	{
+		if (UDataflowEditorMode* const DataflowEdMode = Cast<UDataflowEditorMode>(EditorModeTools->GetActiveScriptableMode(UDataflowEditorMode::EM_DataflowEditorModeId)))
+		{
+			return DataflowEdMode;
+		}
+	}
+	return nullptr;
+}
+
+void SDataflowSimulationViewport::BindCommands()
+{
+	SAssetEditorViewport::BindCommands();
+}
+
+bool SDataflowSimulationViewport::IsVisible() const
+{
+	// Intentionally not calling SEditorViewport::IsVisible because it will return false if our simulation is more than 250ms.
+	return ViewportWidget.IsValid();
+}
+
+TSharedRef<class SEditorViewport> SDataflowSimulationViewport::GetViewportWidget()
+{
+	return SharedThis(this);
+}
+
+TSharedPtr<FExtender> SDataflowSimulationViewport::GetExtenders() const
+{
+	TSharedPtr<FExtender> Result(MakeShareable(new FExtender));
+	return Result;
+}
+
+void SDataflowSimulationViewport::OnFloatingButtonClicked()
+{
+}
+
+float SDataflowSimulationViewport:: GetViewMinInput() const
+{
+	return static_cast<FDataflowPreviewScene*>(Client->GetPreviewScene())->GetDataflowContent()->GetSimulationRange()[0];
+}
+
+float SDataflowSimulationViewport::GetViewMaxInput() const
+{
+	return static_cast<FDataflowPreviewScene*>(Client->GetPreviewScene())->GetDataflowContent()->GetSimulationRange()[1];
+}
+
+
+#undef LOCTEXT_NAMESPACE
diff --git a/Engine/Plugins/Experimental/Dataflow/Source/DataflowEditor/Public/Dataflow/DataflowConstructionViewport.h b/Engine/Plugins/Experimental/Dataflow/Source/DataflowEditor/Public/Dataflow/DataflowConstructionViewport.h
new file mode 100644
--- /dev/null
+++ b/Engine/Plugins/Experimental/Dataflow/Source/DataflowEditor/Public/Dataflow/DataflowConstructionViewport.h
@@ -0,0 +1,43 @@
+// Copyright Epic Games, Inc. All Rights Reserved.
+#pragma once
+
+#include "CoreMinimal.h"
+#include "SBaseCharacterFXEditorViewport.h"
+#include "SCommonEditorViewportToolbarBase.h"
+
+class ADataflowActor;
+class FAdvancedPreviewScene;
+class UDataflowEditorMode;
+class FDataflowConstructionViewportClient;
+
+// ----------------------------------------------------------------------------------
+
+class SDataflowConstructionViewport : public SBaseCharacterFXEditorViewport, public ICommonEditorViewportToolbarInfoProvider
+{
+public:
+	SLATE_BEGIN_ARGS(SDataflowConstructionViewport) {}
+	SLATE_ARGUMENT(TSharedPtr<FEditorViewportClient>, ViewportClient)
+		SLATE_END_ARGS()
+
+	SDataflowConstructionViewport();
+
+	void Construct(const FArguments& InArgs, const FAssetEditorViewportConstructionArgs& InViewportConstructionArgs);
+
+	// SEditorViewport
+	virtual void BindCommands() override;
+	virtual TSharedPtr<SWidget> MakeViewportToolbar() override;
+	virtual bool IsVisible() const override;
+	virtual void OnFocusViewportToSelection() override;
+
+	// ICommonEditorViewportToolbarInfoProvider
+	virtual TSharedRef<class SEditorViewport> GetViewportWidget() override;
+	virtual TSharedPtr<FExtender> GetExtenders() const override;
+	virtual void OnFloatingButtonClicked() override;
+
+private:
+	float GetViewMinInput() const;
+	float GetViewMaxInput() const;
+
+	UDataflowEditorMode* GetEdMode() const;
+};
+
diff --git a/Engine/Plugins/Experimental/Dataflow/Source/DataflowEditor/Public/Dataflow/DataflowConstructionViewportToolbar.h b/Engine/Plugins/Experimental/Dataflow/Source/DataflowEditor/Public/Dataflow/DataflowConstructionViewportToolbar.h
--- a/Engine/Plugins/Experimental/Dataflow/Source/DataflowEditor/Public/Dataflow/DataflowConstructionViewportToolbar.h
+++ b/Engine/Plugins/Experimental/Dataflow/Source/DataflowEditor/Public/Dataflow/DataflowConstructionViewportToolbar.h
@@ -5,7 +5,7 @@
 #include "SCommonEditorViewportToolbarBase.h"
 
 // ----------------------------------------------------------------------------------
-class SDataflowEditorViewport;
+class SDataflowConstructionViewport;
 
 /** Base toolbar for the dataflow. Should be extended to add more features */
 class SDataflowConstructionViewportSelectionToolBar : public SCommonEditorViewportToolbarBase
@@ -15,9 +15,9 @@ public:
 	SLATE_END_ARGS()
 
 	/** Constructs this widget with the given parameters */
-	void Construct(const FArguments& InArgs, TSharedPtr<SDataflowEditorViewport> InDataflowViewport);
+	void Construct(const FArguments& InArgs, TSharedPtr<SDataflowConstructionViewport> InDataflowViewport);
 	
 private:
 	/** Reference to the parent viewport */
-	TWeakPtr<SDataflowEditorViewport> EditorViewport;
+	TWeakPtr<SDataflowConstructionViewport> EditorViewport;
 };
\ No newline at end of file
diff --git a/Engine/Plugins/Experimental/Dataflow/Source/DataflowEditor/Public/Dataflow/DataflowEditorMode.h b/Engine/Plugins/Experimental/Dataflow/Source/DataflowEditor/Public/Dataflow/DataflowEditorMode.h
--- a/Engine/Plugins/Experimental/Dataflow/Source/DataflowEditor/Public/Dataflow/DataflowEditorMode.h
+++ b/Engine/Plugins/Experimental/Dataflow/Source/DataflowEditor/Public/Dataflow/DataflowEditorMode.h
@@ -10,6 +10,8 @@ class UDataflowComponent;
 class FDataflowConstructionScene;
 class SDataflowGraphEditor;
 class FDataflowConstructionViewportClient;
+class FDataflowSimulationScene;
+class FDataflowSimulationViewportClient;
 class IDataflowEditorToolBuilder;
 struct FManagedArrayCollection;
 class UEdGraphNode;
@@ -64,6 +66,9 @@ public:
 	FDataflowConstructionScene* GetDataflowConstructionScene() { return ConstructionScene; }
 	const FDataflowConstructionScene* GetDataflowConstructionScene() const { return ConstructionScene; }
 
+	void SetDataflowSimlationScene(FDataflowSimulationScene* InSimulationScene);
+
+
 	/** Set the data flow graph editor to create nodes once the tools have ended*/
 	void SetDataflowGraphEditor(TSharedPtr<SDataflowGraphEditor> DataflowGraphEditor);
 
@@ -106,9 +111,13 @@ private:
 
 	void RegisterAddNodeCommand(TSharedPtr<FUICommandInfo> AddNodeCommand, const FName& NewNodeType, TSharedPtr<FUICommandInfo> StartToolCommand);
 
-	void SetRestSpaceViewportClient(TWeakPtr<FDataflowConstructionViewportClient, ESPMode::ThreadSafe> ViewportClient);
-	void RefocusRestSpaceViewportClient();
-	void FirstTimeFocusRestSpaceViewport();
+	void SetSimulationViewportClient(TWeakPtr<FDataflowSimulationViewportClient, ESPMode::ThreadSafe>);
+	void RefocusSimulationViewportClient();
+	void FirstTimeFocusSimulationViewport();
+
+	void SetConstructionViewportClient(TWeakPtr<FDataflowConstructionViewportClient, ESPMode::ThreadSafe>);
+	void RefocusConstructionViewportClient();
+	void FirstTimeFocusConstructionViewport();
 	bool IsComponentSelected(const UPrimitiveComponent* InComponent);
 	void OnDataflowNodeDeleted(const TSet<UObject*>& DeletedNodes);
 
@@ -148,6 +157,7 @@ private:
 
 	/** Dataflow preview scene from the toolkit */
 	FDataflowConstructionScene* ConstructionScene = nullptr;
+	FDataflowSimulationScene* SimulationScene = nullptr;
 
 	/** Correspondence between node types and commands to launch tools */
 	TMap<FName, TSharedPtr<const FUICommandInfo>> NodeTypeToToolCommandMap;
@@ -170,16 +180,19 @@ private:
 	TObjectPtr<UEditorInteractiveToolsContext> ActiveToolsContext = nullptr;
 
 	TWeakPtr<FDataflowConstructionViewportClient, ESPMode::ThreadSafe> ConstructionViewportClient;
+	TWeakPtr<FDataflowSimulationViewportClient, ESPMode::ThreadSafe> SimulationViewportClient;
 
 	// The first time we get a valid mesh, refocus the camera on it
 	bool bFirstValid2DMesh = true;
 	bool bFirstValid3DMesh = true;
 
 	// Whether the rest space viewport should focus on the rest space mesh on the next tick
-	bool bShouldFocusRestSpaceView = true;
+	bool bShouldFocusConstructionView = true;
+	bool bShouldFocusSimulationView = true;
 
-	//@todo(brice) : Move to Construction
-	void RestSpaceViewportResized(FViewport* RestspaceViewport, uint32 Unused);
+	//@todo(Dataflow) : Move to Construction and Simulaiton
+	void ConstructionViewportResized(FViewport* ConstructionViewport, uint32 Unused);
+	void SimulationViewportResized(FViewport* ConstructionViewport, uint32 Unused);
 
 	// Dataflow node type whose corresponding tool should be started on the next Tick
 	FName NodeTypeForPendingToolStart;
diff --git a/Engine/Plugins/Experimental/Dataflow/Source/DataflowEditor/Public/Dataflow/DataflowEditorModeToolkit.h b/Engine/Plugins/Experimental/Dataflow/Source/DataflowEditor/Public/Dataflow/DataflowEditorModeToolkit.h
--- a/Engine/Plugins/Experimental/Dataflow/Source/DataflowEditor/Public/Dataflow/DataflowEditorModeToolkit.h
+++ b/Engine/Plugins/Experimental/Dataflow/Source/DataflowEditor/Public/Dataflow/DataflowEditorModeToolkit.h
@@ -5,7 +5,8 @@
 #include "BaseCharacterFXEditorModeToolkit.h"
 
 class UEditorInteractiveToolsContext;
-class SDataflowEditorViewport;
+class SDataflowConstructionViewport;
+class SDataflowSimulationViewport;
 class SBaseCharacterFXEditorViewport;
 
 /**
@@ -35,8 +36,8 @@ public:
 	virtual FName GetToolkitFName() const override;
 	virtual FText GetBaseToolkitName() const override;
 
-	void SetConstructionViewportWidget(TWeakPtr<SDataflowEditorViewport>);
-	//void SetPreviewViewportWidget(TWeakPtr<SChaosClothAssetEditor3DViewport>);
+	void SetConstructionViewportWidget(TWeakPtr<SDataflowConstructionViewport>);
+	void SetSimulationViewportWidget(TWeakPtr<SDataflowSimulationViewport>);
 
 private:
 
@@ -46,6 +47,6 @@ private:
 
 	UEditorInteractiveToolsContext* GetCurrentToolsContext();
 
-	TWeakPtr<SDataflowEditorViewport> ConstructionViewportWidget;
-	//TWeakPtr<SChaosClothAssetEditor3DViewport> PreviewViewportWidget;
+	TWeakPtr<SDataflowConstructionViewport> ConstructionViewportWidget;
+	TWeakPtr<SDataflowSimulationViewport> SimulationViewportWidget;
 };
diff --git a/Engine/Plugins/Experimental/Dataflow/Source/DataflowEditor/Public/Dataflow/DataflowEditorScenes.h b/Engine/Plugins/Experimental/Dataflow/Source/DataflowEditor/Public/Dataflow/DataflowEditorScenes.h
--- a/Engine/Plugins/Experimental/Dataflow/Source/DataflowEditor/Public/Dataflow/DataflowEditorScenes.h
+++ b/Engine/Plugins/Experimental/Dataflow/Source/DataflowEditor/Public/Dataflow/DataflowEditorScenes.h
@@ -159,6 +159,10 @@ public:
 	
 	/** Check if the preview scene can run simulation */
 	virtual bool CanRunSimulation() const { return true; }
+
+	void ResetSimulationScene() {}
+	bool HasRenderableGeometry() { return true; }
+
 };
 
 
diff --git a/Engine/Plugins/Experimental/Dataflow/Source/DataflowEditor/Public/Dataflow/DataflowEditorToolkit.h b/Engine/Plugins/Experimental/Dataflow/Source/DataflowEditor/Public/Dataflow/DataflowEditorToolkit.h
--- a/Engine/Plugins/Experimental/Dataflow/Source/DataflowEditor/Public/Dataflow/DataflowEditorToolkit.h
+++ b/Engine/Plugins/Experimental/Dataflow/Source/DataflowEditor/Public/Dataflow/DataflowEditorToolkit.h
@@ -10,7 +10,8 @@
 #include "TickableEditorObject.h"
 #include "Dataflow/DataflowSelectionView.h"
 #include "Dataflow/DataflowCollectionSpreadSheet.h"
-#include "Dataflow/DataflowEditorViewport.h"
+#include "Dataflow/DataflowConstructionViewport.h"
+#include "Dataflow/DataflowSimulationViewport.h"
 
 class FEditorViewportTabContent;
 class IDetailsView;
@@ -131,8 +132,8 @@ private:
 	static const FName SimulationViewportTabId;
 
 	// List of all the widgets shared ptr that will be built in the editor
-	TSharedPtr<SDataflowEditorViewport> DataflowEditorViewport;
-	TSharedPtr<SDataflowEditorViewport> DataflowSimulationViewport;
+	TSharedPtr<SDataflowConstructionViewport> DataflowConstructionViewport;
+	TSharedPtr<SDataflowSimulationViewport> DataflowSimulationViewport;
 	TSharedPtr<SDataflowGraphEditor> GraphEditor;
 	TSharedPtr<IStructureDetailsView> NodeDetailsEditor;
 	TSharedPtr<class ISkeletonTree> SkeletalEditor;
diff --git a/Engine/Plugins/Experimental/Dataflow/Source/DataflowEditor/Public/Dataflow/DataflowEditorViewport.h b/Engine/Plugins/Experimental/Dataflow/Source/DataflowEditor/Public/Dataflow/DataflowEditorViewport.h
deleted file mode 100644
diff --git a/Engine/Plugins/Experimental/Dataflow/Source/DataflowEditor/Public/Dataflow/DataflowSimulationViewport.h b/Engine/Plugins/Experimental/Dataflow/Source/DataflowEditor/Public/Dataflow/DataflowSimulationViewport.h
new file mode 100644
--- /dev/null
+++ b/Engine/Plugins/Experimental/Dataflow/Source/DataflowEditor/Public/Dataflow/DataflowSimulationViewport.h
@@ -0,0 +1,42 @@
+// Copyright Epic Games, Inc. All Rights Reserved.
+#pragma once
+
+#include "CoreMinimal.h"
+#include "SBaseCharacterFXEditorViewport.h"
+#include "SCommonEditorViewportToolbarBase.h"
+
+class ADataflowActor;
+class FAdvancedPreviewScene;
+class UDataflowEditorMode;
+class FDataflowConstructionViewportClient;
+
+// ----------------------------------------------------------------------------------
+
+class SDataflowSimulationViewport : public SBaseCharacterFXEditorViewport, public ICommonEditorViewportToolbarInfoProvider
+{
+public:
+	SLATE_BEGIN_ARGS(SDataflowSimulationViewport) {}
+	SLATE_ARGUMENT(TSharedPtr<FEditorViewportClient>, ViewportClient)
+		SLATE_END_ARGS()
+
+	SDataflowSimulationViewport();
+
+	void Construct(const FArguments& InArgs, const FAssetEditorViewportConstructionArgs& InViewportConstructionArgs);
+
+	// SEditorViewport
+	virtual void BindCommands() override;
+	virtual bool IsVisible() const override;
+	virtual void OnFocusViewportToSelection() override;
+
+	// ICommonEditorViewportToolbarInfoProvider
+	virtual TSharedRef<class SEditorViewport> GetViewportWidget() override;
+	virtual TSharedPtr<FExtender> GetExtenders() const override;
+	virtual void OnFloatingButtonClicked() override;
+
+private:
+	float GetViewMinInput() const;
+	float GetViewMaxInput() const;
+
+	UDataflowEditorMode* GetEdMode() const;
+};
+
