Dataflow : Rest all the cached timestamps on the dataflow context cache on load.

diff --git a/Engine/Plugins/Experimental/Dataflow/Source/DataflowEditor/Private/Dataflow/AssetDefinition_DataflowContext.cpp b/Engine/Plugins/Experimental/Dataflow/Source/DataflowEditor/Private/Dataflow/AssetDefinition_DataflowContext.cpp
--- a/Engine/Plugins/Experimental/Dataflow/Source/DataflowEditor/Private/Dataflow/AssetDefinition_DataflowContext.cpp
+++ b/Engine/Plugins/Experimental/Dataflow/Source/DataflowEditor/Private/Dataflow/AssetDefinition_DataflowContext.cpp
@@ -131,8 +131,32 @@ namespace DataflowContextDefinitionHelpers
 
 		UDataflowBaseContent* BaseContent = Cast< UDataflowBaseContent>(Asset.Get());
 		if (!BaseContent) return false;
+		if (!BaseContent->GetDataflowContext()) return false;
 
-		BaseContent->SetLastModifiedTimestamp(DataflowAsset->GetRenderingTimestamp().Value+1, false /*bMakeDirty*/);
+		Dataflow::FTimestamp NewTimestamp = DataflowAsset->GetRenderingTimestamp().Value + 1;
+		BaseContent->SetLastModifiedTimestamp(NewTimestamp, false /*bMakeDirty*/);
+
+		TSharedPtr<FEngineContext>& Context = BaseContent->GetDataflowContext();
+		TSet<FContextCacheKey> Keys; Context->GetKeys(Keys);
+		for (FContextCacheKey Key : Keys)
+		{
+			if (TUniquePtr<FContextCacheElementBase>* Data = Context->GetBaseData(Key))
+			{
+				if (*Data)
+				{
+					(*Data)->SetTimestamp(NewTimestamp);
+				}
+				else
+				{
+					return false;
+				}
+			}
+			else
+			{
+				return false;
+			}
+		}
+		Context->Timestamp = NewTimestamp;
 
 		return true;
 	}
diff --git a/Engine/Plugins/Experimental/Dataflow/Source/DataflowEditor/Private/Dataflow/DataflowEditorScenes.cpp b/Engine/Plugins/Experimental/Dataflow/Source/DataflowEditor/Private/Dataflow/DataflowEditorScenes.cpp
--- a/Engine/Plugins/Experimental/Dataflow/Source/DataflowEditor/Private/Dataflow/DataflowEditorScenes.cpp
+++ b/Engine/Plugins/Experimental/Dataflow/Source/DataflowEditor/Private/Dataflow/DataflowEditorScenes.cpp
@@ -23,10 +23,11 @@ bool bDataflowShowFloorDefault = false;
 FAutoConsoleVariableRef CVARDataflowShowFloorDefault(TEXT("p.Dataflow.Editor.ShowFloor"), bDataflowShowFloorDefault, TEXT("Show the floor in the dataflow editor[def:false]"));
 
 
-FDataflowPreviewScene::FDataflowPreviewScene(FPreviewScene::ConstructionValues ConstructionValues, TObjectPtr<UDataflowBaseContent> InEditorContent) 
-	: FAdvancedPreviewScene(ConstructionValues), DataflowContent(InEditorContent)
+FDataflowPreviewScene::FDataflowPreviewScene(FPreviewScene::ConstructionValues ConstructionValues, UDataflowEditor* InEditor)
+	: FAdvancedPreviewScene(ConstructionValues)
+	, DataflowEditor(InEditor)
 {
-	check(DataflowContent);
+	check(DataflowEditor);
 	SetFloorVisibility(bDataflowShowFloorDefault, true);
 
 	RootSceneActor = GetWorld()->SpawnActor<AActor>(AActor::StaticClass());
@@ -35,10 +36,19 @@ FDataflowPreviewScene::FDataflowPreviewScene(FPreviewScene::ConstructionValues C
 FDataflowPreviewScene::~FDataflowPreviewScene()
 {}
 
+TObjectPtr<UDataflowBaseContent> FDataflowPreviewScene::GetDataflowContent() 
+{ 
+	return DataflowEditor->GetDataflowContent();
+}
+
+const TObjectPtr<UDataflowBaseContent> FDataflowPreviewScene::GetDataflowContent() const 
+{ 
+	return DataflowEditor->GetDataflowContent();
+}
+
 void FDataflowPreviewScene::AddReferencedObjects(FReferenceCollector& Collector)
 {
 	FAdvancedPreviewScene::AddReferencedObjects(Collector);
-	Collector.AddReferencedObject(DataflowContent);
 	Collector.AddReferencedObject(RootSceneActor);
 	if (GetDataflowContent())
 	{
@@ -90,8 +100,12 @@ FBox FDataflowPreviewScene::GetBoundingBox() const
 	return SceneBounds;
 }
 
-FDataflowConstructionScene::FDataflowConstructionScene(FPreviewScene::ConstructionValues ConstructionValues, TObjectPtr<UDataflowBaseContent> InEditorContent) 
-	: FDataflowPreviewScene(ConstructionValues,InEditorContent)
+//
+// Construction Scene
+//
+
+FDataflowConstructionScene::FDataflowConstructionScene(FPreviewScene::ConstructionValues ConstructionValues, UDataflowEditor* InEditor)
+	: FDataflowPreviewScene(ConstructionValues, InEditor)
 {}
 
 FDataflowConstructionScene::~FDataflowConstructionScene()
@@ -118,18 +132,21 @@ FORCEINLINE Dataflow::FTimestamp LatestTimestamp(const UDataflow* Dataflow, cons
 
 void FDataflowConstructionScene::TickDataflowScene(const float DeltaSeconds)
 {
-	if (const TSharedPtr<Dataflow::FContext> DataflowContext = DataflowContent->GetDataflowContext())
+	if (TObjectPtr<UDataflowBaseContent> DataflowContent = GetDataflowContent())
 	{
-		if (const UDataflow* Dataflow = DataflowContent->GetDataflowAsset())
+		if (const TSharedPtr<Dataflow::FContext> DataflowContext = DataflowContent->GetDataflowContext())
 		{
-			const Dataflow::FTimestamp SystemTimestamp = LatestTimestamp(Dataflow, DataflowContext.Get());
-			if (SystemTimestamp >= DataflowContent->GetLastModifiedTimestamp() || DataflowContent->IsDirty())
+			if (const UDataflow* Dataflow = DataflowContent->GetDataflowAsset())
 			{
-				DataflowContent->SetLastModifiedTimestamp(SystemTimestamp.Value + 1);
-
-				if(DataflowContent->IsDirty())
+				const Dataflow::FTimestamp SystemTimestamp = LatestTimestamp(Dataflow, DataflowContext.Get());
+				if (SystemTimestamp >= DataflowContent->GetLastModifiedTimestamp() || DataflowContent->IsDirty())
 				{
-					UpdateConstructionScene();
+					DataflowContent->SetLastModifiedTimestamp(SystemTimestamp.Value + 1);
+
+					if (DataflowContent->IsDirty())
+					{
+						UpdateConstructionScene();
+					}
 				}
 			}
 		}
@@ -164,7 +181,7 @@ void FDataflowConstructionScene::UpdateDynamicMeshComponents()
 	// list of UPrimitiveComponents for rendering.
 	ResetDynamicMeshComponents();
 
-	if (DataflowContent)
+	if (TObjectPtr<UDataflowBaseContent> DataflowContent = GetDataflowContent())
 	{
 		const TObjectPtr<UDataflow>& DataflowAsset = DataflowContent->GetDataflowAsset();
 		const TSharedPtr<Dataflow::FEngineContext>& DataflowContext = DataflowContent->GetDataflowContext();
@@ -228,6 +245,7 @@ TObjectPtr<UDynamicMeshComponent>& FDataflowConstructionScene::AddDynamicMeshCom
 	DynamicMeshComponent->SetMesh(MoveTemp(DynamicMesh));
 	
 	// @todo(Material) This is just to have a material, we should transfer the materials from the assets if they have them. 
+	TObjectPtr<UDataflowBaseContent> DataflowContent = GetDataflowContent();
 	if (DataflowContent && DataflowContent->GetDataflowAsset() && DataflowContent->GetDataflowAsset()->Material)
 	{
 		DynamicMeshComponent->ConfigureMaterialSet({ DataflowContent->GetDataflowAsset()->Material });
@@ -344,13 +362,23 @@ void FDataflowConstructionScene::UpdateConstructionScene()
 	// Attach a wireframe renderer to the DynamicMeshComponents
 	UpdateWireframeMeshElementsVisualizer();
 
-	DataflowContent->SetIsDirty(false);
+	if (TObjectPtr<UDataflowBaseContent> DataflowContent = GetDataflowContent())
+	{
+		DataflowContent->SetIsDirty(false);
+	}
 }
 
-FDataflowSimulationScene::FDataflowSimulationScene(FPreviewScene::ConstructionValues ConstructionValues, TObjectPtr<UDataflowBaseContent> InEditorContent) 
-	: FDataflowPreviewScene(ConstructionValues,InEditorContent)
+//
+// FDataflowSimulationScene
+//
+
+FDataflowSimulationScene::FDataflowSimulationScene(FPreviewScene::ConstructionValues ConstructionValues, UDataflowEditor* InEditor)
+	: FDataflowPreviewScene(ConstructionValues, InEditor)
 {
-	DataflowContent->RegisterWorldContent(this, RootSceneActor);
+	if (TObjectPtr<UDataflowBaseContent> DataflowContent = GetDataflowContent())
+	{
+		DataflowContent->RegisterWorldContent(this, RootSceneActor);
+	}
 
 	TInlineComponentArray<UPrimitiveComponent*> PrimComponents;
 	RootSceneActor->GetComponents(PrimComponents);
@@ -371,8 +399,11 @@ FDataflowSimulationScene::~FDataflowSimulationScene()
 	{
 		PrimComponent->SelectionOverrideDelegate.Unbind();
 	}
-	
-	DataflowContent->UnregisterWorldContent(this);
+
+	if (TObjectPtr<UDataflowBaseContent> DataflowContent = GetDataflowContent())
+	{
+		DataflowContent->UnregisterWorldContent(this);
+	}
 }
 
 void FDataflowSimulationScene::TickDataflowScene(const float DeltaSeconds)
diff --git a/Engine/Plugins/Experimental/Dataflow/Source/DataflowEditor/Private/Dataflow/DataflowEditorToolkit.cpp b/Engine/Plugins/Experimental/Dataflow/Source/DataflowEditor/Private/Dataflow/DataflowEditorToolkit.cpp
--- a/Engine/Plugins/Experimental/Dataflow/Source/DataflowEditor/Private/Dataflow/DataflowEditorToolkit.cpp
+++ b/Engine/Plugins/Experimental/Dataflow/Source/DataflowEditor/Private/Dataflow/DataflowEditorToolkit.cpp
@@ -156,8 +156,8 @@ FDataflowEditorToolkit::FDataflowEditorToolkit(UAssetEditor* InOwningAssetEditor
 	PreviewSceneArgs.bShouldSimulatePhysics = 1;
 	PreviewSceneArgs.bCreatePhysicsScene = 1;
 	
-	ObjectScene = MakeUnique<FDataflowConstructionScene>(PreviewSceneArgs, GetDataflowContent());
-	SimulationScene = MakeShared<FDataflowSimulationScene>(PreviewSceneArgs, GetDataflowContent());
+	ObjectScene = MakeUnique<FDataflowConstructionScene>(PreviewSceneArgs, Cast<UDataflowEditor>(InOwningAssetEditor));
+	SimulationScene = MakeShared<FDataflowSimulationScene>(PreviewSceneArgs, Cast<UDataflowEditor>(InOwningAssetEditor));
 }
 
 FDataflowEditorToolkit::~FDataflowEditorToolkit()
diff --git a/Engine/Plugins/Experimental/Dataflow/Source/DataflowEditor/Public/Dataflow/DataflowEditor.h b/Engine/Plugins/Experimental/Dataflow/Source/DataflowEditor/Public/Dataflow/DataflowEditor.h
--- a/Engine/Plugins/Experimental/Dataflow/Source/DataflowEditor/Public/Dataflow/DataflowEditor.h
+++ b/Engine/Plugins/Experimental/Dataflow/Source/DataflowEditor/Public/Dataflow/DataflowEditor.h
@@ -47,6 +47,9 @@ public:
 	/** Initialize an editor with a given content and an owner */
 	void InitializeContent(TObjectPtr<UDataflowBaseContent> BaseContent, const TObjectPtr<UObject>& ContentOwner);
 
+	TObjectPtr<UDataflowBaseContent> GetDataflowContent() { return DataflowContent; }
+	const TObjectPtr<UDataflowBaseContent> GetDataflowContent() const { return DataflowContent; }
+
 private :
 
 	friend class FDataflowEditorToolkit;
diff --git a/Engine/Plugins/Experimental/Dataflow/Source/DataflowEditor/Public/Dataflow/DataflowEditorScenes.h b/Engine/Plugins/Experimental/Dataflow/Source/DataflowEditor/Public/Dataflow/DataflowEditorScenes.h
--- a/Engine/Plugins/Experimental/Dataflow/Source/DataflowEditor/Public/Dataflow/DataflowEditorScenes.h
+++ b/Engine/Plugins/Experimental/Dataflow/Source/DataflowEditor/Public/Dataflow/DataflowEditorScenes.h
@@ -16,6 +16,7 @@ class ADataflowActor;
 class UMeshElementsVisualizer;
 class UInteractiveToolPropertySet;
 class UDataflowEdNode;
+class UDataflowEditor;
 
 namespace UE
 {
@@ -34,15 +35,15 @@ class DATAFLOWEDITOR_API FDataflowPreviewScene : public FAdvancedPreviewScene
 {
 public:
 
-	FDataflowPreviewScene(FPreviewScene::ConstructionValues ConstructionValues, TObjectPtr<UDataflowBaseContent> InDataflowContent);
+	FDataflowPreviewScene(FPreviewScene::ConstructionValues ConstructionValues, UDataflowEditor* Editor);
 	virtual ~FDataflowPreviewScene();
 	
 	// FGCObject interface
 	virtual void AddReferencedObjects(FReferenceCollector& Collector) override;
 
 	/** Dataflow content accessors */
-	TObjectPtr<UDataflowBaseContent> GetDataflowContent() { return DataflowContent; }
-	const TObjectPtr<UDataflowBaseContent> GetDataflowContent() const { return DataflowContent; }
+	TObjectPtr<UDataflowBaseContent> GetDataflowContent();
+		const TObjectPtr<UDataflowBaseContent> GetDataflowContent() const;
 
 	/** Root scene actor accessors */
 	TObjectPtr<AActor> GetRootActor() { return RootSceneActor; }
@@ -68,9 +69,8 @@ protected:
 	
 	/** Root scene actor */
 	TObjectPtr<AActor> RootSceneActor = nullptr;
-	
-	/** Data flows assets datas used to build the scene components */
-	TObjectPtr<UDataflowBaseContent> DataflowContent = nullptr;
+
+	UDataflowEditor* DataflowEditor = nullptr;
 
 	/** Mode Manager for selection */
 	TSharedPtr<FAssetEditorModeManager> DataflowModeManager;
@@ -86,7 +86,7 @@ class DATAFLOWEDITOR_API FDataflowConstructionScene : public FDataflowPreviewSce
 
 public:
 
-	FDataflowConstructionScene(FPreviewScene::ConstructionValues ConstructionValues, TObjectPtr<UDataflowBaseContent> InEditorContent);
+	FDataflowConstructionScene(FPreviewScene::ConstructionValues ConstructionValues, UDataflowEditor* Editor);
 	virtual ~FDataflowConstructionScene();
 
 	// FGCObject interface
@@ -151,7 +151,7 @@ class DATAFLOWEDITOR_API FDataflowSimulationScene : public FDataflowPreviewScene
 {
 public:
 
-	FDataflowSimulationScene(FPreviewScene::ConstructionValues ConstructionValues, TObjectPtr<UDataflowBaseContent> InEditorContent);
+	FDataflowSimulationScene(FPreviewScene::ConstructionValues ConstructionValues, UDataflowEditor* Editor);
 	virtual ~FDataflowSimulationScene();
 
 	/** Tick data flow scene */
diff --git a/Engine/Source/Runtime/Experimental/Dataflow/Core/Public/Dataflow/DataflowNodeParameters.h b/Engine/Source/Runtime/Experimental/Dataflow/Core/Public/Dataflow/DataflowNodeParameters.h
--- a/Engine/Source/Runtime/Experimental/Dataflow/Core/Public/Dataflow/DataflowNodeParameters.h
+++ b/Engine/Source/Runtime/Experimental/Dataflow/Core/Public/Dataflow/DataflowNodeParameters.h
@@ -74,6 +74,7 @@ namespace Dataflow
 
 		const FProperty* GetProperty() const { return Property; }
 		const FTimestamp& GetTimestamp() const { return Timestamp; }
+		void SetTimestamp(const FTimestamp& InTimestamp) { Timestamp = InTimestamp; }
 
 		const FGuid& GetNodeGuid() const { return NodeGuid; }
 		const uint32 GetNodeHash() const { return NodeHash; }
