Dataflow : CVar to disable read of context cache.

diff --git a/Engine/Plugins/Experimental/Dataflow/Source/DataflowEditor/Private/Dataflow/AssetDefinition_DataflowContext.cpp b/Engine/Plugins/Experimental/Dataflow/Source/DataflowEditor/Private/Dataflow/AssetDefinition_DataflowContext.cpp
--- a/Engine/Plugins/Experimental/Dataflow/Source/DataflowEditor/Private/Dataflow/AssetDefinition_DataflowContext.cpp
+++ b/Engine/Plugins/Experimental/Dataflow/Source/DataflowEditor/Private/Dataflow/AssetDefinition_DataflowContext.cpp
@@ -18,6 +18,9 @@
 
 #define LOCTEXT_NAMESPACE "AssetActions_DataflowContext"
 
+bool bDataflowEnableContextCacheLoading = true;
+FAutoConsoleVariableRef CVARDataflowEnableContextCacheLoading(TEXT("p.Dataflow.Editor.ContextCaching"), bDataflowEnableContextCacheLoading, TEXT("Allow the Dataflow editor to load the cached state of the graph when and asset is re-opened.[def:true]"));
+
 namespace DataflowContextDefinitionHelpers
 {
 
@@ -170,13 +173,15 @@ namespace DataflowContextDefinitionHelpers
 	{
 		check(ContentOwner.Get());
 
+		bool bNeedsNewAsset = true;
+		TObjectPtr<UObject> Asset = nullptr;
 		UClass* DataflowClass = T::StaticClass();
 		UDataflow* DataflowAsset = Private::GetDataflowAssetFrom(ContentOwner);
 
 		FContentBrowserModule& ContentBrowserModule = FModuleManager::LoadModuleChecked<FContentBrowserModule>("ContentBrowser");
 
 		FString PackageName = FString::Printf(TEXT("/Game/_GENERATED/Dataflow/%s"), *ContentOwner.GetName());
-		if(DataflowAsset)
+		if (DataflowAsset)
 		{
 			PackageName = FString::Printf(TEXT("%s_%s"), *PackageName, *DataflowAsset->GetName());
 		}
@@ -187,18 +192,20 @@ namespace DataflowContextDefinitionHelpers
 			Package = CreatePackage(*PackageName);
 		}
 
-		bool bNeedsNewAsset = true;
-		TObjectPtr<UObject> Asset = StaticLoadObject(DataflowClass, Package, *PackageName);
-		if (Asset && DataflowAsset)
+		if (bDataflowEnableContextCacheLoading)
 		{
-			// Validate the loaded cache
-			bNeedsNewAsset =
-				!BindContextToGraph(Asset, DataflowAsset) ||
-				!ValidateCachedNodeHash(Asset, DataflowAsset) ||
-				!ResetCacheTimestamp(Asset, DataflowAsset);
+			Asset = StaticLoadObject(DataflowClass, Package, *PackageName);
+			if (Asset && DataflowAsset)
+			{
+				// Validate the loaded cache
+				bNeedsNewAsset =
+					!BindContextToGraph(Asset, DataflowAsset) ||
+					!ValidateCachedNodeHash(Asset, DataflowAsset) ||
+					!ResetCacheTimestamp(Asset, DataflowAsset);
+			}
 		}
-		
-		if(bNeedsNewAsset)
+
+		if (!Asset || bNeedsNewAsset)
 		{
 			const FName AssetName(FPackageName::GetLongPackageAssetName(PackageName));
 			Asset = NewObject<UObject>(Package, DataflowClass, AssetName, RF_Public | RF_Standalone | RF_Transactional);
