[ver2]Dataflow : Isolate the selected mesh for painting.
#perflight 660758cf5d72a085d567a921

diff --git a/Engine/Plugins/Experimental/Dataflow/Source/DataflowEditor/Private/Dataflow/DataflowConstructionViewportClient.cpp b/Engine/Plugins/Experimental/Dataflow/Source/DataflowEditor/Private/Dataflow/DataflowConstructionViewportClient.cpp
--- a/Engine/Plugins/Experimental/Dataflow/Source/DataflowEditor/Private/Dataflow/DataflowConstructionViewportClient.cpp
+++ b/Engine/Plugins/Experimental/Dataflow/Source/DataflowEditor/Private/Dataflow/DataflowConstructionViewportClient.cpp
@@ -58,6 +58,12 @@ void FDataflowConstructionViewportClient::Tick(float DeltaSeconds)
 	}
 }
 
+USelection* FDataflowConstructionViewportClient::GetSelectedComponents() 
+{ 
+	return ModeTools->GetSelectedComponents(); 
+}
+
+
 void FDataflowConstructionViewportClient::ProcessClick(FSceneView& View, HHitProxy* HitProxy, FKey Key, EInputEvent Event, uint32 HitX, uint32 HitY)
 {
 	Super::ProcessClick(View, HitProxy, Key, Event, HitX, HitY);
@@ -83,6 +89,8 @@ void FDataflowConstructionViewportClient::ProcessClick(FSceneView& View, HHitPro
 									{
 										// Start the corresponding tool
 										DataflowMode->StartToolForSelectedNode(SelectedNode);
+
+										return CollectionComponent;
 									}
 								}
 							}
@@ -91,6 +99,7 @@ void FDataflowConstructionViewportClient::ProcessClick(FSceneView& View, HHitPro
 				}
 			}
 		}
+		return (UDataflowEditorCollectionComponent*)nullptr;
 	};
 
 	auto UpdateSelectedComponentInViewport = [&](USelection* SelectedComponents)
@@ -147,6 +156,15 @@ void FDataflowConstructionViewportClient::ProcessClick(FSceneView& View, HHitPro
 		return false;
 	};
 
+	auto IsolateComponent = [&](UDataflowEditorCollectionComponent* SelectedComponent)
+	{
+		if( FDataflowConstructionScene* Scene = static_cast<FDataflowConstructionScene*>(PreviewScene) )
+		{ 
+			Scene->SetVisibility(false);
+			Scene->SetVisibility(true,SelectedComponent);
+		}
+	};
+
 	if (!IsInteractiveToolActive())
 	{
 		if (USelection* SelectedComponents = ModeTools->GetSelectedComponents())
diff --git a/Engine/Plugins/Experimental/Dataflow/Source/DataflowEditor/Private/Dataflow/DataflowEditorCollectionComponent.h b/Engine/Plugins/Experimental/Dataflow/Source/DataflowEditor/Private/Dataflow/DataflowEditorCollectionComponent.h
--- a/Engine/Plugins/Experimental/Dataflow/Source/DataflowEditor/Private/Dataflow/DataflowEditorCollectionComponent.h
+++ b/Engine/Plugins/Experimental/Dataflow/Source/DataflowEditor/Private/Dataflow/DataflowEditorCollectionComponent.h
@@ -8,7 +8,7 @@
 #include "DataflowEditorCollectionComponent.generated.h"
 
 class UDataflowEdNode;
-
+class UMeshWireframeComponent;
 /**
 *	FleshComponent
 */
@@ -24,4 +24,8 @@ public:
 	UPROPERTY()
 	TObjectPtr<const UDataflowEdNode> Node = nullptr;
 
+	UPROPERTY()
+	TObjectPtr<UMeshWireframeComponent> WireframeComponent = nullptr;
+
+
 };
diff --git a/Engine/Plugins/Experimental/Dataflow/Source/DataflowEditor/Private/Dataflow/DataflowEditorMode.cpp b/Engine/Plugins/Experimental/Dataflow/Source/DataflowEditor/Private/Dataflow/DataflowEditorMode.cpp
--- a/Engine/Plugins/Experimental/Dataflow/Source/DataflowEditor/Private/Dataflow/DataflowEditorMode.cpp
+++ b/Engine/Plugins/Experimental/Dataflow/Source/DataflowEditor/Private/Dataflow/DataflowEditorMode.cpp
@@ -75,6 +75,11 @@ void UDataflowEditorMode::Enter()
 	UE::TransformGizmoUtil::RegisterTransformGizmoContextObject(GetInteractiveToolsContext());
 }
 
+void UDataflowEditorMode::SetDataflowEditor(UDataflowEditor* InDataflowEditor) 
+{ 
+	DataflowEditor = InDataflowEditor; 
+}
+
 void UDataflowEditorMode::AddToolTargetFactories()
 {
 	GetInteractiveToolsContext()->TargetManager->AddTargetFactory(NewObject<UDynamicMeshComponentToolTargetFactory>(GetToolManager()));
@@ -195,15 +200,10 @@ void UDataflowEditorMode::RegisterTools()
 
 	UEditorInteractiveToolsContext* const ConstructionViewportToolsContext = GetInteractiveToolsContext();
 
-	UDataflowEditorWeightMapPaintToolBuilder* WeightMapPaintToolBuilder = NewObject<UDataflowEditorWeightMapPaintToolBuilder>();
+	UDataflowEditorWeightMapPaintToolBuilder* WeightMapPaintToolBuilder = NewObject<UDataflowEditorWeightMapPaintToolBuilder>(this);
+	WeightMapPaintToolBuilder->SetEditorMode(this);
 	RegisterDataflowTool(CommandInfos.BeginWeightMapPaintTool, FDataflowEditorCommandsImpl::BeginWeightMapPaintToolIdentifier, WeightMapPaintToolBuilder, WeightMapPaintToolBuilder, ConstructionViewportToolsContext);
 	RegisterAddNodeCommand(CommandInfos.AddWeightMapNode, FDataflowCollectionAddScalarVertexPropertyNode::StaticType(), CommandInfos.BeginWeightMapPaintTool);
-
-	// @todo(brice) Remove Example Tools
-	//RegisterTool(CommandInfos.BeginAttributeEditorTool, FDataflowEditorCommandsImpl::BeginAttributeEditorToolIdentifier, NewObject<UAttributeEditorToolBuilder>());
-	//RegisterTool(CommandInfos.BeginMeshSelectionTool, FDataflowEditorCommandsImpl::BeginMeshSelectionToolIdentifier, NewObject<UMeshSelectionToolBuilder>());
-	//RegisterTool(CommandInfos.BeginMeshSelectionTool, FDataflowEditorCommandsImpl::BeginMeshSelectionToolIdentifier, NewObject<UMeshVertexPaintToolBuilder>());
-	//RegisterTool(CommandInfos.BeginMeshSelectionTool, FDataflowEditorCommandsImpl::BeginMeshSelectionToolIdentifier, NewObject<UMeshAttributePaintToolBuilder>());
 }
 
 bool UDataflowEditorMode::ShouldToolStartBeAllowed(const FString& ToolIdentifier) const
diff --git a/Engine/Plugins/Experimental/Dataflow/Source/DataflowEditor/Private/Dataflow/DataflowEditorScenes.cpp b/Engine/Plugins/Experimental/Dataflow/Source/DataflowEditor/Private/Dataflow/DataflowEditorScenes.cpp
--- a/Engine/Plugins/Experimental/Dataflow/Source/DataflowEditor/Private/Dataflow/DataflowEditorScenes.cpp
+++ b/Engine/Plugins/Experimental/Dataflow/Source/DataflowEditor/Private/Dataflow/DataflowEditorScenes.cpp
@@ -113,6 +113,37 @@ FDataflowConstructionScene::~FDataflowConstructionScene()
 	ResetDynamicMeshComponents();
 }
 
+/** Hide all or a single component */
+void FDataflowConstructionScene::SetVisibility(bool bVisibility, UActorComponent* InComponent)
+{
+	auto SetCollectionVisiblity = [](bool bVisibility,TObjectPtr<UDataflowEditorCollectionComponent> Component) {
+		Component->SetVisibility(bVisibility);
+		if (Component->WireframeComponent)
+		{
+			Component->WireframeComponent->SetVisibility(bVisibility);
+		}
+	};
+
+	for (FRenderElement& RenderElement : DynamicMeshComponents)
+	{
+		if (TObjectPtr<UDataflowEditorCollectionComponent> DynamicMeshComponent = Cast<UDataflowEditorCollectionComponent>(RenderElement.Value))
+		{
+			if (InComponent != nullptr)
+			{
+				if (InComponent == DynamicMeshComponent.Get())
+				{
+					SetCollectionVisiblity(bVisibility,DynamicMeshComponent);
+				}
+			}
+			else
+			{
+				SetCollectionVisiblity(bVisibility,DynamicMeshComponent);
+			}
+		}
+	}
+}
+
+
 void FDataflowConstructionScene::AddReferencedObjects(FReferenceCollector& Collector)
 {
 	FDataflowPreviewScene::AddReferencedObjects(Collector);
@@ -277,37 +308,39 @@ void FDataflowConstructionScene::AddWireframeMeshElementsVisualizer()
 	ensure(WireframeElements.Num()==0);
 	for(FRenderElement Elem : DynamicMeshComponents)
 	{
-		TObjectPtr<UDynamicMeshComponent> DynamicMeshComponent = Elem.Value;
-
-		// Set up the wireframe display of the rest space mesh.
+		if( TObjectPtr<UDataflowEditorCollectionComponent> DynamicMeshComponent = Cast<UDataflowEditorCollectionComponent>(Elem.Value) )
+		{
+			// Set up the wireframe display of the rest space mesh.
 
-		TObjectPtr<UMeshElementsVisualizer> WireframeDraw = NewObject<UMeshElementsVisualizer>(RootSceneActor);
-		WireframeElements.Add(DynamicMeshComponent, WireframeDraw);
+			TObjectPtr<UMeshElementsVisualizer> WireframeDraw = NewObject<UMeshElementsVisualizer>(RootSceneActor);
+			WireframeElements.Add(DynamicMeshComponent, WireframeDraw);
 
-		WireframeDraw->CreateInWorld(GetWorld(), FTransform::Identity);
-		WireframeDraw->Settings->DepthBias = 2.0;
-		WireframeDraw->Settings->bAdjustDepthBiasUsingMeshSize = false;
-		WireframeDraw->Settings->bShowWireframe = true;
-		WireframeDraw->Settings->bShowBorders = true;
-		WireframeDraw->Settings->bShowUVSeams = false;
-		WireframeDraw->WireframeComponent->BoundaryEdgeThickness = 2;
+			WireframeDraw->CreateInWorld(GetWorld(), FTransform::Identity);
+			WireframeDraw->Settings->DepthBias = 2.0;
+			WireframeDraw->Settings->bAdjustDepthBiasUsingMeshSize = false;
+			WireframeDraw->Settings->bShowWireframe = true;
+			WireframeDraw->Settings->bShowBorders = true;
+			WireframeDraw->Settings->bShowUVSeams = false;
+			WireframeDraw->WireframeComponent->BoundaryEdgeThickness = 2;
+			DynamicMeshComponent->WireframeComponent = WireframeDraw->WireframeComponent;
 
-		WireframeDraw->SetMeshAccessFunction([DynamicMeshComponent](UMeshElementsVisualizer::ProcessDynamicMeshFunc ProcessFunc)
-		{
-			ProcessFunc(*DynamicMeshComponent->GetMesh());
-		});
+			WireframeDraw->SetMeshAccessFunction([DynamicMeshComponent](UMeshElementsVisualizer::ProcessDynamicMeshFunc ProcessFunc)
+				{
+					ProcessFunc(*DynamicMeshComponent->GetMesh());
+				});
 
-		for (FRenderElement RenderElement : DynamicMeshComponents)
-		{
-			RenderElement.Value->OnMeshChanged.Add(FSimpleMulticastDelegate::FDelegate::CreateLambda([WireframeDraw,this]()
+			for (FRenderElement RenderElement : DynamicMeshComponents)
 			{
-				WireframeDraw->NotifyMeshChanged();
-			}));
+				RenderElement.Value->OnMeshChanged.Add(FSimpleMulticastDelegate::FDelegate::CreateLambda([WireframeDraw, this]()
+					{
+						WireframeDraw->NotifyMeshChanged();
+					}));
 
-			const bool bRestSpaceMeshVisible = RenderElement.Value->GetVisibleFlag();
-			WireframeDraw->Settings->bVisible = bRestSpaceMeshVisible && bConstructionViewWireframe;
+				const bool bRestSpaceMeshVisible = RenderElement.Value->GetVisibleFlag();
+				WireframeDraw->Settings->bVisible = bRestSpaceMeshVisible && bConstructionViewWireframe;
+			}
+			PropertyObjectsToTick.Add(WireframeDraw->Settings);
 		}
-		PropertyObjectsToTick.Add(WireframeDraw->Settings);
 	}
 }
 
diff --git a/Engine/Plugins/Experimental/Dataflow/Source/DataflowEditor/Private/Dataflow/DataflowEditorToolkit.cpp b/Engine/Plugins/Experimental/Dataflow/Source/DataflowEditor/Private/Dataflow/DataflowEditorToolkit.cpp
--- a/Engine/Plugins/Experimental/Dataflow/Source/DataflowEditor/Private/Dataflow/DataflowEditorToolkit.cpp
+++ b/Engine/Plugins/Experimental/Dataflow/Source/DataflowEditor/Private/Dataflow/DataflowEditorToolkit.cpp
@@ -66,8 +66,9 @@ const FName FDataflowEditorToolkit::SimulationViewportTabId(TEXT("DataflowEditor
 
 FDataflowEditorToolkit::FDataflowEditorToolkit(UAssetEditor* InOwningAssetEditor)
 	: FBaseCharacterFXEditorToolkit(InOwningAssetEditor, FName("DataflowEditor"))
+	, DataflowEditor(Cast< UDataflowEditor>(InOwningAssetEditor))
 {
-	check(Cast<UDataflowEditor>(InOwningAssetEditor));
+	check(DataflowEditor);
 
 	StandaloneDefaultLayout = FTabManager::NewLayout(FName("DataflowEditorLayout03"))
 		->AddArea
@@ -156,8 +157,8 @@ FDataflowEditorToolkit::FDataflowEditorToolkit(UAssetEditor* InOwningAssetEditor
 	PreviewSceneArgs.bShouldSimulatePhysics = 1;
 	PreviewSceneArgs.bCreatePhysicsScene = 1;
 	
-	ObjectScene = MakeUnique<FDataflowConstructionScene>(PreviewSceneArgs, Cast<UDataflowEditor>(InOwningAssetEditor));
-	SimulationScene = MakeShared<FDataflowSimulationScene>(PreviewSceneArgs, Cast<UDataflowEditor>(InOwningAssetEditor));
+	ObjectScene = MakeUnique<FDataflowConstructionScene>(PreviewSceneArgs, DataflowEditor);
+	SimulationScene = MakeShared<FDataflowSimulationScene>(PreviewSceneArgs, DataflowEditor);
 }
 
 FDataflowEditorToolkit::~FDataflowEditorToolkit()
@@ -265,12 +266,12 @@ FEditorModeID FDataflowEditorToolkit::GetEditorModeId() const
 
 TObjectPtr<UDataflowBaseContent> FDataflowEditorToolkit::GetDataflowContent()
 {
-	return Cast<UDataflowEditor>(OwningAssetEditor)->DataflowContent;
+	return DataflowEditor->DataflowContent;
 }
 
 const TObjectPtr<UDataflowBaseContent> FDataflowEditorToolkit::GetDataflowContent() const
 {
-	return Cast<UDataflowEditor>(OwningAssetEditor)->DataflowContent;
+	return DataflowEditor->DataflowContent;
 }
 
 bool FDataflowEditorToolkit::OnRequestClose(EAssetEditorCloseReason InCloseReason)
@@ -298,12 +299,8 @@ bool FDataflowEditorToolkit::OnRequestClose(EAssetEditorCloseReason InCloseReaso
 
 void FDataflowEditorToolkit::PostInitAssetEditor()
 {
-	// @todo(DataflowMode) : Enable base class functionality when the mode is ready.
 	FBaseCharacterFXEditorToolkit::PostInitAssetEditor();
 	
-	// @todo(DataflowMode) : Access to mode in Toolkit construction
-	//UDataflowEditorMode* const DataflowEditorMode = CastChecked<UDataflowEditorMode>(EditorModeManager->GetActiveScriptableMode(UDataflowEditorMode::EM_DataflowEditorModeId));
-	
 	auto SetCommonViewportClientOptions = [](FEditorViewportClient* Client)
 	{
 		// Normally the bIsRealtime flag is determined by whether the connection is remote, but our
@@ -345,9 +342,7 @@ void FDataflowEditorToolkit::PostInitAssetEditor()
 	SetCommonViewportClientOptions(ViewportClient.Get());
 	SetCommonViewportClientOptions(SimulationViewportClient.Get());
 
-	//@todo(Dataflow) : FChaosClothAssetEditorToolkit::PostInitAssetEditor() has alot more going on. 
 
-	//@todo(brice) : can we remove this coupling of the viewport and mode
 	UDataflowEditorMode* const DataflowMode = CastChecked<UDataflowEditorMode>(EditorModeManager->GetActiveScriptableMode(UDataflowEditorMode::EM_DataflowEditorModeId));
 	const TWeakPtr<FViewportClient> WeakConstructionViewportClient(ViewportClient);
 	DataflowMode->SetConstructionViewportClient(StaticCastWeakPtr<FDataflowConstructionViewportClient>(WeakConstructionViewportClient));
@@ -359,6 +354,7 @@ void FDataflowEditorToolkit::InitializeEdMode(UBaseCharacterFXEditorMode* EdMode
 {
 	UDataflowEditorMode* DataflowMode = Cast<UDataflowEditorMode>(EdMode);
 	check(DataflowMode);
+	DataflowMode->SetDataflowEditor(DataflowEditor);
 
 	// We first set the preview scene in order to store the dynamic mesh elements
 	// generated by the tools
@@ -366,7 +362,6 @@ void FDataflowEditorToolkit::InitializeEdMode(UBaseCharacterFXEditorMode* EdMode
 
 	// Set of the graph editor to be able to add nodes
 	DataflowMode->SetDataflowGraphEditor(GraphEditor);
-	
 	TArray<TObjectPtr<UObject>> ObjectsToEdit;
 	OwningAssetEditor->GetObjectsToEdit(MutableView(ObjectsToEdit));
 	DataflowMode->InitializeTargets(ObjectsToEdit);
diff --git a/Engine/Plugins/Experimental/Dataflow/Source/DataflowEditor/Private/DataflowEditorTools/DataflowEditorWeightMapPaintTool.cpp b/Engine/Plugins/Experimental/Dataflow/Source/DataflowEditor/Private/DataflowEditorTools/DataflowEditorWeightMapPaintTool.cpp
--- a/Engine/Plugins/Experimental/Dataflow/Source/DataflowEditor/Private/DataflowEditorTools/DataflowEditorWeightMapPaintTool.cpp
+++ b/Engine/Plugins/Experimental/Dataflow/Source/DataflowEditor/Private/DataflowEditorTools/DataflowEditorWeightMapPaintTool.cpp
@@ -11,11 +11,14 @@
 #include "Changes/BasicChanges.h"
 #include "ContextObjectStore.h"
 #include "DataflowEditorTools/DataflowEditorWeightMapPaintBrushOps.h"
+#include "Dataflow/DataflowCollectionAddScalarVertexPropertyNode.h"
+#include "Dataflow/DataflowConstructionViewportClient.h"
 #include "Dataflow/DataflowContent.h"
 #include "Dataflow/DataflowEdNode.h"
 #include "Dataflow/DataflowEditorCollectionComponent.h"
+#include "Dataflow/DataflowEditorMode.h"
+#include "Dataflow/DataflowEditorScenes.h"
 #include "Dataflow/DataflowObject.h"
-#include "Dataflow/DataflowCollectionAddScalarVertexPropertyNode.h"
 #include "Dataflow/DataflowSNode.h"
 #include "Dataflow/DataflowGraphEditor.h"
 #include "Drawing/MeshElementsVisualizer.h"
@@ -35,6 +38,7 @@
 #include "Polygon2.h"
 #include "Sculpting/StampFalloffs.h"
 #include "Sculpting/MeshSculptUtil.h"
+#include "Selection.h"
 #include "Selections/MeshConnectedComponents.h"
 #include "Selections/MeshFaceSelection.h"
 #include "Selections/MeshVertexSelection.h"
@@ -63,7 +67,6 @@ namespace Dataflow::Private
 /*
  * ToolBuilder
  */
-
 void UDataflowEditorWeightMapPaintToolBuilder::GetSupportedViewModes(TArray<Dataflow::EDataflowPatternVertexType>& Modes) const
 {
 	Modes.Add(Dataflow::EDataflowPatternVertexType::Sim3D);
@@ -122,6 +125,7 @@ bool UDataflowEditorWeightMapPaintToolBuilder::CanBuildTool(const FToolBuilderSt
 UMeshSurfacePointTool* UDataflowEditorWeightMapPaintToolBuilder::CreateNewTool(const FToolBuilderState& SceneState) const
 {
 	UDataflowEditorWeightMapPaintTool* PaintTool = NewObject<UDataflowEditorWeightMapPaintTool>(SceneState.ToolManager);
+	PaintTool->SetEditorMode(Mode);
 	PaintTool->SetWorld(SceneState.World);
 
 	if (UDataflowContextObject* ContextObject = SceneState.ToolManager->GetContextObjectStore()->FindContext<UDataflowContextObject>())
@@ -166,6 +170,26 @@ void UDataflowEditorWeightMapPaintTool::Setup()
 {
 	UMeshSculptToolBase::Setup();
 
+	auto IsolateComponent = [&](UDataflowEditorCollectionComponent* SelectedComponent)
+	{
+		if (FDataflowConstructionScene* Scene = Mode->GetDataflowConstructionScene())
+		{
+			Scene->SetVisibility(false);
+			Scene->SetVisibility(true, SelectedComponent);
+		}
+	};
+	if (TSharedPtr<FDataflowConstructionViewportClient> VC = Mode->GetConstructionViewportClient().Pin())
+	{
+		if (USelection* SelectedComponents = VC->GetSelectedComponents())
+		{
+			if (UDataflowEditorCollectionComponent* DataflowComponent = SelectedComponents->GetBottom<UDataflowEditorCollectionComponent>())
+			{
+				IsolateComponent(DataflowComponent);
+			}
+		}
+	}
+
+
 	// Get the selected weight map node
 	WeightMapNodeToUpdate = DataflowEditorContextObject->GetPrimarySelectedNodeOfType<FDataflowCollectionAddScalarVertexPropertyNode>();
 	checkf(WeightMapNodeToUpdate, TEXT("No Weight Map Node is currently selected, or more than one node is selected"));
diff --git a/Engine/Plugins/Experimental/Dataflow/Source/DataflowEditor/Public/Dataflow/DataflowConstructionViewportClient.h b/Engine/Plugins/Experimental/Dataflow/Source/DataflowEditor/Public/Dataflow/DataflowConstructionViewportClient.h
--- a/Engine/Plugins/Experimental/Dataflow/Source/DataflowEditor/Public/Dataflow/DataflowConstructionViewportClient.h
+++ b/Engine/Plugins/Experimental/Dataflow/Source/DataflowEditor/Public/Dataflow/DataflowConstructionViewportClient.h
@@ -15,6 +15,7 @@ class UCombinedTransformGizmo;
 class FTransformGizmoDataBinder;
 class FDataflowPreviewScene;
 class UInputBehaviorSet;
+class USelection;
 
 class DATAFLOWEDITOR_API FDataflowConstructionViewportClient : public FEditorViewportClient //, public IInputBehaviorSource
 {
@@ -29,6 +30,7 @@ public:
 
 	// IInputBehaviorSource
 	// virtual const UInputBehaviorSet* GetInputBehaviors() const override;
+	USelection* GetSelectedComponents();
 
 	/** Set the data flow toolkit used to create the client*/
 	void SetDataflowEditorToolkit(TWeakPtr<FDataflowEditorToolkit> DataflowToolkit);
diff --git a/Engine/Plugins/Experimental/Dataflow/Source/DataflowEditor/Public/Dataflow/DataflowEditorMode.h b/Engine/Plugins/Experimental/Dataflow/Source/DataflowEditor/Public/Dataflow/DataflowEditorMode.h
--- a/Engine/Plugins/Experimental/Dataflow/Source/DataflowEditor/Public/Dataflow/DataflowEditorMode.h
+++ b/Engine/Plugins/Experimental/Dataflow/Source/DataflowEditor/Public/Dataflow/DataflowEditorMode.h
@@ -16,7 +16,7 @@ class IDataflowEditorToolBuilder;
 struct FManagedArrayCollection;
 class UEdGraphNode;
 class FViewport;
-
+class UDataflowEditor;
 
 /**
  * The dataflow editor mode is the mode used in the cloth asset editor. It holds most of the inter-tool state.
@@ -34,6 +34,9 @@ public:
 
 	UDataflowEditorMode();
 
+	void SetDataflowEditor(UDataflowEditor* InDataflowEditor);
+	TWeakPtr<FDataflowConstructionViewportClient, ESPMode::ThreadSafe> GetConstructionViewportClient() { return ConstructionViewportClient; }
+
 	/**
 	* Gets the scene bounding box
 	*/
@@ -197,6 +200,7 @@ private:
 	// Dataflow node type whose corresponding tool should be started on the next Tick
 	FName NodeTypeForPendingToolStart;
 
+	UDataflowEditor* DataflowEditor = nullptr;
 
 };
 
diff --git a/Engine/Plugins/Experimental/Dataflow/Source/DataflowEditor/Public/Dataflow/DataflowEditorScenes.h b/Engine/Plugins/Experimental/Dataflow/Source/DataflowEditor/Public/Dataflow/DataflowEditorScenes.h
--- a/Engine/Plugins/Experimental/Dataflow/Source/DataflowEditor/Public/Dataflow/DataflowEditorScenes.h
+++ b/Engine/Plugins/Experimental/Dataflow/Source/DataflowEditor/Public/Dataflow/DataflowEditorScenes.h
@@ -107,6 +107,9 @@ public:
 	/** Update the construction scene */
 	void UpdateConstructionScene();
 	
+	/** Hide all or a single component */
+	void SetVisibility(bool bVisible, UActorComponent* Component = nullptr);
+
 private:
 
 	/** Reset the scene mesh visualizer */
diff --git a/Engine/Plugins/Experimental/Dataflow/Source/DataflowEditor/Public/Dataflow/DataflowEditorToolkit.h b/Engine/Plugins/Experimental/Dataflow/Source/DataflowEditor/Public/Dataflow/DataflowEditorToolkit.h
--- a/Engine/Plugins/Experimental/Dataflow/Source/DataflowEditor/Public/Dataflow/DataflowEditorToolkit.h
+++ b/Engine/Plugins/Experimental/Dataflow/Source/DataflowEditor/Public/Dataflow/DataflowEditorToolkit.h
@@ -25,6 +25,7 @@ class FDataflowConstructionScene;
 class FDataflowSimulationViewportClient;
 class UDataflowBaseContent;
 class FDataflowSimulationScene;
+class UDataflowEditor;
 
 namespace Dataflow
 {
@@ -73,6 +74,8 @@ public:
 
 protected:
 
+	UDataflowEditor* DataflowEditor = nullptr;
+
 	// List of dataflow actions callbacks
 	void OnPropertyValueChanged(const FPropertyChangedEvent& PropertyChangedEvent);
 	bool OnNodeVerifyTitleCommit(const FText& NewText, UEdGraphNode* GraphNode, FText& OutErrorMessage);
diff --git a/Engine/Plugins/Experimental/Dataflow/Source/DataflowEditor/Public/DataflowEditorTools/DataflowEditorWeightMapPaintTool.h b/Engine/Plugins/Experimental/Dataflow/Source/DataflowEditor/Public/DataflowEditorTools/DataflowEditorWeightMapPaintTool.h
--- a/Engine/Plugins/Experimental/Dataflow/Source/DataflowEditor/Public/DataflowEditorTools/DataflowEditorWeightMapPaintTool.h
+++ b/Engine/Plugins/Experimental/Dataflow/Source/DataflowEditor/Public/DataflowEditorTools/DataflowEditorWeightMapPaintTool.h
@@ -33,6 +33,7 @@ class UDataflowWeightMapSmoothBrushOpProps;
 class UPolygonSelectionMechanic;
 class UDataflowContextObject;
 struct FDataflowCollectionAddScalarVertexPropertyNode;
+class UDataflowEditorMode;
 
 DECLARE_STATS_GROUP(TEXT("WeightMapPaintTool"), STATGROUP_WeightMapPaintTool, STATCAT_Advanced);
 DECLARE_CYCLE_STAT(TEXT("WeightMapPaintTool_UpdateROI"), WeightMapPaintTool_UpdateROI, STATGROUP_WeightMapPaintTool);
@@ -58,12 +59,15 @@ UCLASS()
 class DATAFLOWEDITOR_API UDataflowEditorWeightMapPaintToolBuilder : public UMeshSurfacePointMeshEditingToolBuilder, public IDataflowEditorToolBuilder
 {
 	GENERATED_BODY()
+public:
+	void SetEditorMode(UDataflowEditorMode* InMode) { Mode = InMode; }
 
 private:
 	virtual void GetSupportedViewModes(TArray<Dataflow::EDataflowPatternVertexType>& Modes) const override;
 	virtual bool CanBuildTool(const FToolBuilderState& SceneState) const override;
 	virtual UMeshSurfacePointTool* CreateNewTool(const FToolBuilderState& SceneState) const override;
 	virtual bool CanSetConstructionViewWireframeActive() const { return false; }
+	UDataflowEditorMode* Mode = nullptr;
 };
 
 
@@ -284,6 +288,7 @@ public:
 	virtual void CommitResult(UBaseDynamicMeshComponent* Component, bool bModifiedTopology) override;
 
 	void SetDataflowEditorContextObject(TObjectPtr<UDataflowContextObject> InDataflowEditorContextObject);
+	void SetEditorMode(UDataflowEditorMode* InMode) { Mode = InMode; }
 
 public:
 
@@ -486,6 +491,8 @@ protected:
 	void UpdateSelectedNode();
 
 	void UpdateVertexColorOverlay(const TSet<int>* TrianglesToUpdate = nullptr);
+
+	UDataflowEditorMode* Mode = nullptr;
 };
 
 
