Dataflow : Wireframe CVAR

diff --git a/Engine/Plugins/Experimental/Dataflow/Source/DataflowEditor/Private/Dataflow/DataflowEditorScenes.cpp b/Engine/Plugins/Experimental/Dataflow/Source/DataflowEditor/Private/Dataflow/DataflowEditorScenes.cpp
--- a/Engine/Plugins/Experimental/Dataflow/Source/DataflowEditor/Private/Dataflow/DataflowEditorScenes.cpp
+++ b/Engine/Plugins/Experimental/Dataflow/Source/DataflowEditor/Private/Dataflow/DataflowEditorScenes.cpp
@@ -22,6 +22,9 @@
 bool bDataflowShowFloorDefault = false;
 FAutoConsoleVariableRef CVARDataflowShowFloorDefault(TEXT("p.Dataflow.Editor.ShowFloor"), bDataflowShowFloorDefault, TEXT("Show the floor in the dataflow editor[def:false]"));
 
+bool bDataflowShowWireframeInConstructionView = true;
+FAutoConsoleVariableRef CVARDataflowShowWireframeInConstructionView(TEXT("p.Dataflow.Editor.Construction.ShowWireframe"), bDataflowShowWireframeInConstructionView, TEXT("Show the wireframe model in the dataflows construction view[def:true]"));
+
 
 FDataflowPreviewScene::FDataflowPreviewScene(FPreviewScene::ConstructionValues ConstructionValues, UDataflowEditor* InEditor)
 	: FAdvancedPreviewScene(ConstructionValues)
@@ -305,6 +308,8 @@ TObjectPtr<UDynamicMeshComponent>& FDataflowConstructionScene::AddDynamicMeshCom
 
 void FDataflowConstructionScene::AddWireframeMeshElementsVisualizer()
 {
+	if(!bDataflowShowWireframeInConstructionView) return;
+
 	ensure(WireframeElements.Num()==0);
 	for(FRenderElement Elem : DynamicMeshComponents)
 	{
