Dataflow : Default the wireframe on smooth to off.

diff --git a/Engine/Plugins/Experimental/Dataflow/Source/DataflowEditor/Private/Dataflow/DataflowEditorScenes.cpp b/Engine/Plugins/Experimental/Dataflow/Source/DataflowEditor/Private/Dataflow/DataflowEditorScenes.cpp
--- a/Engine/Plugins/Experimental/Dataflow/Source/DataflowEditor/Private/Dataflow/DataflowEditorScenes.cpp
+++ b/Engine/Plugins/Experimental/Dataflow/Source/DataflowEditor/Private/Dataflow/DataflowEditorScenes.cpp
@@ -22,7 +22,7 @@
 bool bDataflowShowFloorDefault = false;
 FAutoConsoleVariableRef CVARDataflowShowFloorDefault(TEXT("p.Dataflow.Editor.ShowFloor"), bDataflowShowFloorDefault, TEXT("Show the floor in the dataflow editor[def:false]"));
 
-bool bDataflowShowWireframeInConstructionView = true;
+bool bDataflowShowWireframeInConstructionView = false;
 FAutoConsoleVariableRef CVARDataflowShowWireframeInConstructionView(TEXT("p.Dataflow.Editor.Construction.ShowWireframe"), bDataflowShowWireframeInConstructionView, TEXT("Show the wireframe model in the dataflows construction view[def:true]"));
 
 
