Dataflow : Painting weights bounds check

diff --git a/Engine/Plugins/Experimental/Dataflow/Source/DataflowEditor/Private/DataflowEditorTools/DataflowEditorWeightMapPaintTool.cpp b/Engine/Plugins/Experimental/Dataflow/Source/DataflowEditor/Private/DataflowEditorTools/DataflowEditorWeightMapPaintTool.cpp
--- a/Engine/Plugins/Experimental/Dataflow/Source/DataflowEditor/Private/DataflowEditorTools/DataflowEditorWeightMapPaintTool.cpp
+++ b/Engine/Plugins/Experimental/Dataflow/Source/DataflowEditor/Private/DataflowEditorTools/DataflowEditorWeightMapPaintTool.cpp
@@ -381,7 +381,16 @@ void UDataflowEditorWeightMapPaintTool::Setup()
 				for (int32 DynamicMeshVert = 0; DynamicMeshVert < Mesh->VertexCount(); ++DynamicMeshVert)
 				{
 					DynamicMeshToWeight[DynamicMeshVert] = NonManifoldMapping.GetOriginalNonManifoldVertexID(DynamicMeshVert);
-					WeightToDynamicMesh[DynamicMeshToWeight[DynamicMeshVert]].Add(DynamicMeshVert);
+					if (0 <= DynamicMeshToWeight[DynamicMeshVert] && DynamicMeshToWeight[DynamicMeshVert] < WeightToDynamicMesh.Num())
+					{
+						WeightToDynamicMesh[DynamicMeshToWeight[DynamicMeshVert]].Add(DynamicMeshVert);
+					}
+					else
+					{
+						bHaveDynamicMeshToWeightConversion = false;
+						UE_LOG(LogTemp, Warning, TEXT("Weight map misalignment."));
+						break;
+					}
 				}
 			}
 		}
