Dataflow : Split out scene files.

diff --git a/Engine/Plugins/Experimental/Dataflow/Source/DataflowEditor/Private/Dataflow/DataflowConstructionScene.cpp b/Engine/Plugins/Experimental/Dataflow/Source/DataflowEditor/Private/Dataflow/DataflowConstructionScene.cpp
new file mode 100644
--- /dev/null
+++ b/Engine/Plugins/Experimental/Dataflow/Source/DataflowEditor/Private/Dataflow/DataflowConstructionScene.cpp
@@ -0,0 +1,324 @@
+// Copyright Epic Games, Inc. All Rights Reserved.
+
+#include "Dataflow/DataflowConstructionScene.h"
+
+#include "AssetEditorModeManager.h"
+#include "Dataflow/CollectionRenderingPatternUtility.h"
+#include "Dataflow/DataflowEditorCollectionComponent.h"
+#include "Dataflow/DataflowEditor.h"
+#include "Dataflow/DataflowObject.h"
+#include "Dataflow/DataflowEditorStyle.h"
+#include "Drawing/MeshElementsVisualizer.h"
+#include "Elements/Framework/EngineElementsLibrary.h"
+#include "Selection.h"
+
+#define LOCTEXT_NAMESPACE "FDataflowConstructionScene"
+
+bool bDataflowShowWireframeInConstructionView = false;
+FAutoConsoleVariableRef CVARDataflowShowWireframeInConstructionView(TEXT("p.Dataflow.Editor.Construction.ShowWireframe"), bDataflowShowWireframeInConstructionView, TEXT("Show the wireframe model in the dataflows construction view[def:true]"));
+
+//
+// Construction Scene
+//
+
+FDataflowConstructionScene::FDataflowConstructionScene(FPreviewScene::ConstructionValues ConstructionValues, UDataflowEditor* InEditor)
+	: FDataflowPreviewSceneBase(ConstructionValues, InEditor)
+{}
+
+FDataflowConstructionScene::~FDataflowConstructionScene()
+{
+	ResetDynamicMeshComponents();
+}
+
+/** Hide all or a single component */
+void FDataflowConstructionScene::SetVisibility(bool bVisibility, UActorComponent* InComponent)
+{
+	auto SetCollectionVisiblity = [](bool bVisibility,TObjectPtr<UDataflowEditorCollectionComponent> Component) {
+		Component->SetVisibility(bVisibility);
+		if (Component->WireframeComponent)
+		{
+			Component->WireframeComponent->SetVisibility(bVisibility);
+		}
+	};
+
+	for (FRenderElement& RenderElement : DynamicMeshComponents)
+	{
+		if (TObjectPtr<UDataflowEditorCollectionComponent> DynamicMeshComponent = Cast<UDataflowEditorCollectionComponent>(RenderElement.Value))
+		{
+			if (InComponent != nullptr)
+			{
+				if (InComponent == DynamicMeshComponent.Get())
+				{
+					SetCollectionVisiblity(bVisibility,DynamicMeshComponent);
+				}
+			}
+			else
+			{
+				SetCollectionVisiblity(bVisibility,DynamicMeshComponent);
+			}
+		}
+	}
+}
+
+
+void FDataflowConstructionScene::AddReferencedObjects(FReferenceCollector& Collector)
+{
+	FDataflowPreviewSceneBase::AddReferencedObjects(Collector);
+
+	Collector.AddReferencedObjects(DynamicMeshComponents);
+	Collector.AddReferencedObjects(WireframeElements);
+}
+
+FORCEINLINE Dataflow::FTimestamp LatestTimestamp(const UDataflow* Dataflow, const ::Dataflow::FContext* Context)
+{
+	if (Dataflow && Context)
+	{
+		return FMath::Max(Dataflow->GetRenderingTimestamp().Value, Context->GetTimestamp().Value);
+	}
+	return ::Dataflow::FTimestamp::Invalid;
+}
+
+void FDataflowConstructionScene::TickDataflowScene(const float DeltaSeconds)
+{
+	if (TObjectPtr<UDataflowBaseContent> DataflowContent = GetDataflowContent())
+	{
+		if (const TSharedPtr<Dataflow::FContext> DataflowContext = DataflowContent->GetDataflowContext())
+		{
+			if (const UDataflow* Dataflow = DataflowContent->GetDataflowAsset())
+			{
+				const Dataflow::FTimestamp SystemTimestamp = LatestTimestamp(Dataflow, DataflowContext.Get());
+				if (SystemTimestamp >= DataflowContent->GetLastModifiedTimestamp() || DataflowContent->IsDirty())
+				{
+					DataflowContent->SetLastModifiedTimestamp(SystemTimestamp.Value + 1);
+
+					if (DataflowContent->IsDirty())
+					{
+						UpdateConstructionScene();
+					}
+				}
+			}
+		}
+	}
+	for (TObjectPtr<UInteractiveToolPropertySet>& Propset : PropertyObjectsToTick)
+	{
+		if (Propset)
+		{
+			if (Propset->IsPropertySetEnabled())
+			{
+				Propset->CheckAndUpdateWatched();
+			}
+			else
+			{
+				Propset->SilentUpdateWatched();
+			}
+		}
+	}
+
+	for (FRenderWireElement Elem : WireframeElements)
+	{
+		Elem.Value->OnTick(DeltaSeconds);
+	}
+}
+
+void FDataflowConstructionScene::UpdateDynamicMeshComponents()
+{
+	using namespace UE::Geometry;//FDynamicMesh3
+
+	// The preview scene for the construction view will be
+	// cleared and rebuilt from scratch. This will genrate a 
+	// list of UPrimitiveComponents for rendering.
+	ResetDynamicMeshComponents();
+
+	if (TObjectPtr<UDataflowBaseContent> DataflowContent = GetDataflowContent())
+	{
+		const TObjectPtr<UDataflow>& DataflowAsset = DataflowContent->GetDataflowAsset();
+		const TSharedPtr<Dataflow::FEngineContext>& DataflowContext = DataflowContent->GetDataflowContext();
+		if(DataflowAsset && DataflowContext)
+		{
+			for (TObjectPtr<const UDataflowEdNode> Target : DataflowAsset->GetRenderTargets())
+			{
+				if (Target)
+				{
+					TSharedPtr<FManagedArrayCollection> RenderCollection(new FManagedArrayCollection);
+					GeometryCollection::Facades::FRenderingFacade Facade(*RenderCollection);
+					Facade.DefineSchema();
+
+					Target->Render(Facade, DataflowContext);
+
+					int32 NumGeometry = Facade.NumGeometry();
+					for (int32 MeshIndex = 0; MeshIndex < NumGeometry; MeshIndex++)
+					{
+						FDynamicMesh3 DynamicMesh;
+						Dataflow::Conversion::RenderingFacadeToDynamicMesh(Facade, MeshIndex, DynamicMesh);
+
+						if (DynamicMesh.VertexCount())
+						{
+							if (Target == DataflowContent->GetPrimarySelectedNode())
+							{
+								DataflowContent->SetPrimaryRenderCollection(RenderCollection);
+							}
+
+							AddDynamicMeshComponent({Target, MeshIndex }, MoveTemp(DynamicMesh), {});
+						}
+					}
+				}
+			}
+		}
+	}
+}
+
+void FDataflowConstructionScene::ResetDynamicMeshComponents()
+{
+	USelection* SelectedComponents = DataflowModeManager->GetSelectedComponents();
+	for (FRenderElement RenderElement : DynamicMeshComponents)
+	{
+		TObjectPtr<UDynamicMeshComponent>& DynamicMeshComponent = RenderElement.Value;
+
+		DynamicMeshComponent->SelectionOverrideDelegate.Unbind();
+		if (SelectedComponents->IsSelected(DynamicMeshComponent))
+		{
+			SelectedComponents->Deselect(DynamicMeshComponent);
+			DynamicMeshComponent->PushSelectionToProxy();
+		}
+		RemoveComponent(DynamicMeshComponent);
+	}
+	DynamicMeshComponents.Reset();
+}
+
+TObjectPtr<UDynamicMeshComponent>& FDataflowConstructionScene::AddDynamicMeshComponent(FDataflowRenderKey InKey, UE::Geometry::FDynamicMesh3&& DynamicMesh, const TArray<UMaterialInterface*>& MaterialSet)
+{
+	TObjectPtr<UDataflowEditorCollectionComponent> DynamicMeshComponent = NewObject<UDataflowEditorCollectionComponent>(RootSceneActor);
+	DynamicMeshComponent->MeshIndex = InKey.Value;
+	DynamicMeshComponent->Node = InKey.Key;;
+	DynamicMeshComponent->SetMesh(MoveTemp(DynamicMesh));
+	
+	// @todo(Material) This is just to have a material, we should transfer the materials from the assets if they have them. 
+	TObjectPtr<UDataflowBaseContent> DataflowContent = GetDataflowContent();
+	if (DataflowContent && DataflowContent->GetDataflowAsset() && DataflowContent->GetDataflowAsset()->Material)
+	{
+		DynamicMeshComponent->ConfigureMaterialSet({ DataflowContent->GetDataflowAsset()->Material });
+	}
+	else
+	{
+		DynamicMeshComponent->SetOverrideRenderMaterial(FDataflowEditorStyle::Get().VertexMaterial);
+		DynamicMeshComponent->SetShadowsEnabled(false);
+	}
+	//else if (FDataflowEditorStyle::Get().DefaultMaterial)
+	//{
+	//	DynamicMeshComponent->ConfigureMaterialSet({ FDataflowEditorStyle::Get().DefaultMaterial });
+	//}
+	//else
+	//{
+	//	DynamicMeshComponent->ValidateMaterialSlots(true, false);
+	//}
+
+	DynamicMeshComponent->SelectionOverrideDelegate = UPrimitiveComponent::FSelectionOverride::CreateRaw(this, &FDataflowPreviewSceneBase::IsComponentSelected);
+	DynamicMeshComponent->UpdateBounds();
+
+	AddComponent(DynamicMeshComponent, DynamicMeshComponent->GetRelativeTransform());	
+	DynamicMeshComponents.Emplace(InKey, DynamicMeshComponent);
+	return DynamicMeshComponents[InKey];
+}
+
+void FDataflowConstructionScene::AddWireframeMeshElementsVisualizer()
+{
+	if(!bDataflowShowWireframeInConstructionView) return;
+
+	ensure(WireframeElements.Num()==0);
+	for(FRenderElement Elem : DynamicMeshComponents)
+	{
+		if( TObjectPtr<UDataflowEditorCollectionComponent> DynamicMeshComponent = Cast<UDataflowEditorCollectionComponent>(Elem.Value) )
+		{
+			// Set up the wireframe display of the rest space mesh.
+
+			TObjectPtr<UMeshElementsVisualizer> WireframeDraw = NewObject<UMeshElementsVisualizer>(RootSceneActor);
+			WireframeElements.Add(DynamicMeshComponent, WireframeDraw);
+
+			WireframeDraw->CreateInWorld(GetWorld(), FTransform::Identity);
+			WireframeDraw->Settings->DepthBias = 2.0;
+			WireframeDraw->Settings->bAdjustDepthBiasUsingMeshSize = false;
+			WireframeDraw->Settings->bShowWireframe = true;
+			WireframeDraw->Settings->bShowBorders = true;
+			WireframeDraw->Settings->bShowUVSeams = false;
+			WireframeDraw->WireframeComponent->BoundaryEdgeThickness = 2;
+			DynamicMeshComponent->WireframeComponent = WireframeDraw->WireframeComponent;
+
+			WireframeDraw->SetMeshAccessFunction([DynamicMeshComponent](UMeshElementsVisualizer::ProcessDynamicMeshFunc ProcessFunc)
+				{
+					ProcessFunc(*DynamicMeshComponent->GetMesh());
+				});
+
+			for (FRenderElement RenderElement : DynamicMeshComponents)
+			{
+				RenderElement.Value->OnMeshChanged.Add(FSimpleMulticastDelegate::FDelegate::CreateLambda([WireframeDraw, this]()
+					{
+						WireframeDraw->NotifyMeshChanged();
+					}));
+
+				const bool bRestSpaceMeshVisible = RenderElement.Value->GetVisibleFlag();
+				WireframeDraw->Settings->bVisible = bRestSpaceMeshVisible && bConstructionViewWireframe;
+			}
+			PropertyObjectsToTick.Add(WireframeDraw->Settings);
+		}
+	}
+}
+
+void FDataflowConstructionScene::ResetWireframeMeshElementsVisualizer()
+{
+	for (FRenderWireElement Elem : WireframeElements)
+	{
+		Elem.Value->Disconnect();
+	}
+	WireframeElements.Empty();
+}
+
+void FDataflowConstructionScene::UpdateWireframeMeshElementsVisualizer()
+{
+	ResetWireframeMeshElementsVisualizer();
+	AddWireframeMeshElementsVisualizer();
+}
+
+bool FDataflowConstructionScene::HasRenderableGeometry()
+{
+	for (FRenderElement RenderElement : DynamicMeshComponents)
+	{
+		if (RenderElement.Value->GetMesh()->TriangleCount() > 0)
+		{
+			return true;
+		}
+	}
+	return false;
+}
+
+void FDataflowConstructionScene::ResetConstructionScene()
+{
+	// Some objects, like the UMeshElementsVisualizer and Settings Objects
+	// are not part of a tool, so they won't get ticked.This member holds
+	// ticked objects that get rebuilt on Update
+	PropertyObjectsToTick.Empty();
+
+	ResetWireframeMeshElementsVisualizer();
+
+	ResetDynamicMeshComponents();
+}
+
+void FDataflowConstructionScene::UpdateConstructionScene()
+{
+	ResetConstructionScene();
+
+	// The preview scene for the construction view will be
+	// cleared and rebuilt from scratch. This will genrate a 
+	// list of UPrimitiveComponents for rendering.
+	UpdateDynamicMeshComponents();
+	
+	// Attach a wireframe renderer to the DynamicMeshComponents
+	UpdateWireframeMeshElementsVisualizer();
+
+	if (TObjectPtr<UDataflowBaseContent> DataflowContent = GetDataflowContent())
+	{
+		DataflowContent->SetIsDirty(false);
+	}
+}
+
+#undef LOCTEXT_NAMESPACE
+
diff --git a/Engine/Plugins/Experimental/Dataflow/Source/DataflowEditor/Private/Dataflow/DataflowConstructionViewport.cpp b/Engine/Plugins/Experimental/Dataflow/Source/DataflowEditor/Private/Dataflow/DataflowConstructionViewport.cpp
--- a/Engine/Plugins/Experimental/Dataflow/Source/DataflowEditor/Private/Dataflow/DataflowConstructionViewport.cpp
+++ b/Engine/Plugins/Experimental/Dataflow/Source/DataflowEditor/Private/Dataflow/DataflowConstructionViewport.cpp
@@ -9,7 +9,8 @@
 #include "EditorModeManager.h"
 #include "Dataflow/DataflowContent.h"
 #include "Dataflow/DataflowConstructionViewportToolbar.h"
-#include "Dataflow/DataflowEditorScenes.h"
+#include "Dataflow/DataflowConstructionScene.h"
+#include "Dataflow/DataflowEditorPreviewSceneBase.h"
 #include "Dataflow/DataflowSimulationPanel.h"
 
 #define LOCTEXT_NAMESPACE "SDataflowConstructionViewport"
@@ -34,7 +35,7 @@ TSharedPtr<SWidget> SDataflowConstructionViewport::MakeViewportToolbar()
 
 void SDataflowConstructionViewport::OnFocusViewportToSelection()
 {
-	if(const FDataflowPreviewScene* PreviewScene = static_cast<FDataflowPreviewScene*>(Client->GetPreviewScene()))
+	if(const FDataflowPreviewSceneBase* PreviewScene = static_cast<FDataflowPreviewSceneBase*>(Client->GetPreviewScene()))
 	{
 		const FBox SceneBoundingBox = PreviewScene->GetBoundingBox();
 		Client->FocusViewportOnBox(SceneBoundingBox);
@@ -81,12 +82,12 @@ void SDataflowConstructionViewport::OnFloatingButtonClicked()
 
 float SDataflowConstructionViewport:: GetViewMinInput() const
 {
-	return static_cast<FDataflowPreviewScene*>(Client->GetPreviewScene())->GetDataflowContent()->GetSimulationRange()[0];
+	return static_cast<FDataflowPreviewSceneBase*>(Client->GetPreviewScene())->GetDataflowContent()->GetSimulationRange()[0];
 }
 
 float SDataflowConstructionViewport::GetViewMaxInput() const
 {
-	return static_cast<FDataflowPreviewScene*>(Client->GetPreviewScene())->GetDataflowContent()->GetSimulationRange()[1];
+	return static_cast<FDataflowPreviewSceneBase*>(Client->GetPreviewScene())->GetDataflowContent()->GetSimulationRange()[1];
 }
 
 
diff --git a/Engine/Plugins/Experimental/Dataflow/Source/DataflowEditor/Private/Dataflow/DataflowConstructionViewportClient.cpp b/Engine/Plugins/Experimental/Dataflow/Source/DataflowEditor/Private/Dataflow/DataflowConstructionViewportClient.cpp
--- a/Engine/Plugins/Experimental/Dataflow/Source/DataflowEditor/Private/Dataflow/DataflowConstructionViewportClient.cpp
+++ b/Engine/Plugins/Experimental/Dataflow/Source/DataflowEditor/Private/Dataflow/DataflowConstructionViewportClient.cpp
@@ -8,7 +8,8 @@
 #include "Dataflow/DataflowEditorCollectionComponent.h"
 #include "Dataflow/DataflowEngineSceneHitProxies.h"
 #include "Dataflow/DataflowGraphEditor.h"
-#include "Dataflow/DataflowEditorScenes.h"
+#include "Dataflow/DataflowConstructionScene.h"
+#include "Dataflow/DataflowEditorPreviewSceneBase.h"
 #include "EditorModeManager.h"
 #include "EdModeInteractiveToolsContext.h"
 #include "GraphEditor.h"
@@ -30,7 +31,7 @@ FDataflowConstructionViewportClient::FDataflowConstructionViewportClient(FEditor
 	EngineShowFlags.SetSelectionOutline(true);
 	EngineShowFlags.EnableAdvancedFeatures();
 
-	PreviewScene = static_cast<FDataflowPreviewScene*>(InPreviewScene);
+	PreviewScene = static_cast<FDataflowPreviewSceneBase*>(InPreviewScene);
 	bEnableSceneTicking = bCouldTickScene;
 }
 
diff --git a/Engine/Plugins/Experimental/Dataflow/Source/DataflowEditor/Private/Dataflow/DataflowEditorMode.cpp b/Engine/Plugins/Experimental/Dataflow/Source/DataflowEditor/Private/Dataflow/DataflowEditorMode.cpp
--- a/Engine/Plugins/Experimental/Dataflow/Source/DataflowEditor/Private/Dataflow/DataflowEditorMode.cpp
+++ b/Engine/Plugins/Experimental/Dataflow/Source/DataflowEditor/Private/Dataflow/DataflowEditorMode.cpp
@@ -16,7 +16,9 @@
 #include "Dataflow/DataflowEditorCommands.h"
 #include "Dataflow/DataflowEditorModeToolkit.h"
 #include "Dataflow/DataflowGraphEditor.h"
-#include "Dataflow/DataflowEditorScenes.h"
+#include "Dataflow/DataflowEditorPreviewSceneBase.h"
+#include "Dataflow/DataflowConstructionScene.h"
+#include "Dataflow/DataflowSimulationScene.h"
 #include "Dataflow/DataflowSimulationViewportClient.h"
 #include "Dataflow/DataflowSNode.h"
 #include "Dataflow/DataflowToolTarget.h"
diff --git a/Engine/Plugins/Experimental/Dataflow/Source/DataflowEditor/Private/Dataflow/DataflowEditorModeToolkit.cpp b/Engine/Plugins/Experimental/Dataflow/Source/DataflowEditor/Private/Dataflow/DataflowEditorModeToolkit.cpp
--- a/Engine/Plugins/Experimental/Dataflow/Source/DataflowEditor/Private/Dataflow/DataflowEditorModeToolkit.cpp
+++ b/Engine/Plugins/Experimental/Dataflow/Source/DataflowEditor/Private/Dataflow/DataflowEditorModeToolkit.cpp
@@ -7,7 +7,8 @@
 #include "Dataflow/DataflowEditorMode.h"
 #include "Dataflow/DataflowEditorStyle.h"
 #include "Dataflow/DataflowConstructionViewport.h"
-#include "Dataflow/DataflowEditorScenes.h"
+#include "Dataflow/DataflowEditorPreviewSceneBase.h"
+#include "Dataflow/DataflowConstructionScene.h"
 #include "EdModeInteractiveToolsContext.h"
 #include "Framework/Application/SlateApplication.h"
 #include "InteractiveToolManager.h"
@@ -158,7 +159,7 @@ SBaseCharacterFXEditorViewport* FDataflowEditorModeToolkit::GetViewportWidgetFor
 
 		if (const UDataflowEditorMode* const DataflowEdMode = Cast<UDataflowEditorMode>(Mode))
 		{
-			if (const FDataflowPreviewScene* const PreviewScene = DataflowEdMode->GetDataflowConstructionScene())
+			if (const FDataflowPreviewSceneBase* const PreviewScene = DataflowEdMode->GetDataflowConstructionScene())
 			{
 				if (const UEditorInteractiveToolsContext* const PreviewToolsContext = PreviewScene->GetDataflowModeManager()->GetInteractiveToolsContext())
 				{
diff --git a/Engine/Plugins/Experimental/Dataflow/Source/DataflowEditor/Private/Dataflow/DataflowEditorPreviewSceneBase.cpp b/Engine/Plugins/Experimental/Dataflow/Source/DataflowEditor/Private/Dataflow/DataflowEditorPreviewSceneBase.cpp
new file mode 100644
--- /dev/null
+++ b/Engine/Plugins/Experimental/Dataflow/Source/DataflowEditor/Private/Dataflow/DataflowEditorPreviewSceneBase.cpp
@@ -0,0 +1,97 @@
+// Copyright Epic Games, Inc. All Rights Reserved.
+
+#include "Dataflow/DataflowEditorPreviewSceneBase.h"
+
+
+#include "AssetEditorModeManager.h"
+#include "Dataflow/DataflowEditor.h"
+#include "Elements/Framework/EngineElementsLibrary.h"
+#include "Selection.h"
+
+#define LOCTEXT_NAMESPACE "FDataflowPreviewSceneBase"
+
+
+bool bDataflowShowFloorDefault = false;
+FAutoConsoleVariableRef CVARDataflowShowFloorDefault(TEXT("p.Dataflow.Editor.ShowFloor"), bDataflowShowFloorDefault, TEXT("Show the floor in the dataflow editor[def:false]"));
+
+
+FDataflowPreviewSceneBase::FDataflowPreviewSceneBase(FPreviewScene::ConstructionValues ConstructionValues, UDataflowEditor* InEditor)
+	: FAdvancedPreviewScene(ConstructionValues)
+	, DataflowEditor(InEditor)
+{
+	check(DataflowEditor);
+	SetFloorVisibility(bDataflowShowFloorDefault, true);
+
+	RootSceneActor = GetWorld()->SpawnActor<AActor>(AActor::StaticClass());
+}
+
+FDataflowPreviewSceneBase::~FDataflowPreviewSceneBase()
+{}
+
+TObjectPtr<UDataflowBaseContent> FDataflowPreviewSceneBase::GetDataflowContent() 
+{ 
+	return DataflowEditor->GetDataflowContent();
+}
+
+const TObjectPtr<UDataflowBaseContent> FDataflowPreviewSceneBase::GetDataflowContent() const 
+{ 
+	return DataflowEditor->GetDataflowContent();
+}
+
+void FDataflowPreviewSceneBase::AddReferencedObjects(FReferenceCollector& Collector)
+{
+	FAdvancedPreviewScene::AddReferencedObjects(Collector);
+	Collector.AddReferencedObject(RootSceneActor);
+	if (GetDataflowContent())
+	{
+		GetDataflowContent()->AddContentObjects(Collector);
+	}
+}
+
+bool FDataflowPreviewSceneBase::IsComponentSelected(const UPrimitiveComponent* InComponent) const
+{
+	if(DataflowModeManager.IsValid())
+	{
+		if (const UTypedElementSelectionSet* const TypedElementSelectionSet = DataflowModeManager->GetEditorSelectionSet())
+		{
+			if (const FTypedElementHandle ComponentElement = UEngineElementsLibrary::AcquireEditorComponentElementHandle(InComponent))
+			{
+				const bool bElementSelected = TypedElementSelectionSet->IsElementSelected(ComponentElement, FTypedElementIsSelectedOptions());
+				return bElementSelected;
+			}
+		}
+	}
+	return false;
+}
+
+FBox FDataflowPreviewSceneBase::GetBoundingBox() const
+{
+	FBox SceneBounds(ForceInitToZero);
+	if(DataflowModeManager.IsValid())
+	{
+		USelection* const SelectedComponents = DataflowModeManager->GetSelectedComponents();
+
+		TArray<TWeakObjectPtr<UObject>> SelectedObjects;
+		const int32 NumSelected = SelectedComponents->GetSelectedObjects(SelectedObjects);
+		
+		if(NumSelected > 0)
+		{
+			for(const TWeakObjectPtr<UObject> SelectedObject : SelectedObjects)
+			{
+				if(const UPrimitiveComponent* SelectedComponent = Cast<UPrimitiveComponent>(SelectedObject))
+				{
+					SceneBounds += SelectedComponent->Bounds.GetBox();
+				}
+			}
+		}
+		else
+		{
+			SceneBounds += RootSceneActor->GetComponentsBoundingBox(true);
+		}
+	}
+	return SceneBounds;
+}
+
+
+#undef LOCTEXT_NAMESPACE
+
diff --git a/Engine/Plugins/Experimental/Dataflow/Source/DataflowEditor/Private/Dataflow/DataflowEditorScenes.cpp b/Engine/Plugins/Experimental/Dataflow/Source/DataflowEditor/Private/Dataflow/DataflowEditorScenes.cpp
deleted file mode 100644
diff --git a/Engine/Plugins/Experimental/Dataflow/Source/DataflowEditor/Private/Dataflow/DataflowEditorToolkit.cpp b/Engine/Plugins/Experimental/Dataflow/Source/DataflowEditor/Private/Dataflow/DataflowEditorToolkit.cpp
--- a/Engine/Plugins/Experimental/Dataflow/Source/DataflowEditor/Private/Dataflow/DataflowEditorToolkit.cpp
+++ b/Engine/Plugins/Experimental/Dataflow/Source/DataflowEditor/Private/Dataflow/DataflowEditorToolkit.cpp
@@ -21,7 +21,9 @@
 #include "Dataflow/DataflowNodeFactory.h"
 #include "Dataflow/DataflowObject.h"
 #include "Dataflow/DataflowObjectInterface.h"
-#include "Dataflow/DataflowEditorScenes.h"
+#include "Dataflow/DataflowEditorPreviewSceneBase.h"
+#include "Dataflow/DataflowConstructionScene.h"
+#include "Dataflow/DataflowSimulationScene.h"
 #include "Dataflow/DataflowSchema.h"
 #include "Dataflow/DataflowSimulationViewportClient.h"
 #include "DynamicMeshBuilder.h"
@@ -190,7 +192,7 @@ void FDataflowEditorToolkit::CreateEditorModeManager()
 {
 	// Setup the construction manager / scene
 	FBaseCharacterFXEditorToolkit::CreateEditorModeManager();
-	static_cast<FDataflowPreviewScene*>(ObjectScene.Get())->GetDataflowModeManager()
+	static_cast<FDataflowPreviewSceneBase*>(ObjectScene.Get())->GetDataflowModeManager()
 		= StaticCastSharedPtr<FAssetEditorModeManager>(EditorModeManager);
 
 	// Setup the simulation manager / scene
@@ -198,7 +200,7 @@ void FDataflowEditorToolkit::CreateEditorModeManager()
 	StaticCastSharedPtr<FAssetEditorModeManager>(SimulationModeManager)->SetPreviewScene(
 		SimulationScene.Get());
 
-	static_cast<FDataflowPreviewScene*>(SimulationScene.Get())->GetDataflowModeManager()
+	static_cast<FDataflowPreviewSceneBase*>(SimulationScene.Get())->GetDataflowModeManager()
 		= StaticCastSharedPtr<FAssetEditorModeManager>(SimulationModeManager);
 }
 
@@ -591,7 +593,7 @@ void FDataflowEditorToolkit::OnNodeSelectionChanged(const TSet<UObject*>& InNewS
 				}
 			}
 
-			if (TObjectPtr<AActor> RootActor = static_cast<FDataflowPreviewScene*>(ObjectScene.Get())->GetRootActor())
+			if (TObjectPtr<AActor> RootActor = static_cast<FDataflowPreviewSceneBase*>(ObjectScene.Get())->GetRootActor())
 			{
 				for (UActorComponent* ActorComponent : RootActor->GetComponents())
 				{
diff --git a/Engine/Plugins/Experimental/Dataflow/Source/DataflowEditor/Private/Dataflow/DataflowSimulationPanel.cpp b/Engine/Plugins/Experimental/Dataflow/Source/DataflowEditor/Private/Dataflow/DataflowSimulationPanel.cpp
--- a/Engine/Plugins/Experimental/Dataflow/Source/DataflowEditor/Private/Dataflow/DataflowSimulationPanel.cpp
+++ b/Engine/Plugins/Experimental/Dataflow/Source/DataflowEditor/Private/Dataflow/DataflowSimulationPanel.cpp
@@ -1,7 +1,7 @@
-﻿// Copyright Epic Games, Inc. All Rights Reserved.
+// Copyright Epic Games, Inc. All Rights Reserved.
 
 #include "Dataflow/DataflowSimulationPanel.h"
-#include "Dataflow/DataflowEditorScenes.h"
+#include "Dataflow/DataflowSimulationScene.h"
 #include "Widgets/SBoxPanel.h"
 #include "SScrubControlPanel.h"
 #include "Animation/AnimSingleNodeInstance.h"
diff --git a/Engine/Plugins/Experimental/Dataflow/Source/DataflowEditor/Private/Dataflow/DataflowSimulationScene.cpp b/Engine/Plugins/Experimental/Dataflow/Source/DataflowEditor/Private/Dataflow/DataflowSimulationScene.cpp
new file mode 100644
--- /dev/null
+++ b/Engine/Plugins/Experimental/Dataflow/Source/DataflowEditor/Private/Dataflow/DataflowSimulationScene.cpp
@@ -0,0 +1,54 @@
+// Copyright Epic Games, Inc. All Rights Reserved.
+
+#include "Dataflow/DataflowSimulationScene.h"
+
+#include "Dataflow/DataflowEditor.h"
+#include "Selection.h"
+
+#define LOCTEXT_NAMESPACE "FDataflowSimulationScene"
+
+//
+// FDataflowSimulationScene
+//
+
+FDataflowSimulationScene::FDataflowSimulationScene(FPreviewScene::ConstructionValues ConstructionValues, UDataflowEditor* InEditor)
+	: FDataflowPreviewSceneBase(ConstructionValues, InEditor)
+{
+	if (TObjectPtr<UDataflowBaseContent> DataflowContent = GetDataflowContent())
+	{
+		DataflowContent->RegisterWorldContent(this, RootSceneActor);
+	}
+
+	TInlineComponentArray<UPrimitiveComponent*> PrimComponents;
+	RootSceneActor->GetComponents(PrimComponents);
+
+	for(UPrimitiveComponent* PrimComponent : PrimComponents)
+	{
+		PrimComponent->SelectionOverrideDelegate =
+			UPrimitiveComponent::FSelectionOverride::CreateRaw(this, &FDataflowPreviewSceneBase::IsComponentSelected);
+	}
+}
+
+FDataflowSimulationScene::~FDataflowSimulationScene()
+{
+	TInlineComponentArray<UPrimitiveComponent*> PrimComponents;
+	RootSceneActor->GetComponents(PrimComponents);
+
+	for(UPrimitiveComponent* PrimComponent : PrimComponents)
+	{
+		PrimComponent->SelectionOverrideDelegate.Unbind();
+	}
+
+	if (TObjectPtr<UDataflowBaseContent> DataflowContent = GetDataflowContent())
+	{
+		DataflowContent->UnregisterWorldContent(this);
+	}
+}
+
+void FDataflowSimulationScene::TickDataflowScene(const float DeltaSeconds)
+{
+	GetWorld()->Tick(ELevelTick::LEVELTICK_All, DeltaSeconds);
+}
+
+#undef LOCTEXT_NAMESPACE
+
diff --git a/Engine/Plugins/Experimental/Dataflow/Source/DataflowEditor/Private/Dataflow/DataflowSimulationViewport.cpp b/Engine/Plugins/Experimental/Dataflow/Source/DataflowEditor/Private/Dataflow/DataflowSimulationViewport.cpp
--- a/Engine/Plugins/Experimental/Dataflow/Source/DataflowEditor/Private/Dataflow/DataflowSimulationViewport.cpp
+++ b/Engine/Plugins/Experimental/Dataflow/Source/DataflowEditor/Private/Dataflow/DataflowSimulationViewport.cpp
@@ -8,7 +8,7 @@
 #include "Dataflow/DataflowEditorToolkit.h"
 #include "EditorModeManager.h"
 #include "Dataflow/DataflowContent.h"
-#include "Dataflow/DataflowEditorScenes.h"
+#include "Dataflow/DataflowSimulationScene.h"
 #include "Dataflow/DataflowSimulationPanel.h"
 
 #define LOCTEXT_NAMESPACE "SDataflowSimulationViewport"
@@ -55,7 +55,7 @@ void SDataflowSimulationViewport::Construct(const FArguments& InArgs, const FAss
 
 void SDataflowSimulationViewport::OnFocusViewportToSelection()
 {
-	if(const FDataflowPreviewScene* PreviewScene = static_cast<FDataflowPreviewScene*>(Client->GetPreviewScene()))
+	if(const FDataflowPreviewSceneBase* PreviewScene = static_cast<FDataflowPreviewSceneBase*>(Client->GetPreviewScene()))
 	{
 		const FBox SceneBoundingBox = PreviewScene->GetBoundingBox();
 		Client->FocusViewportOnBox(SceneBoundingBox);
@@ -102,12 +102,12 @@ void SDataflowSimulationViewport::OnFloatingButtonClicked()
 
 float SDataflowSimulationViewport:: GetViewMinInput() const
 {
-	return static_cast<FDataflowPreviewScene*>(Client->GetPreviewScene())->GetDataflowContent()->GetSimulationRange()[0];
+	return static_cast<FDataflowPreviewSceneBase*>(Client->GetPreviewScene())->GetDataflowContent()->GetSimulationRange()[0];
 }
 
 float SDataflowSimulationViewport::GetViewMaxInput() const
 {
-	return static_cast<FDataflowPreviewScene*>(Client->GetPreviewScene())->GetDataflowContent()->GetSimulationRange()[1];
+	return static_cast<FDataflowPreviewSceneBase*>(Client->GetPreviewScene())->GetDataflowContent()->GetSimulationRange()[1];
 }
 
 
diff --git a/Engine/Plugins/Experimental/Dataflow/Source/DataflowEditor/Private/Dataflow/DataflowSimulationViewportClient.cpp b/Engine/Plugins/Experimental/Dataflow/Source/DataflowEditor/Private/Dataflow/DataflowSimulationViewportClient.cpp
--- a/Engine/Plugins/Experimental/Dataflow/Source/DataflowEditor/Private/Dataflow/DataflowSimulationViewportClient.cpp
+++ b/Engine/Plugins/Experimental/Dataflow/Source/DataflowEditor/Private/Dataflow/DataflowSimulationViewportClient.cpp
@@ -8,7 +8,7 @@
 #include "Dataflow/DataflowEditorCollectionComponent.h"
 #include "Dataflow/DataflowEngineSceneHitProxies.h"
 #include "Dataflow/DataflowGraphEditor.h"
-#include "Dataflow/DataflowEditorScenes.h"
+#include "Dataflow/DataflowEditorPreviewSceneBase.h"
 #include "EditorModeManager.h"
 #include "EdModeInteractiveToolsContext.h"
 #include "GraphEditor.h"
@@ -30,7 +30,7 @@ FDataflowSimulationViewportClient::FDataflowSimulationViewportClient(FEditorMode
 	EngineShowFlags.SetSelectionOutline(true);
 	EngineShowFlags.EnableAdvancedFeatures();
 
-	PreviewScene = static_cast<FDataflowPreviewScene*>(InPreviewScene);
+	PreviewScene = static_cast<FDataflowPreviewSceneBase*>(InPreviewScene);
 	bEnableSceneTicking = bCouldTickScene;
 }
 
diff --git a/Engine/Plugins/Experimental/Dataflow/Source/DataflowEditor/Private/DataflowEditorTools/DataflowEditorWeightMapPaintTool.cpp b/Engine/Plugins/Experimental/Dataflow/Source/DataflowEditor/Private/DataflowEditorTools/DataflowEditorWeightMapPaintTool.cpp
--- a/Engine/Plugins/Experimental/Dataflow/Source/DataflowEditor/Private/DataflowEditorTools/DataflowEditorWeightMapPaintTool.cpp
+++ b/Engine/Plugins/Experimental/Dataflow/Source/DataflowEditor/Private/DataflowEditorTools/DataflowEditorWeightMapPaintTool.cpp
@@ -17,7 +17,7 @@
 #include "Dataflow/DataflowEdNode.h"
 #include "Dataflow/DataflowEditorCollectionComponent.h"
 #include "Dataflow/DataflowEditorMode.h"
-#include "Dataflow/DataflowEditorScenes.h"
+#include "Dataflow/DataflowConstructionScene.h"
 #include "Dataflow/DataflowObject.h"
 #include "Dataflow/DataflowSNode.h"
 #include "Dataflow/DataflowGraphEditor.h"
diff --git a/Engine/Plugins/Experimental/Dataflow/Source/DataflowEditor/Public/Dataflow/DataflowConstructionScene.h b/Engine/Plugins/Experimental/Dataflow/Source/DataflowEditor/Public/Dataflow/DataflowConstructionScene.h
new file mode 100644
--- /dev/null
+++ b/Engine/Plugins/Experimental/Dataflow/Source/DataflowEditor/Public/Dataflow/DataflowConstructionScene.h
@@ -0,0 +1,90 @@
+// Copyright Epic Games, Inc. All Rights Reserved.
+
+#pragma once
+
+#include "Dataflow/DataflowEditorPreviewSceneBase.h"
+
+class UDataflowEditor;
+class USkeletalMeshComponent;
+class UDynamicMeshComponent;
+class UMeshElementsVisualizer;
+class UDataflowEdNode;
+class UInteractiveToolPropertySet;
+
+namespace UE
+{
+	namespace Geometry
+	{
+		class FDynamicMesh3;
+	}
+}
+
+/**
+ * Dataflow construction scene holding all the dynamic meshes used/displayed in the dataflow graph
+ */
+class DATAFLOWEDITOR_API FDataflowConstructionScene : public FDataflowPreviewSceneBase
+{
+	/*(const void* ObjectPointer, const int32 ObjectIndex)*/
+	typedef TPair<TObjectPtr<const UDataflowEdNode>, int32> FDataflowRenderKey;
+
+public:
+
+	FDataflowConstructionScene(FPreviewScene::ConstructionValues ConstructionValues, UDataflowEditor* Editor);
+	virtual ~FDataflowConstructionScene();
+
+	// FGCObject interface
+	virtual void AddReferencedObjects(FReferenceCollector& Collector) override;
+	
+	/** Tick data flow scene */
+	virtual void TickDataflowScene(const float DeltaSeconds) override;
+
+	/** Check if the preview scene can run simulation */
+	virtual bool CanRunSimulation() const override {return false;}
+	
+	/** Is there ANYTHING to render */
+	bool HasRenderableGeometry();
+
+	/** Close the construction scene */
+	void ResetConstructionScene();
+
+	/** Update the construction scene */
+	void UpdateConstructionScene();
+	
+	/** Hide all or a single component */
+	void SetVisibility(bool bVisible, UActorComponent* Component = nullptr);
+
+private:
+
+	/** Reset the scene mesh visualizer */
+	void ResetWireframeMeshElementsVisualizer();
+
+	/** Update the scene mesh visualizer */
+	void UpdateWireframeMeshElementsVisualizer();
+
+	/** Add a mesh visualizer to the scene */
+	void AddWireframeMeshElementsVisualizer();
+
+	/** Reset all the dynamic mesh components */
+    void ResetDynamicMeshComponents();
+	
+	/** Build the dynamic mesh components from the dataflow graph render targets*/
+	void UpdateDynamicMeshComponents();
+	
+	/** Add a dynamic mesh component to the scene */
+	TObjectPtr<UDynamicMeshComponent>& AddDynamicMeshComponent(
+		FDataflowRenderKey InKey, UE::Geometry::FDynamicMesh3&& DynamicMesh,  const TArray<UMaterialInterface*>& MaterialSet);
+	
+	/** Dynamic mesh components generated by the tools for edit/manipulation*/
+	typedef TPair<FDataflowRenderKey, TObjectPtr<UDynamicMeshComponent>> FRenderElement;
+	TMap<FDataflowRenderKey, TObjectPtr<UDynamicMeshComponent>> DynamicMeshComponents;
+
+	/** Rest space wireframe. They have to get ticked to be able to respond to setting changes. [ticked] */
+	typedef TPair<TObjectPtr<UDynamicMeshComponent>, TObjectPtr<UMeshElementsVisualizer>> FRenderWireElement;
+	TMap<TObjectPtr<UDynamicMeshComponent>,TObjectPtr<UMeshElementsVisualizer> > WireframeElements;
+
+	/** Property objects (visible or not) that get ticked. [ticked] */
+	TArray<TObjectPtr<UInteractiveToolPropertySet>> PropertyObjectsToTick;
+
+	/** Show the wireframe on render */ 
+	bool bConstructionViewWireframe = true;
+};
diff --git a/Engine/Plugins/Experimental/Dataflow/Source/DataflowEditor/Public/Dataflow/DataflowConstructionViewportClient.h b/Engine/Plugins/Experimental/Dataflow/Source/DataflowEditor/Public/Dataflow/DataflowConstructionViewportClient.h
--- a/Engine/Plugins/Experimental/Dataflow/Source/DataflowEditor/Public/Dataflow/DataflowConstructionViewportClient.h
+++ b/Engine/Plugins/Experimental/Dataflow/Source/DataflowEditor/Public/Dataflow/DataflowConstructionViewportClient.h
@@ -13,7 +13,7 @@ class ADataflowActor;
 class UTransformProxy;
 class UCombinedTransformGizmo;
 class FTransformGizmoDataBinder;
-class FDataflowPreviewScene;
+class FDataflowPreviewSceneBase;
 class UInputBehaviorSet;
 class USelection;
 
@@ -55,7 +55,7 @@ private:
 	TWeakPtr<FDataflowEditorToolkit> DataflowEditorToolkitPtr = nullptr;
 
 	/** Dataflow preview scene from the toolkit */
-	FDataflowPreviewScene* PreviewScene = nullptr;
+	FDataflowPreviewSceneBase* PreviewScene = nullptr;
 
 	// @todo(brice) : Is this needed?
 	TWeakPtr<FUICommandList> ToolCommandList;
diff --git a/Engine/Plugins/Experimental/Dataflow/Source/DataflowEditor/Public/Dataflow/DataflowEditorPreviewSceneBase.h b/Engine/Plugins/Experimental/Dataflow/Source/DataflowEditor/Public/Dataflow/DataflowEditorPreviewSceneBase.h
new file mode 100644
--- /dev/null
+++ b/Engine/Plugins/Experimental/Dataflow/Source/DataflowEditor/Public/Dataflow/DataflowEditorPreviewSceneBase.h
@@ -0,0 +1,60 @@
+// Copyright Epic Games, Inc. All Rights Reserved.
+
+#pragma once
+
+#include "AdvancedPreviewScene.h"
+
+class UDataflowBaseContent;
+class UDataflowEditor;
+class FAssetEditorModeManager;
+
+/**
+ * Dataflow preview scene base
+ * @brief the scene is holding all the objects that will be
+ * visible and potentially editable within the viewport
+ */
+class DATAFLOWEDITOR_API FDataflowPreviewSceneBase : public FAdvancedPreviewScene
+{
+public:
+
+	FDataflowPreviewSceneBase(FPreviewScene::ConstructionValues ConstructionValues, UDataflowEditor* Editor);
+	virtual ~FDataflowPreviewSceneBase();
+	
+	// FGCObject interface
+	virtual void AddReferencedObjects(FReferenceCollector& Collector) override;
+
+	/** Dataflow content accessors */
+	TObjectPtr<UDataflowBaseContent> GetDataflowContent();
+		const TObjectPtr<UDataflowBaseContent> GetDataflowContent() const;
+
+	/** Root scene actor accessors */
+	TObjectPtr<AActor> GetRootActor() { return RootSceneActor; }
+	const TObjectPtr<AActor> GetRootActor() const { return RootSceneActor; }
+
+	/** Dataflow mode manager accessors */
+	TSharedPtr<FAssetEditorModeManager>& GetDataflowModeManager() { return DataflowModeManager; }
+	const TSharedPtr<FAssetEditorModeManager>& GetDataflowModeManager() const { return DataflowModeManager; }
+	
+	/** Build the scene bounding box */
+	FBox GetBoundingBox() const;
+
+	/** Tick data flow scene */
+	virtual void TickDataflowScene(const float DeltaSeconds) {}
+
+	/** Check if a primitive component is selected */
+	bool IsComponentSelected(const UPrimitiveComponent* InComponent) const;
+
+	/** Check if the preview scene can run simulation */
+	virtual bool CanRunSimulation() const {return false;}
+	
+protected:
+	
+	/** Root scene actor */
+	TObjectPtr<AActor> RootSceneActor = nullptr;
+
+	UDataflowEditor* DataflowEditor = nullptr;
+
+	/** Mode Manager for selection */
+	TSharedPtr<FAssetEditorModeManager> DataflowModeManager;
+};
+
diff --git a/Engine/Plugins/Experimental/Dataflow/Source/DataflowEditor/Public/Dataflow/DataflowEditorScenes.h b/Engine/Plugins/Experimental/Dataflow/Source/DataflowEditor/Public/Dataflow/DataflowEditorScenes.h
deleted file mode 100644
diff --git a/Engine/Plugins/Experimental/Dataflow/Source/DataflowEditor/Public/Dataflow/DataflowSimulationPanel.h b/Engine/Plugins/Experimental/Dataflow/Source/DataflowEditor/Public/Dataflow/DataflowSimulationPanel.h
--- a/Engine/Plugins/Experimental/Dataflow/Source/DataflowEditor/Public/Dataflow/DataflowSimulationPanel.h
+++ b/Engine/Plugins/Experimental/Dataflow/Source/DataflowEditor/Public/Dataflow/DataflowSimulationPanel.h
@@ -2,7 +2,7 @@
 
 #pragma once
 
-#include "DataflowEditorScenes.h"
+#include "Dataflow/DataflowSimulationScene.h"
 #include "Widgets/SCompoundWidget.h"
 #include "ITransportControl.h"  // EPlaybackMode::Type
 
diff --git a/Engine/Plugins/Experimental/Dataflow/Source/DataflowEditor/Public/Dataflow/DataflowSimulationScene.h b/Engine/Plugins/Experimental/Dataflow/Source/DataflowEditor/Public/Dataflow/DataflowSimulationScene.h
new file mode 100644
--- /dev/null
+++ b/Engine/Plugins/Experimental/Dataflow/Source/DataflowEditor/Public/Dataflow/DataflowSimulationScene.h
@@ -0,0 +1,31 @@
+// Copyright Epic Games, Inc. All Rights Reserved.
+
+#pragma once
+
+#include "Dataflow/DataflowEditorPreviewSceneBase.h"
+
+class UDataflowEditor;
+
+
+/**
+ * Dataflow simulation scene holding all the dataflow content components
+ */
+class DATAFLOWEDITOR_API FDataflowSimulationScene : public FDataflowPreviewSceneBase
+{
+public:
+
+	FDataflowSimulationScene(FPreviewScene::ConstructionValues ConstructionValues, UDataflowEditor* Editor);
+	virtual ~FDataflowSimulationScene();
+
+	/** Tick data flow scene */
+	virtual void TickDataflowScene(const float DeltaSeconds) override;
+	
+	/** Check if the preview scene can run simulation */
+	virtual bool CanRunSimulation() const { return true; }
+
+	void ResetSimulationScene() {}
+	bool HasRenderableGeometry() { return true; }
+
+};
+
+
diff --git a/Engine/Plugins/Experimental/Dataflow/Source/DataflowEditor/Public/Dataflow/DataflowSimulationViewportClient.h b/Engine/Plugins/Experimental/Dataflow/Source/DataflowEditor/Public/Dataflow/DataflowSimulationViewportClient.h
--- a/Engine/Plugins/Experimental/Dataflow/Source/DataflowEditor/Public/Dataflow/DataflowSimulationViewportClient.h
+++ b/Engine/Plugins/Experimental/Dataflow/Source/DataflowEditor/Public/Dataflow/DataflowSimulationViewportClient.h
@@ -13,7 +13,7 @@ class ADataflowActor;
 class UTransformProxy;
 class UCombinedTransformGizmo;
 class FTransformGizmoDataBinder;
-class FDataflowPreviewScene;
+class FDataflowPreviewSceneBase;
 class UInputBehaviorSet;
 
 class DATAFLOWEDITOR_API FDataflowSimulationViewportClient : public FEditorViewportClient //, public IInputBehaviorSource
@@ -49,7 +49,7 @@ private:
 	TWeakPtr<FDataflowEditorToolkit> DataflowEditorToolkitPtr = nullptr;
 
 	/** Dataflow preview scene from the toolkit */
-	FDataflowPreviewScene* PreviewScene = nullptr;
+	FDataflowPreviewSceneBase* PreviewScene = nullptr;
 
 	// @todo(brice) : Is this needed?
 	TWeakPtr<FUICommandList> ToolCommandList;
