Dataflow - Draw non active nodes as transparent.

diff --git a/Engine/Plugins/Experimental/Dataflow/Source/DataflowEditor/Private/Dataflow/DataflowEditorCommands.cpp b/Engine/Plugins/Experimental/Dataflow/Source/DataflowEditor/Private/Dataflow/DataflowEditorCommands.cpp
--- a/Engine/Plugins/Experimental/Dataflow/Source/DataflowEditor/Private/Dataflow/DataflowEditorCommands.cpp
+++ b/Engine/Plugins/Experimental/Dataflow/Source/DataflowEditor/Private/Dataflow/DataflowEditorCommands.cpp
@@ -355,6 +355,12 @@ void FDataflowEditorCommands::OnPropertyValueChanged(UDataflow* OutDataflow, TSh
 					{
 						UpdatedNode = DataflowNode;
 						DataflowNode->Invalidate();
+
+						// Reflect the active state on the drawing of the node
+						if (DataflowNode->bActive != Node->IsNodeEnabled())
+						{
+							Node->SetEnabledState(DataflowNode->bActive ? ENodeEnabledState::Enabled : ENodeEnabledState::Disabled);
+						}
 					}
 				}
 			}
