Dataflow : ContextCaching
- Allow disabling with cvar 'p.Dataflow.Editor.ContextCaching'
- Udpated path to place cache in sub-directory of the asset.

diff --git a/Engine/Plugins/Experimental/Dataflow/Source/DataflowEditor/Private/Dataflow/AssetDefinition_DataflowContext.cpp b/Engine/Plugins/Experimental/Dataflow/Source/DataflowEditor/Private/Dataflow/AssetDefinition_DataflowContext.cpp
--- a/Engine/Plugins/Experimental/Dataflow/Source/DataflowEditor/Private/Dataflow/AssetDefinition_DataflowContext.cpp
+++ b/Engine/Plugins/Experimental/Dataflow/Source/DataflowEditor/Private/Dataflow/AssetDefinition_DataflowContext.cpp
@@ -18,8 +18,9 @@
 
 #define LOCTEXT_NAMESPACE "AssetActions_DataflowContext"
 
-bool bDataflowEnableContextCacheLoading = true;
-FAutoConsoleVariableRef CVARDataflowEnableContextCacheLoading(TEXT("p.Dataflow.Editor.ContextCaching"), bDataflowEnableContextCacheLoading, TEXT("Allow the Dataflow editor to load the cached state of the graph when and asset is re-opened.[def:true]"));
+bool bDataflowEnableContextCaching = true;
+FAutoConsoleVariableRef CVARDataflowEnableContextCaching(TEXT("p.Dataflow.Editor.ContextCaching"), bDataflowEnableContextCaching, 
+	TEXT("Allow the Dataflow editor to crate and use a pre-evaluated graph when the dataflow editor is re-opened.[def:true]"));
 
 namespace DataflowContextDefinitionHelpers
 {
@@ -172,28 +173,33 @@ namespace DataflowContextDefinitionHelpers
 	TObjectPtr<T> CreateNewDataflowContext(const TObjectPtr<UObject>& ContentOwner)
 	{
 		check(ContentOwner.Get());
-
 		bool bNeedsNewAsset = true;
 		TObjectPtr<UObject> Asset = nullptr;
-		UClass* DataflowClass = T::StaticClass();
-		UDataflow* DataflowAsset = Private::GetDataflowAssetFrom(ContentOwner);
-
-		FContentBrowserModule& ContentBrowserModule = FModuleManager::LoadModuleChecked<FContentBrowserModule>("ContentBrowser");
 
-		FString PackageName = FString::Printf(TEXT("/Game/_GENERATED/Dataflow/%s"), *ContentOwner.GetName());
-		if (DataflowAsset)
+		if (bDataflowEnableContextCaching)
 		{
-			PackageName = FString::Printf(TEXT("%s_%s"), *PackageName, *DataflowAsset->GetName());
-		}
+			// Setup an asset that lives in the content broswer.
+			FContentBrowserModule& ContentBrowserModule = FModuleManager::LoadModuleChecked<FContentBrowserModule>("ContentBrowser");
+
+			UClass* DataflowClass = T::StaticClass();
+			UDataflow* DataflowAsset = Private::GetDataflowAssetFrom(ContentOwner);
+
+			const FString AssetPackageName = ContentOwner->GetOutermost()->GetName();
+			const FString AssetDefaultPath = FPackageName::GetLongPackagePath(AssetPackageName);
+
+			FString PackageName = FString::Printf(TEXT("%s/Cache/Transient/DataflowContext_%s"), *AssetDefaultPath, *ContentOwner.GetName());
+			if (DataflowAsset)
+			{
+				PackageName = FString::Printf(TEXT("%s_%s"), *PackageName, *DataflowAsset->GetName());
+			}
+
+			UPackage* Package = FindPackage(nullptr, *PackageName);
+			if (!Package)
+			{
+				Package = CreatePackage(*PackageName);
+			}
 
-		UPackage* Package = FindPackage(nullptr, *PackageName);
-		if (!Package)
-		{
-			Package = CreatePackage(*PackageName);
-		}
 
-		if (bDataflowEnableContextCacheLoading)
-		{
 			Asset = StaticLoadObject(DataflowClass, Package, *PackageName);
 			if (Asset && DataflowAsset)
 			{
@@ -203,22 +209,26 @@ namespace DataflowContextDefinitionHelpers
 					!ValidateCachedNodeHash(Asset, DataflowAsset) ||
 					!ResetCacheTimestamp(Asset, DataflowAsset);
 			}
-		}
-
-		if (!Asset || bNeedsNewAsset)
-		{
-			const FName AssetName(FPackageName::GetLongPackageAssetName(PackageName));
-			Asset = NewObject<UObject>(Package, DataflowClass, AssetName, RF_Public | RF_Standalone | RF_Transactional);
-
-			Asset->MarkPackageDirty();
-			FAssetRegistryModule::AssetCreated(Asset);
 
-			if (UDataflowBaseContent* BaseContent = Cast< UDataflowBaseContent>(Asset.Get()))
+			if (!Asset || bNeedsNewAsset)
 			{
-				BaseContent->BuildBaseContent(ContentOwner);
-				BaseContent->SetDataflowOwner(ContentOwner);
+				const FName AssetName(FPackageName::GetLongPackageAssetName(PackageName));
+				Asset = NewObject<UObject>(Package, DataflowClass, AssetName, RF_Public | RF_Standalone | RF_Transactional);
+
+				Asset->MarkPackageDirty();
+				FAssetRegistryModule::AssetCreated(Asset);
 			}
 		}
+		else
+		{
+			Asset = NewObject<T>(ContentOwner, T::StaticClass());
+		}
+
+		if (UDataflowBaseContent* BaseContent = Cast< UDataflowBaseContent>(Asset.Get()))
+		{
+			BaseContent->BuildBaseContent(ContentOwner);
+			BaseContent->SetDataflowOwner(ContentOwner);
+		}
 
 		return Cast<T>(Asset.Get());
 	}
