ChaosFlesh: Expose mesh name in render facade.

diff --git a/Engine/Plugins/Experimental/Dataflow/Source/DataflowEditor/Private/Dataflow/DataflowConstructionScene.cpp b/Engine/Plugins/Experimental/Dataflow/Source/DataflowEditor/Private/Dataflow/DataflowConstructionScene.cpp
--- a/Engine/Plugins/Experimental/Dataflow/Source/DataflowEditor/Private/Dataflow/DataflowConstructionScene.cpp
+++ b/Engine/Plugins/Experimental/Dataflow/Source/DataflowEditor/Private/Dataflow/DataflowConstructionScene.cpp
@@ -157,8 +157,8 @@ void FDataflowConstructionScene::UpdateDynamicMeshComponents()
 							{
 								DataflowContent->SetPrimaryRenderCollection(RenderCollection);
 							}
-
-							AddDynamicMeshComponent({Target, MeshIndex }, MoveTemp(DynamicMesh), {});
+							const FString MeshName = Facade.GetGeometryName()[MeshIndex];
+							AddDynamicMeshComponent({Target, MeshIndex}, MeshName, MoveTemp(DynamicMesh), {});
 						}
 					}
 				}
@@ -185,9 +185,9 @@ void FDataflowConstructionScene::ResetDynamicMeshComponents()
 	DynamicMeshComponents.Reset();
 }
 
-TObjectPtr<UDynamicMeshComponent>& FDataflowConstructionScene::AddDynamicMeshComponent(FDataflowRenderKey InKey, UE::Geometry::FDynamicMesh3&& DynamicMesh, const TArray<UMaterialInterface*>& MaterialSet)
+TObjectPtr<UDynamicMeshComponent>& FDataflowConstructionScene::AddDynamicMeshComponent(FDataflowRenderKey InKey, const FString& MeshName, UE::Geometry::FDynamicMesh3&& DynamicMesh, const TArray<UMaterialInterface*>& MaterialSet)
 {
-	TObjectPtr<UDataflowEditorCollectionComponent> DynamicMeshComponent = NewObject<UDataflowEditorCollectionComponent>(RootSceneActor);
+	TObjectPtr<UDataflowEditorCollectionComponent> DynamicMeshComponent = NewObject<UDataflowEditorCollectionComponent>(RootSceneActor, FName(MeshName));
 	DynamicMeshComponent->MeshIndex = InKey.Value;
 	DynamicMeshComponent->Node = InKey.Key;;
 	DynamicMeshComponent->SetMesh(MoveTemp(DynamicMesh));
diff --git a/Engine/Plugins/Experimental/Dataflow/Source/DataflowEditor/Public/Dataflow/DataflowConstructionScene.h b/Engine/Plugins/Experimental/Dataflow/Source/DataflowEditor/Public/Dataflow/DataflowConstructionScene.h
--- a/Engine/Plugins/Experimental/Dataflow/Source/DataflowEditor/Public/Dataflow/DataflowConstructionScene.h
+++ b/Engine/Plugins/Experimental/Dataflow/Source/DataflowEditor/Public/Dataflow/DataflowConstructionScene.h
@@ -72,7 +72,7 @@ private:
 	
 	/** Add a dynamic mesh component to the scene */
 	TObjectPtr<UDynamicMeshComponent>& AddDynamicMeshComponent(
-		FDataflowRenderKey InKey, UE::Geometry::FDynamicMesh3&& DynamicMesh,  const TArray<UMaterialInterface*>& MaterialSet);
+		FDataflowRenderKey InKey, const FString& MeshName, UE::Geometry::FDynamicMesh3&& DynamicMesh,  const TArray<UMaterialInterface*>& MaterialSet);
 	
 	/** Dynamic mesh components generated by the tools for edit/manipulation*/
 	typedef TPair<FDataflowRenderKey, TObjectPtr<UDynamicMeshComponent>> FRenderElement;
diff --git a/Engine/Plugins/Experimental/Dataflow/Source/DataflowEnginePlugin/Private/Dataflow/DataflowEngineRendering.cpp b/Engine/Plugins/Experimental/Dataflow/Source/DataflowEnginePlugin/Private/Dataflow/DataflowEngineRendering.cpp
--- a/Engine/Plugins/Experimental/Dataflow/Source/DataflowEnginePlugin/Private/Dataflow/DataflowEngineRendering.cpp
+++ b/Engine/Plugins/Experimental/Dataflow/Source/DataflowEnginePlugin/Private/Dataflow/DataflowEngineRendering.cpp
@@ -120,6 +120,7 @@ namespace Dataflow
 		const TManagedArray<int32>& BoneIndex = Collection.GetAttribute<int32>("BoneMap", FGeometryCollection::VerticesGroup);
 		const TManagedArray<int32>& Parents = Collection.GetAttribute<int32>(FTransformCollection::ParentAttribute, FTransformCollection::TransformGroup);
 		const TManagedArray<FTransform3f>& Transforms = Collection.GetAttribute<FTransform3f>(FTransformCollection::TransformAttribute, FTransformCollection::TransformGroup);
+		const TManagedArray<FString>& BoneNames = Collection.GetAttribute<FString>("BoneName", FGeometryCollection::TransformGroup);
 		const TManagedArray<FVector3f>& Vertex = Collection.GetAttribute<FVector3f>("Vertex", FGeometryCollection::VerticesGroup);
 		const TManagedArray<FIntVector>& Faces = Collection.GetAttribute<FIntVector>("Indices", FGeometryCollection::FacesGroup);
 		const TManagedArray<bool>* FaceVisible = Collection.FindAttribute<bool>("Visible", FGeometryCollection::FacesGroup);
@@ -128,6 +129,7 @@ namespace Dataflow
 		const TManagedArray<int32>& VertexCount = Collection.GetAttribute<int32>("VertexCount", FGeometryCollection::GeometryGroup);
 		const TManagedArray<int32>& FacesStart = Collection.GetAttribute<int32>("FaceStart", FGeometryCollection::GeometryGroup);
 		const TManagedArray<int32>& FacesCount = Collection.GetAttribute<int32>("FaceCount", FGeometryCollection::GeometryGroup);
+		
 
 		TArray<FTransform> M;
 		GeometryCollectionAlgo::GlobalMatrices(Transforms, Parents, M);
@@ -213,6 +215,10 @@ namespace Dataflow
 			if (Vertices.Num() && Tris.Num())
 			{
 				FString GeometryName = State.GetGuid().ToString(); GeometryName.AppendChar('.').AppendInt(Gdx);
+				if (BoneIndex[VertexStart[Gdx]] != INDEX_NONE)
+				{
+					 GeometryName = BoneNames[BoneIndex[VertexStart[Gdx]]];
+				}
 				int32 GeometryIndex = RenderCollection.StartGeometryGroup(GeometryName);
 				RenderCollection.AddSurface(MoveTemp(Vertices), MoveTemp(Tris), MoveTemp(VertexNormals), MoveTemp(VertexColors));
 				RenderCollection.EndGeometryGroup(GeometryIndex);
