Dataflow : ContextCaching
- Moved Transient directory to Dataflow
- Fixed issue with Context regenration on load

diff --git a/Engine/Plugins/Experimental/Dataflow/Source/DataflowEditor/Private/Dataflow/AssetDefinition_DataflowContext.cpp b/Engine/Plugins/Experimental/Dataflow/Source/DataflowEditor/Private/Dataflow/AssetDefinition_DataflowContext.cpp
--- a/Engine/Plugins/Experimental/Dataflow/Source/DataflowEditor/Private/Dataflow/AssetDefinition_DataflowContext.cpp
+++ b/Engine/Plugins/Experimental/Dataflow/Source/DataflowEditor/Private/Dataflow/AssetDefinition_DataflowContext.cpp
@@ -187,7 +187,7 @@ namespace DataflowContextDefinitionHelpers
 			const FString AssetPackageName = ContentOwner->GetOutermost()->GetName();
 			const FString AssetDefaultPath = FPackageName::GetLongPackagePath(AssetPackageName);
 
-			FString PackageName = FString::Printf(TEXT("%s/Cache/Transient/DataflowContext_%s"), *AssetDefaultPath, *ContentOwner.GetName());
+			FString PackageName = FString::Printf(TEXT("%s/Cache/Dataflow/DataflowContext_%s"), *AssetDefaultPath, *ContentOwner.GetName());
 			if (DataflowAsset)
 			{
 				PackageName = FString::Printf(TEXT("%s_%s"), *PackageName, *DataflowAsset->GetName());
@@ -217,19 +217,27 @@ namespace DataflowContextDefinitionHelpers
 
 				Asset->MarkPackageDirty();
 				FAssetRegistryModule::AssetCreated(Asset);
+
+				if (UDataflowBaseContent* BaseContent = Cast< UDataflowBaseContent>(Asset.Get()))
+				{
+					BaseContent->BuildBaseContent(ContentOwner);
+				}
 			}
 		}
 		else
 		{
 			Asset = NewObject<T>(ContentOwner, T::StaticClass());
+			if (UDataflowBaseContent* BaseContent = Cast< UDataflowBaseContent>(Asset.Get()))
+			{
+				BaseContent->BuildBaseContent(ContentOwner);
+			}
 		}
 
+
 		if (UDataflowBaseContent* BaseContent = Cast< UDataflowBaseContent>(Asset.Get()))
 		{
-			BaseContent->BuildBaseContent(ContentOwner);
 			BaseContent->SetDataflowOwner(ContentOwner);
 		}
-
 		return Cast<T>(Asset.Get());
 	}
 
