Dataflow : Collection spreadsheet index out of bounds fix.

diff --git a/Engine/Plugins/Experimental/Dataflow/Source/DataflowEditor/Private/Dataflow/CollectionSpreadSheetWidget.cpp b/Engine/Plugins/Experimental/Dataflow/Source/DataflowEditor/Private/Dataflow/CollectionSpreadSheetWidget.cpp
--- a/Engine/Plugins/Experimental/Dataflow/Source/DataflowEditor/Private/Dataflow/CollectionSpreadSheetWidget.cpp
+++ b/Engine/Plugins/Experimental/Dataflow/Source/DataflowEditor/Private/Dataflow/CollectionSpreadSheetWidget.cpp
@@ -115,6 +115,11 @@ namespace CollectionSpreadSheetHelpers
 			return FString("<Unknown Attribute>");
 		}
 
+		if (InIdxColumn < 0 || Array->Num() <= InIdxColumn)
+		{
+			return FString("<Index out of bounds>");
+		}
+
 		return AttributeValueToString((*Array)[InIdxColumn]);
 	}
 
