GeometryCollection : AttributeKey Nodes

diff --git a/Engine/Plugins/Experimental/Dataflow/Source/DataflowNodes/Private/Dataflow/DataflowCollectionAttributeKeyNodes.cpp b/Engine/Plugins/Experimental/Dataflow/Source/DataflowNodes/Private/Dataflow/DataflowCollectionAttributeKeyNodes.cpp
new file mode 100644
--- /dev/null
+++ b/Engine/Plugins/Experimental/Dataflow/Source/DataflowNodes/Private/Dataflow/DataflowCollectionAttributeKeyNodes.cpp
@@ -0,0 +1,32 @@
+// Copyright Epic Games, Inc. All Rights Reserved.
+
+#include "Dataflow/DataflowCollectionAttributeKeyNodes.h"
+#include "GeometryCollection/GeometryCollection.h"
+
+#include "Dataflow/DataflowCore.h"
+
+#include UE_INLINE_GENERATED_CPP_BY_NAME(DataflowCollectionAttributeKeyNodes)
+
+namespace Dataflow
+{
+	void DataflowCollectionAttributeKeyNodes()
+	{
+		DATAFLOW_NODE_REGISTER_CREATION_FACTORY(FMakeCollectionKeyDataflowNode);
+		DATAFLOW_NODE_REGISTER_CREATION_FACTORY(FBreakCollectionKeyDataflowNode);
+	}
+}
+
+
+void FMakeCollectionKeyDataflowNode::Evaluate(Dataflow::FContext& Context, const FDataflowOutput* Out) const
+{
+	FString GroupName = GetValue<FString>(Context, &GroupIn);
+	FString AttributeName = GetValue<FString>(Context, &AttributeIn);
+	SetValue(Context, FCollectionKey(GroupName, AttributeName), &CollectionKeyOut);
+}
+
+void FBreakCollectionKeyDataflowNode::Evaluate(Dataflow::FContext& Context, const FDataflowOutput* Out) const
+{
+	FCollectionKey CollectionKey = GetValue<FCollectionKey>(Context, &CollectionKeyIn);
+	SetValue(Context, CollectionKey.Group, &GroupOut);
+	SetValue(Context, CollectionKey.Attribute, &AttributeOut);
+}
\ No newline at end of file
diff --git a/Engine/Plugins/Experimental/Dataflow/Source/DataflowNodes/Private/Dataflow/DataflowNodesPlugin.cpp b/Engine/Plugins/Experimental/Dataflow/Source/DataflowNodes/Private/Dataflow/DataflowNodesPlugin.cpp
--- a/Engine/Plugins/Experimental/Dataflow/Source/DataflowNodes/Private/Dataflow/DataflowNodesPlugin.cpp
+++ b/Engine/Plugins/Experimental/Dataflow/Source/DataflowNodes/Private/Dataflow/DataflowNodesPlugin.cpp
@@ -3,6 +3,7 @@
 #include "Dataflow/DataflowNodesPlugin.h"
 
 #include "CoreMinimal.h"
+#include "Dataflow/DataflowCollectionAttributeKeyNodes.h"
 #include "Dataflow/DataflowSkeletalMeshNodes.h"
 #include "Dataflow/DataflowStaticMeshNodes.h"
 #include "Dataflow/DataflowNodeFactory.h"
@@ -19,6 +20,7 @@ void IDataflowNodesPlugin::StartupModule()
 	Dataflow::RegisterStaticMeshNodes();
 	Dataflow::RegisterSelectionNodes();
 	Dataflow::RegisterContextOverridesNodes();
+	Dataflow::DataflowCollectionAttributeKeyNodes();
 	DATAFLOW_NODE_REGISTER_CREATION_FACTORY(FDataflowCollectionAddScalarVertexPropertyNode);
 }
 
diff --git a/Engine/Plugins/Experimental/Dataflow/Source/DataflowNodes/Public/Dataflow/DataflowCollectionAttributeKeyNodes.h b/Engine/Plugins/Experimental/Dataflow/Source/DataflowNodes/Public/Dataflow/DataflowCollectionAttributeKeyNodes.h
new file mode 100644
--- /dev/null
+++ b/Engine/Plugins/Experimental/Dataflow/Source/DataflowNodes/Public/Dataflow/DataflowCollectionAttributeKeyNodes.h
@@ -0,0 +1,88 @@
+// Copyright Epic Games, Inc. All Rights Reserved.
+
+#pragma once 
+
+#include "CoreMinimal.h"
+#include "Dataflow/DataflowEngine.h"
+
+#include "DataflowCollectionAttributeKeyNodes.generated.h"
+
+class USkeletalMesh;
+
+
+USTRUCT(BlueprintType)
+struct DATAFLOWNODES_API FCollectionKey
+{
+	GENERATED_USTRUCT_BODY()
+public:
+	FCollectionKey() : Group(""), Attribute("") {}
+	FCollectionKey(FString InGroup, FString InAttribute)
+		:Group(InGroup), Attribute(InAttribute) {}
+	TPair<FName, FName> GetNamedKey() { return { FName(Group), FName(Attribute)}; }
+	FString Group;
+	FString Attribute;
+};
+
+
+USTRUCT(meta = (Dataflow))
+struct FMakeCollectionKeyDataflowNode : public FDataflowNode
+{
+	GENERATED_USTRUCT_BODY()
+	DATAFLOW_NODE_DEFINE_INTERNAL(FMakeCollectionKeyDataflowNode, "MakeCollectionKey", "GeometryCollection", "")
+
+public:
+
+	UPROPERTY(EditAnywhere, Category = "Collection Key", meta = (DataflowInput, DisplayName = "Group"))
+	FString GroupIn = "";
+
+	UPROPERTY(EditAnywhere, Category = "Collection Key", meta = (DataflowInput, DisplayName = "Attribute"))
+	FString AttributeIn = "";
+
+	UPROPERTY(meta = (DataflowOutput, DisplayName = "CollectionKey"))
+	FCollectionKey CollectionKeyOut;
+
+	FMakeCollectionKeyDataflowNode(const Dataflow::FNodeParameters& InParam, FGuid InGuid = FGuid::NewGuid())
+		: FDataflowNode(InParam, InGuid)
+	{
+		RegisterInputConnection(&GroupIn);
+		RegisterInputConnection(&AttributeIn);
+		RegisterOutputConnection(&CollectionKeyOut);
+	}
+
+	virtual void Evaluate(Dataflow::FContext& Context, const FDataflowOutput* Out) const override;
+};
+
+USTRUCT(meta = (DataflowGeometryCollection))
+struct FBreakCollectionKeyDataflowNode : public FDataflowNode
+{
+	GENERATED_USTRUCT_BODY()
+	DATAFLOW_NODE_DEFINE_INTERNAL(FBreakCollectionKeyDataflowNode, "BreakCollectionKey", "GeometryCollection", "")
+
+public:
+
+	UPROPERTY(meta = (DataflowInput, DisplayName = "CollectionKey"))
+	FCollectionKey CollectionKeyIn;
+
+	UPROPERTY(EditAnywhere, Category = "Collection Key", meta = (DataflowOutput, DisplayName = "Group"))
+	FString GroupOut = "";
+
+	UPROPERTY(EditAnywhere, Category = "Collection Key", meta = (DataflowOutput, DisplayName = "Attribute"))
+	FString AttributeOut = "";
+
+
+	FBreakCollectionKeyDataflowNode(const Dataflow::FNodeParameters& InParam, FGuid InGuid = FGuid::NewGuid())
+		: FDataflowNode(InParam, InGuid)
+	{
+		RegisterInputConnection(&CollectionKeyIn);
+		RegisterOutputConnection(&GroupOut);
+		RegisterOutputConnection(&AttributeOut);
+	}
+
+	virtual void Evaluate(Dataflow::FContext& Context, const FDataflowOutput* Out) const override;
+};
+
+namespace Dataflow
+{
+	void DataflowCollectionAttributeKeyNodes();
+}
+
