Dataflow : Skeleton View
-Refactored Skeleton View to match the FDataflowNodeView paradigm.
-Renamed ?Skeletal Hierarchy? to ?Skelton View?
-Added conversion for Collection to Skelton
-Updated FCollectionTransformFacade to support optional BoneName attribute.

diff --git a/Engine/Plugins/Experimental/Dataflow/Source/DataflowEditor/DataflowEditor.Build.cs b/Engine/Plugins/Experimental/Dataflow/Source/DataflowEditor/DataflowEditor.Build.cs
--- a/Engine/Plugins/Experimental/Dataflow/Source/DataflowEditor/DataflowEditor.Build.cs
+++ b/Engine/Plugins/Experimental/Dataflow/Source/DataflowEditor/DataflowEditor.Build.cs
@@ -36,6 +36,7 @@ namespace UnrealBuildTool.Rules
 					"EditorInteractiveToolsFramework",
 					"EditorStyle",
 					"GeometryCore",
+					"GeometryCollectionEngine",
 					"GeometryFramework",
 					"GraphEditor",
 					"InputCore",
@@ -52,6 +53,7 @@ namespace UnrealBuildTool.Rules
 					"RenderCore",
 					"RHI",
 					"SceneOutliner",
+					"SkeletonEditor",
 					"Slate",
 					"SlateCore",
 					"StaticMeshDescription",
diff --git a/Engine/Plugins/Experimental/Dataflow/Source/DataflowEditor/Private/Dataflow/DataflowEditorToolkit.cpp b/Engine/Plugins/Experimental/Dataflow/Source/DataflowEditor/Private/Dataflow/DataflowEditorToolkit.cpp
--- a/Engine/Plugins/Experimental/Dataflow/Source/DataflowEditor/Private/Dataflow/DataflowEditorToolkit.cpp
+++ b/Engine/Plugins/Experimental/Dataflow/Source/DataflowEditor/Private/Dataflow/DataflowEditorToolkit.cpp
@@ -8,6 +8,7 @@
 #include "Dataflow/DataflowCore.h"
 #include "Dataflow/DataflowEditor.h"
 #include "Dataflow/DataflowContent.h"
+#include "Dataflow/DataflowConstructionScene.h"
 #include "Dataflow/DataflowConstructionViewportClient.h"
 #include "Dataflow/DataflowEditorCollectionComponent.h"
 #include "Dataflow/DataflowEditorCommands.h"
@@ -25,6 +26,7 @@
 #include "Dataflow/DataflowConstructionScene.h"
 #include "Dataflow/DataflowSimulationScene.h"
 #include "Dataflow/DataflowSchema.h"
+#include "Dataflow/DataflowSkeletonView.h"
 #include "Dataflow/DataflowSimulationViewportClient.h"
 #include "DynamicMeshBuilder.h"
 #include "EditorModeManager.h"
@@ -48,14 +50,13 @@
 #include "Styling/SlateStyleRegistry.h"
 #include "Widgets/Docking/SDockTab.h"
 
-
 #define LOCTEXT_NAMESPACE "DataflowEditorToolkit"
 
 //DEFINE_LOG_CATEGORY_STATIC(EditorToolkitLog, Log, All);
 
 const FName FDataflowEditorToolkit::GraphCanvasTabId(TEXT("DataflowEditor_GraphCanvas"));
 const FName FDataflowEditorToolkit::NodeDetailsTabId(TEXT("DataflowEditor_NodeDetails"));
-const FName FDataflowEditorToolkit::SkeletalTabId(TEXT("DataflowEditor_Skeletal"));
+const FName FDataflowEditorToolkit::SkeletonViewTabId(TEXT("DataflowEditor_SkeletonView"));
 const FName FDataflowEditorToolkit::SelectionViewTabId_1(TEXT("DataflowEditor_SelectionView_1"));
 const FName FDataflowEditorToolkit::SelectionViewTabId_2(TEXT("DataflowEditor_SelectionView_2"));
 const FName FDataflowEditorToolkit::SelectionViewTabId_3(TEXT("DataflowEditor_SelectionView_3"));
@@ -418,8 +419,6 @@ void FDataflowEditorToolkit::CreateWidgets()
 			NodeDetailsEditor = CreateNodeDetailsEditorWidget(DataflowContent->GetDataflowOwner());
 			AssetDetailsEditor = CreateAssetDetailsEditorWidget(DataflowContent->GetDataflowOwner());
 			GraphEditor = CreateGraphEditorWidget(DataflowAsset, NodeDetailsEditor);
-			SkeletalEditor = CreateSkeletalEditorWidget();
-
 			CreateSimulationViewportClient();
 		}
 	}
@@ -532,6 +531,8 @@ void FDataflowEditorToolkit::OnNodeTitleCommitted(const FText& InNewText, ETextC
 	FDataflowEditorCommands::OnNodeTitleCommitted(InNewText, InCommitType, GraphNode);
 }
 
+
+
 void FDataflowEditorToolkit::OnNodeSelectionChanged(const TSet<UObject*>& InNewSelection)
 {
 	auto FindDataflowNodesInSet = [](const TSet<UObject*>& InSet) {
@@ -611,6 +612,8 @@ void FDataflowEditorToolkit::OnNodeSelectionChanged(const TSet<UObject*>& InNewS
 		}
 	};
 
+
+
 	// Despite this function's name, we might not have actually changed which node is selected
 	bool bPrimarySelectionChanged = false;
 
@@ -839,21 +842,6 @@ TSharedPtr<IDetailsView> FDataflowEditorToolkit::CreateAssetDetailsEditorWidget(
 
 }
 
-TSharedPtr<ISkeletonTree> FDataflowEditorToolkit::CreateSkeletalEditorWidget()
-{
-	if (const TObjectPtr<UDataflowSkeletalContent> SkeletalContent = Cast<UDataflowSkeletalContent>(GetDataflowContent()))
-	{
-		if (SkeletalContent->GetDataflowAsset())
-		{
-			const FSkeletonTreeArgs SkeletonTreeArgs;
-			ISkeletonEditorModule& SkeletonEditorModule = FModuleManager::LoadModuleChecked<ISkeletonEditorModule>("SkeletonEditor");
-			TSharedPtr<ISkeletonTree> SkeletonTree = SkeletonEditorModule.CreateSkeletonTree(SkeletalContent->GetSkeleton(), SkeletonTreeArgs);
-			return SkeletonTree;
-		}
-	}
-	return TSharedPtr<ISkeletonTree>(nullptr);
-}
-
 TSharedRef<SDockTab> FDataflowEditorToolkit::SpawnTab_AssetDetails(const FSpawnTabArgs& Args)
 {
 	check(Args.GetTabId() == DetailsTabID);
@@ -899,24 +887,36 @@ TSharedRef<SDockTab> FDataflowEditorToolkit::SpawnTab_NodeDetails(const FSpawnTa
 		];
 }
 
-TSharedRef<SDockTab> FDataflowEditorToolkit::SpawnTab_Skeletal(const FSpawnTabArgs& Args)
+TSharedRef<SDockTab> FDataflowEditorToolkit::SpawnTab_SkeletonView(const FSpawnTabArgs& Args)
 {
-	check(Args.GetTabId() == SkeletalTabId);
-	if(SkeletalEditor)
+	check(Args.GetTabId() == SkeletonViewTabId);
+
+
+	SkeletonEditorView = MakeShared<FDataflowSkeletonView>(DataflowEditor);
+
+	if (const TObjectPtr<UDataflowSkeletalContent> SkeletalContent = Cast<UDataflowSkeletalContent>(GetDataflowContent()))
 	{
-		if (const TObjectPtr<UDataflowSkeletalContent> EditorContent = Cast<UDataflowSkeletalContent>(GetDataflowContent()))
+		if (SkeletalContent->GetDataflowAsset())
 		{
-			SkeletalEditor->SetSkeletalMesh(EditorContent->GetSkeletalMesh());
+			SkeletonEditorView->SetSkeleton(SkeletalContent->GetSkeleton());
 		}
-		return SNew(SDockTab)
-			.Label(LOCTEXT("DataflowEditor_Skeletal_TabTitle", "Skeletal Hierarchy"))
-			[
-				SkeletalEditor.ToSharedRef()
-			];
 	}
+	ViewListeners.Add(SkeletonEditorView.Get());
+
+	FSkeletonTreeArgs SkeletonTreeArgs;
+	//TSharedRef<FDataflowSkeletonView> DataflowSkeletonViewRef = SkeletonEditorView.ToSharedRef();
+	SkeletonTreeArgs.OnSelectionChanged = FOnSkeletonTreeSelectionChanged::CreateSP(SkeletonEditorView.ToSharedRef(), &FDataflowSkeletonView::SkeletonViewSelectionChanged);
+	SkeletonTreeArgs.ContextName = GetToolkitFName();
+
+	ISkeletonEditorModule& SkeletonEditorModule = FModuleManager::LoadModuleChecked<ISkeletonEditorModule>("SkeletonEditor");
+	TSharedPtr<ISkeletonTree> SkeletonEditor = SkeletonEditorModule.CreateSkeletonTree(SkeletonEditorView->GetSkeleton(), SkeletonTreeArgs);
+	SkeletonEditorView->SetSkeletonEditor(SkeletonEditor);
 
 	return SNew(SDockTab)
-		.Label(LOCTEXT("DataflowEditor_Skeletal_TabTitle", "Skeletal Hierarchy"));
+		.Label(LOCTEXT("DataflowEditor_Skeleton_TabTitle", "Skeleton View"))
+		[
+			SkeletonEditor.ToSharedRef()
+		];
 }
 
 TSharedRef<SDockTab> FDataflowEditorToolkit::SpawnTab_SelectionView(const FSpawnTabArgs& Args)
@@ -1142,10 +1142,10 @@ void FDataflowEditorToolkit::RegisterTabSpawners(const TSharedRef<FTabManager>&
 		.SetGroup(EditorMenuCategory.ToSharedRef())
 		.SetIcon(FSlateIcon(FAppStyle::GetAppStyleSetName(), "LevelEditor.Tabs.Details"));
 
-	InTabManager->RegisterTabSpawner(SkeletalTabId, FOnSpawnTab::CreateSP(this, &FDataflowEditorToolkit::SpawnTab_Skeletal))
-		.SetDisplayName(LOCTEXT("DataflowSkeletalTab", "Skeletal Hierarchy"))
+	InTabManager->RegisterTabSpawner(SkeletonViewTabId, FOnSpawnTab::CreateSP(this, &FDataflowEditorToolkit::SpawnTab_SkeletonView))
+		.SetDisplayName(LOCTEXT("DataflowSkeletonTab", "Skeleton View"))
 		.SetGroup(EditorMenuCategory.ToSharedRef())
-		.SetIcon(FSlateIcon(FAppStyle::GetAppStyleSetName(), "LevelEditor.Tabs.SkeletalHierarchy"));
+		.SetIcon(FSlateIcon(FAppStyle::GetAppStyleSetName(), "LevelEditor.Tabs.SkeletonHierarchy"));
 
 	InTabManager->RegisterTabSpawner(SelectionViewTabId_1, FOnSpawnTab::CreateSP(this, &FDataflowEditorToolkit::SpawnTab_SelectionView))
 		.SetDisplayName(LOCTEXT("DataflowSelectionViewTab1", "Selection View 1"))
@@ -1194,7 +1194,7 @@ void FDataflowEditorToolkit::UnregisterTabSpawners(const TSharedRef<class FTabMa
 
 	InTabManager->UnregisterTabSpawner(GraphCanvasTabId);
 	InTabManager->UnregisterTabSpawner(NodeDetailsTabId);
-	InTabManager->UnregisterTabSpawner(SkeletalTabId);
+	InTabManager->UnregisterTabSpawner(SkeletonViewTabId);
 	InTabManager->UnregisterTabSpawner(SelectionViewTabId_1);
 	InTabManager->UnregisterTabSpawner(SelectionViewTabId_2);
 	InTabManager->UnregisterTabSpawner(SelectionViewTabId_3);
@@ -1240,6 +1240,10 @@ void FDataflowEditorToolkit::OnTabClosed(TSharedRef<SDockTab> Tab)
 	{
 		ViewListeners.Remove(DataflowCollectionSpreadSheet_4.Get());
 	}
+	else if (Tab->GetTabLabel().EqualTo(FText::FromString("Skeleton View")))
+	{
+		ViewListeners.Remove(SkeletonEditorView.Get());
+	}
 }
 
 FName FDataflowEditorToolkit::GetToolkitFName() const
diff --git a/Engine/Plugins/Experimental/Dataflow/Source/DataflowEditor/Private/Dataflow/DataflowSkeletonView.cpp b/Engine/Plugins/Experimental/Dataflow/Source/DataflowEditor/Private/Dataflow/DataflowSkeletonView.cpp
new file mode 100644
--- /dev/null
+++ b/Engine/Plugins/Experimental/Dataflow/Source/DataflowEditor/Private/Dataflow/DataflowSkeletonView.cpp
@@ -0,0 +1,141 @@
+// Copyright Epic Games, Inc. All Rights Reserved.
+
+#include "Dataflow/DataflowSkeletonView.h"
+
+#include "Animation/Skeleton.h"
+#include "Dataflow/DataflowEditor.h"
+#include "Dataflow/DataflowEdNode.h"
+#include "GeometryCollection/GeometryCollectionEngineConversion.h"
+#include "IEditableSkeleton.h"
+#include "ISkeletonTree.h"
+#include "ISkeletonTreeItem.h"
+#include "Engine/SkeletalMesh.h"
+
+//#include "Widgets/SCompoundWidget.h"
+
+FDataflowSkeletonView::FDataflowSkeletonView(UDataflowEditor* InDataflowEditor)
+	: DataflowEditor(InDataflowEditor)
+	, SkeletonEditor(nullptr)
+	, SkeletalMesh(NewObject<USkeletalMesh>())
+	, SelectedNode(nullptr)
+	, CollectionIndexRemap(TArray<int32>())
+
+{
+	SkeletalMesh->SetSkeleton(NewObject<USkeleton>(SkeletalMesh, NAME_Name));
+}
+
+FDataflowSkeletonView::~FDataflowSkeletonView()
+{
+	if (SkeletonEditor)
+	{
+		// remove widget delegates (see FDataflowCollectionSpreadSheet)
+	}
+}
+
+void FDataflowSkeletonView::SetSupportedOutputTypes()
+{
+	GetSupportedOutputTypes().Empty();
+
+	GetSupportedOutputTypes().Add("FManagedArrayCollection");
+}
+
+void FDataflowSkeletonView::SetSkeleton(USkeleton* Skeleton)
+{
+	if (Skeleton)
+	{
+		SkeletalMesh = NewObject<USkeletalMesh>();
+		SkeletalMesh->SetSkeleton(Skeleton);
+		SkeletalMesh->SetRefSkeleton(Skeleton->GetReferenceSkeleton());
+	}
+	else
+	{
+		SkeletalMesh = NewObject<USkeletalMesh>();
+		SkeletalMesh->SetSkeleton(NewObject<USkeleton>(SkeletalMesh, NAME_Name));
+	}
+}
+
+USkeleton* FDataflowSkeletonView::GetSkeleton()
+{
+	if (SkeletalMesh)
+	{
+		return SkeletalMesh->GetSkeleton();
+	}
+	return nullptr;
+}
+
+void FDataflowSkeletonView::UpdateViewData()
+{
+	if(!ensure(DataflowEditor)) return;
+
+	if (TObjectPtr<UDataflowEdNode> EdNode = GetSelectedNode())
+	{
+		if (GetSelectedNode()->IsBound())
+		{
+			if (TSharedPtr<FDataflowNode> Node = GetSelectedNode()->DataflowGraph->FindBaseNode(GetSelectedNode()->DataflowNodeGuid))
+			{
+				if (FDataflowOutput* Output = Node->FindOutput(FName("Collection")))
+				{
+					if (TObjectPtr<UDataflowBaseContent> EditorContent = DataflowEditor->GetDataflowContent())
+					{
+						if (EditorContent->GetDataflowAsset())
+						{
+							Dataflow::FTimestamp TimeStamp = EditorContent->GetLastModifiedTimestamp();
+							if (EditorContent->GetDataflowContext())
+							{
+								FManagedArrayCollection DefaultCollection;
+								const FManagedArrayCollection& Result = Output->GetValue(*EditorContent->GetDataflowContext(), DefaultCollection);
+
+								SelectedNode = EdNode;
+								SkeletalMesh = NewObject<USkeletalMesh>();
+								TObjectPtr<USkeleton> Skeleton = NewObject<USkeleton>(SkeletalMesh, Node->Name);
+
+								FGeometryCollectionEngineConversion::ConvertCollectionToSkeleton(Result, Skeleton, CollectionIndexRemap);
+								SkeletalMesh->SetSkeleton(Skeleton);
+								SkeletalMesh->SetRefSkeleton(Skeleton->GetReferenceSkeleton());
+
+								if (SkeletonEditor)
+								{
+									SkeletonEditor->GetEditableSkeleton()->RecreateBoneTree(SkeletalMesh);
+									SkeletonEditor->SetSkeletalMesh(SkeletalMesh);
+									SkeletonEditor->Refresh();
+								}
+							}
+						}
+					}
+				}
+			}
+		}
+	}
+}
+
+void FDataflowSkeletonView::SetSkeletonEditor(TSharedPtr<ISkeletonTree>& InSkeletonEditor)
+{
+	ensure(!SkeletonEditor);
+
+	SkeletonEditor = InSkeletonEditor;
+
+	if (SkeletonEditor)
+	{
+		// add widget delegates (see FDataflowCollectionSpreadSheet)
+	}
+}
+
+void FDataflowSkeletonView::SkeletonViewSelectionChanged(const TArrayView<TSharedPtr<ISkeletonTreeItem>>& InSelectedItems, ESelectInfo::Type InSelectInfo)
+{
+	if (SkeletonEditor)
+	{
+		//TArray<UObject*> Objects;
+		//Algo::TransformIf(InSelectedItems, Objects,
+		//	[](const TSharedPtr<ISkeletonTreeItem>& InItem) { return InItem->GetObject() != nullptr; },
+		//	[](const TSharedPtr<ISkeletonTreeItem>& InItem) { return InItem->GetObject(); });
+		//DetailsView->SetObjects(Objects);
+	}
+}
+
+void FDataflowSkeletonView::AddReferencedObjects(FReferenceCollector& Collector)
+{
+	FDataflowNodeView::AddReferencedObjects(Collector);
+
+	Collector.AddReferencedObject(SelectedNode);
+	Collector.AddReferencedObject(SkeletalMesh);
+}
diff --git a/Engine/Plugins/Experimental/Dataflow/Source/DataflowEditor/Public/Dataflow/DataflowEditorToolkit.h b/Engine/Plugins/Experimental/Dataflow/Source/DataflowEditor/Public/Dataflow/DataflowEditorToolkit.h
--- a/Engine/Plugins/Experimental/Dataflow/Source/DataflowEditor/Public/Dataflow/DataflowEditorToolkit.h
+++ b/Engine/Plugins/Experimental/Dataflow/Source/DataflowEditor/Public/Dataflow/DataflowEditorToolkit.h
@@ -25,8 +25,10 @@ class FDataflowConstructionScene;
 class FDataflowSimulationViewportClient;
 class UDataflowBaseContent;
 class FDataflowSimulationScene;
+class FDataflowSkeletonView;
 class UDataflowEditor;
 
+
 namespace Dataflow
 {
 	class DATAFLOWEDITOR_API FAssetContext : public TEngineContext<FContextSingle>
@@ -93,7 +95,7 @@ private:
 	// Spawning of all the additional tabs (viewport,details ones are coming from the base asset toolkit)
 	TSharedRef<SDockTab> SpawnTab_GraphCanvas(const FSpawnTabArgs& Args);
 	TSharedRef<SDockTab> SpawnTab_NodeDetails(const FSpawnTabArgs& Args);
-	TSharedRef<SDockTab> SpawnTab_Skeletal(const FSpawnTabArgs& Args);
+	TSharedRef<SDockTab> SpawnTab_SkeletonView(const FSpawnTabArgs& Args);
 	TSharedRef<SDockTab> SpawnTab_SelectionView(const FSpawnTabArgs& Args);
 	TSharedRef<SDockTab> SpawnTab_CollectionSpreadSheet(const FSpawnTabArgs& Args);
 	TSharedRef<SDockTab> SpawnTab_AssetDetails(const FSpawnTabArgs& Args);
@@ -123,7 +125,7 @@ private:
 	// List of all the tab names ids that will be used to identify the editor widgets
 	static const FName GraphCanvasTabId;
 	static const FName NodeDetailsTabId;
-	static const FName SkeletalTabId;
+	static const FName SkeletonViewTabId;
 	static const FName SelectionViewTabId_1;
 	static const FName SelectionViewTabId_2;
 	static const FName SelectionViewTabId_3;
@@ -139,7 +141,7 @@ private:
 	TSharedPtr<SDataflowSimulationViewport> DataflowSimulationViewport;
 	TSharedPtr<SDataflowGraphEditor> GraphEditor;
 	TSharedPtr<IStructureDetailsView> NodeDetailsEditor;
-	TSharedPtr<class ISkeletonTree> SkeletalEditor;
+	TSharedPtr<FDataflowSkeletonView> SkeletonEditorView;
 	TSharedPtr<IDetailsView> AssetDetailsEditor;
 	TSharedPtr<FDataflowSelectionView> DataflowSelectionView_1;
 	TSharedPtr<FDataflowSelectionView> DataflowSelectionView_2;
@@ -154,7 +156,6 @@ private:
 	TSharedRef<SDataflowGraphEditor> CreateGraphEditorWidget(UDataflow* ObjectToEdit, TSharedPtr<IStructureDetailsView> PropertiesEditor);
     TSharedPtr<IDetailsView> CreateAssetDetailsEditorWidget(UObject* ObjectToEdit);
     TSharedPtr<IStructureDetailsView> CreateNodeDetailsEditorWidget(UObject* ObjectToEdit);
-	TSharedPtr<ISkeletonTree> CreateSkeletalEditorWidget();
 
 	/** Create the simulation viewport client */
 	void CreateSimulationViewportClient();
diff --git a/Engine/Plugins/Experimental/Dataflow/Source/DataflowEditor/Public/Dataflow/DataflowSkeletonView.h b/Engine/Plugins/Experimental/Dataflow/Source/DataflowEditor/Public/Dataflow/DataflowSkeletonView.h
new file mode 100644
--- /dev/null
+++ b/Engine/Plugins/Experimental/Dataflow/Source/DataflowEditor/Public/Dataflow/DataflowSkeletonView.h
@@ -0,0 +1,60 @@
+// Copyright Epic Games, Inc. All Rights Reserved.
+
+#pragma once
+
+#include "Dataflow/DataflowView.h"
+
+class UDataflowEditor;
+class ISkeletonTree;
+class ISkeletonTreeItem;
+class USkeleton;
+class USkeletalMesh;
+class UDataflowEdNode;
+
+/**
+*
+* Class to handle the SelectionView widget
+*
+*/
+class FDataflowSkeletonView : public FDataflowNodeView
+{
+public:
+	FDataflowSkeletonView(UDataflowEditor* InDataflowEditor = nullptr);
+	~FDataflowSkeletonView();
+
+	/** Set the Skeleton Tree*/
+	void SetSkeletonEditor(TSharedPtr<ISkeletonTree>& InSkeletonTree);
+
+	/** Set selection types*/
+	virtual void SetSupportedOutputTypes() override;
+
+	/** Skeleton Access */
+	USkeleton* GetSkeleton();
+
+	/** Update Data*/
+	void SetSkeleton(USkeleton* Skeleton);
+
+	/** Update the view */
+	virtual void UpdateViewData() override;
+
+	/** Add GC managed objects*/
+	virtual void AddReferencedObjects(FReferenceCollector& Collector) override;
+
+	/** Selection View Callbacks */
+	void SkeletonViewSelectionChanged(const TArrayView<TSharedPtr<ISkeletonTreeItem>>& InSelectedItems, ESelectInfo::Type InSelectInfo);
+
+private:
+	UDataflowEditor* DataflowEditor = nullptr;
+
+	TSharedPtr<ISkeletonTree> SkeletonEditor;
+
+
+	/* Skeletal Mesh in the SkeletalViewer*/
+	TObjectPtr<USkeletalMesh> SkeletalMesh = nullptr;
+
+	TObjectPtr<UDataflowEdNode> SelectedNode = nullptr;
+
+	/* Rempping from the selected node to the SkeletalMesh*/
+	TArray<int32> CollectionIndexRemap;
+};
+
diff --git a/Engine/Source/Runtime/Experimental/Chaos/Private/GeometryCollection/Facades/CollectionTransformFacade.cpp b/Engine/Source/Runtime/Experimental/Chaos/Private/GeometryCollection/Facades/CollectionTransformFacade.cpp
--- a/Engine/Source/Runtime/Experimental/Chaos/Private/GeometryCollection/Facades/CollectionTransformFacade.cpp
+++ b/Engine/Source/Runtime/Experimental/Chaos/Private/GeometryCollection/Facades/CollectionTransformFacade.cpp
@@ -11,12 +11,14 @@ namespace GeometryCollection::Facades
 		: ParentAttribute(InCollection, FTransformCollection::ParentAttribute, FTransformCollection::TransformGroup)
 		, ChildrenAttribute(InCollection, FTransformCollection::ChildrenAttribute, FTransformCollection::TransformGroup)
 		, TransformAttribute(InCollection, FTransformCollection::TransformAttribute, FTransformCollection::TransformGroup)
+		, BoneNameAttribute(InCollection, "BoneName", FTransformCollection::TransformGroup)
 	{}
 
 	FCollectionTransformFacade::FCollectionTransformFacade(const FManagedArrayCollection& InCollection)
 		: ParentAttribute(InCollection, FTransformCollection::ParentAttribute, FTransformCollection::TransformGroup)
 		, ChildrenAttribute(InCollection, FTransformCollection::ChildrenAttribute, FTransformCollection::TransformGroup)
 		, TransformAttribute(InCollection, FTransformCollection::TransformAttribute, FTransformCollection::TransformGroup)
+		, BoneNameAttribute(InCollection, "BoneName", FTransformCollection::TransformGroup)
 	{}
 
 	bool FCollectionTransformFacade::IsValid() const
@@ -24,6 +26,26 @@ namespace GeometryCollection::Facades
 		return ParentAttribute.IsValid() && ChildrenAttribute.IsValid() && TransformAttribute.IsValid();
 	}
 
+	int32 FCollectionTransformFacade::Num() const
+	{
+		return ParentAttribute.Num();
+	}
+
+	bool FCollectionTransformFacade::HasBoneNameAttribute() const
+	{
+		return BoneNameAttribute.IsValid();
+	}
+
+	FString FCollectionTransformFacade::BoneName(int32 Index) const
+	{
+		if(BoneNameAttribute.IsValid() && BoneNameAttribute.IsValidIndex(Index))
+		{
+			return BoneNameAttribute[Index];
+		}
+		return FString();
+	}
+
+
 	TArray<int32> FCollectionTransformFacade::GetRootIndices() const
 	{
 		return Chaos::Facades::FCollectionHierarchyFacade::GetRootIndices(ParentAttribute);
diff --git a/Engine/Source/Runtime/Experimental/Chaos/Public/GeometryCollection/Facades/CollectionTransformFacade.h b/Engine/Source/Runtime/Experimental/Chaos/Public/GeometryCollection/Facades/CollectionTransformFacade.h
--- a/Engine/Source/Runtime/Experimental/Chaos/Public/GeometryCollection/Facades/CollectionTransformFacade.h
+++ b/Engine/Source/Runtime/Experimental/Chaos/Public/GeometryCollection/Facades/CollectionTransformFacade.h
@@ -26,6 +26,10 @@ namespace GeometryCollection::Facades
 		/** Is the facade defined constant. */
 		bool IsConst() const { return ParentAttribute.IsConst(); }
 
+		/** Get the number of transforms */
+		CHAOS_API int32 Num() const;
+
+
 		/** Gets the root index */
 		CHAOS_API TArray<int32> GetRootIndices() const;
 
@@ -44,6 +48,11 @@ namespace GeometryCollection::Facades
 		*/
 		const TManagedArray<FTransform3f>* FindTransforms() const { return TransformAttribute.Find(); }
 
+		/**
+		* Returns the bone names from the collection. Null if not initialized.
+		*/
+		const TManagedArray<FString>* FindBoneNames() const { return BoneNameAttribute.Find(); }
+
 		/**
 		* Returns array of transforms for transforming from bone space to collection space
 		* Vertex(inCollectionSpace) = TransformComputed.TransformPosition(Vertex(inBoneSpace))
@@ -65,6 +74,13 @@ namespace GeometryCollection::Facades
 		/** Transforms selected bones in the collection */
 		CHAOS_API void Transform(const FTransform& InTransform, const TArray<int32>& InSelection);
 
+		/** Check if the facade has the bone name attribute. */
+		CHAOS_API bool HasBoneNameAttribute() const;
+
+		/** Get a bone name from the index if the facade has the attribute defined. */
+		CHAOS_API FString BoneName(int32 Index) const;
+
+
 		/** Builds a FMatrix from all the components */
 		static CHAOS_API FMatrix BuildMatrix(const FVector& Translate,
 			const uint8 RotationOrder,
@@ -93,5 +109,6 @@ namespace GeometryCollection::Facades
 		TManagedArrayAccessor<int32>		ParentAttribute;
 		TManagedArrayAccessor<TSet<int32>>	ChildrenAttribute;
 		TManagedArrayAccessor<FTransform3f>	TransformAttribute;
+		TManagedArrayAccessor<FString>	    BoneNameAttribute;
 	};
 }
diff --git a/Engine/Source/Runtime/Experimental/GeometryCollectionEngine/Private/GeometryCollection/GeometryCollectionEngineConversion.cpp b/Engine/Source/Runtime/Experimental/GeometryCollectionEngine/Private/GeometryCollection/GeometryCollectionEngineConversion.cpp
--- a/Engine/Source/Runtime/Experimental/GeometryCollectionEngine/Private/GeometryCollection/GeometryCollectionEngineConversion.cpp
+++ b/Engine/Source/Runtime/Experimental/GeometryCollectionEngine/Private/GeometryCollection/GeometryCollectionEngineConversion.cpp
@@ -33,6 +33,7 @@
 #include "Physics/Experimental/ChaosInterfaceUtils.h"
 #include "PhysicsEngine/BodySetup.h"
 #include "Rendering/SkeletalMeshRenderData.h"
+#include "ReferenceSkeleton.h"
 #include "SkeletalMeshAttributes.h"
 #include "StaticMeshAttributes.h"
 #include "StaticMeshOperations.h"
@@ -1666,4 +1667,70 @@ void FGeometryCollectionEngineConversion::ConvertActorToGeometryCollection(const
 #endif //WITH_EDITORONLY_DATA
 }
 
+
+void FGeometryCollectionEngineConversion::ConvertCollectionToSkeleton(const FManagedArrayCollection& InCollection, USkeleton* OutSkeleton, TArray<int32>& OutIndexRemap)
+{
+	GeometryCollection::Facades::FCollectionTransformFacade Transforms(InCollection);
+	if (Transforms.IsValid() && Transforms.HasBoneNameAttribute() && OutSkeleton)
+	{
+		OutIndexRemap.Init(INDEX_NONE, Transforms.Num());
+		auto AddMapping = [&OutIndexRemap](int32 A, int32 B)
+		{
+			OutIndexRemap[A] = B;
+		};
+
+		FReferenceSkeletonModifier Edit(OutSkeleton);
+		auto AddChildren = [&Transforms, &Edit, &AddMapping](const TArray<int32>& CollectionChildren)
+		{
+			TQueue<int32> Children;
+			auto Enqueue = [&Children](const TArray<int32>& List)
+			{
+				for (int32 Elem : List)
+					Children.Enqueue(Elem);
+			};
+
+			int CurrentIndex = INDEX_NONE;
+			Enqueue(CollectionChildren);
+			while (!Children.IsEmpty())
+			{
+				Children.Dequeue(CurrentIndex);
+
+				int32 CollectionParentIndex = (*Transforms.GetParents())[CurrentIndex];
+				int32 SkeletionParentIndex = Edit.FindBoneIndex(FName((*Transforms.FindBoneNames())[CollectionParentIndex]));
+
+				FName BoneName = FName((*Transforms.FindBoneNames())[CurrentIndex]);
+				FTransform Transform = FTransform((*Transforms.FindTransforms())[CurrentIndex]);
+				FMeshBoneInfo Info(BoneName, BoneName.ToString(), SkeletionParentIndex);
+				Edit.Add(Info, Transform, true /*bAllowMultipleRoots*/);
+				AddMapping(CurrentIndex, Edit.GetReferenceSkeleton().GetRawBoneNum());
+
+				Enqueue((*Transforms.FindChildren())[CurrentIndex].Array());
+			}
+		};
+
+
+		// must insert in decending order from parent to child. 
+		TArray<FString> BoneNameStrings = Transforms.FindBoneNames()->GetConstArray();
+		for (int i = 0; i < Transforms.GetParents()->Num(); i++)
+		{
+			if ((*Transforms.GetParents())[i] == INDEX_NONE) // No parent
+			{
+				FName BoneName = FName((*Transforms.FindBoneNames())[i]);
+				FTransform Transform = FTransform((*Transforms.FindTransforms())[i]);
+				FMeshBoneInfo Info(BoneName, BoneName.ToString(), INDEX_NONE);
+				Edit.Add(Info, Transform, true /*bAllowMultipleRoots*/);
+				AddMapping(i,Edit.GetReferenceSkeleton().GetNum());
+
+				if ((*Transforms.FindChildren())[i].Num())
+				{
+					AddChildren((*Transforms.FindChildren())[i].Array());
+				}
+			}
+		}
+
+		UE_LOG(UGeometryCollectionConversionLogging, Log, TEXT("FGeometryCollectionEngineConversion::ConvertCollectionToSkeleton(NumTransforms:%d)"), OutSkeleton->GetReferenceSkeleton().GetRawBoneNum());
+	}
+}
+
+
 #undef LOCTEXT_NAMESPACE 
diff --git a/Engine/Source/Runtime/Experimental/GeometryCollectionEngine/Public/GeometryCollection/GeometryCollectionEngineConversion.h b/Engine/Source/Runtime/Experimental/GeometryCollectionEngine/Public/GeometryCollection/GeometryCollectionEngineConversion.h
--- a/Engine/Source/Runtime/Experimental/GeometryCollectionEngine/Public/GeometryCollection/GeometryCollectionEngineConversion.h
+++ b/Engine/Source/Runtime/Experimental/GeometryCollectionEngine/Public/GeometryCollection/GeometryCollectionEngineConversion.h
@@ -210,4 +210,13 @@ public:
 	*/
 	static GEOMETRYCOLLECTIONENGINE_API void ConvertActorToGeometryCollection(const AActor* Actor, FManagedArrayCollection& OutCollection, TArray<TObjectPtr<UMaterial>>& OutMaterials, TArray<FGeometryCollectionAutoInstanceMesh>& OutInstancedMeshes, const FSkeletalMeshToCollectionConversionParameters& ConversionParameters, bool bSplitComponents = false);
 
+	/**
+	*  Converts an FManagedArrayCollection to a USkeleton
+	*  @param InManagedArrayCollection : Collection that supports transform hierarchy
+	*  @param OutSkeleton : USkeleton to populate
+	*  @param OutIndexRemap : Remapping from Collection transforms to skeletal transforms.
+	*/
+	static GEOMETRYCOLLECTIONENGINE_API void ConvertCollectionToSkeleton(const FManagedArrayCollection& InCollection, USkeleton* OutSkeleton, TArray<int32>& OutIndexRemap);
+
+
 };
