DataflowEditor : Add referenced objects to editor elements.
#preflght 6631320b6608538f460e26a2

diff --git a/Engine/Plugins/ChaosClothAssetEditor/Source/ChaosClothAssetEditor/Private/ChaosClothAsset/ClothEditorToolkit.cpp b/Engine/Plugins/ChaosClothAssetEditor/Source/ChaosClothAssetEditor/Private/ChaosClothAsset/ClothEditorToolkit.cpp
--- a/Engine/Plugins/ChaosClothAssetEditor/Source/ChaosClothAssetEditor/Private/ChaosClothAsset/ClothEditorToolkit.cpp
+++ b/Engine/Plugins/ChaosClothAssetEditor/Source/ChaosClothAssetEditor/Private/ChaosClothAsset/ClothEditorToolkit.cpp
@@ -1094,7 +1094,13 @@ TSharedPtr<IStructureDetailsView> FChaosClothAssetEditorToolkit::CreateNodeDetai
 
 void FChaosClothAssetEditorToolkit::OnPropertyValueChanged(const FPropertyChangedEvent& PropertyChangedEvent)
 {
-	FDataflowEditorCommands::OnPropertyValueChanged(GetDataflow(), DataflowContext, LastDataflowNodeTimestamp, PropertyChangedEvent, GraphEditor ? GraphEditor->GetSelectedNodes() : FGraphPanelSelectionSet());
+	auto AsObjectPointers = [](const TSet<UObject*>& Set) {
+		TSet<TObjectPtr<UObject> > Objs; for (UObject* Elem : Set) Objs.Add(Elem);
+		return Objs;
+	};
+
+	FDataflowEditorCommands::OnPropertyValueChanged(GetDataflow(), DataflowContext, LastDataflowNodeTimestamp, PropertyChangedEvent,
+		GraphEditor ? AsObjectPointers(GraphEditor->GetSelectedNodes()) : TSet<TObjectPtr<UObject>>());
 }
 
 bool FChaosClothAssetEditorToolkit::OnNodeVerifyTitleCommit(const FText& NewText, UEdGraphNode* GraphNode, FText& OutErrorMessage) const
diff --git a/Engine/Plugins/Experimental/Dataflow/Source/DataflowEditor/Private/Dataflow/DataflowEditorCommands.cpp b/Engine/Plugins/Experimental/Dataflow/Source/DataflowEditor/Private/Dataflow/DataflowEditorCommands.cpp
--- a/Engine/Plugins/Experimental/Dataflow/Source/DataflowEditor/Private/Dataflow/DataflowEditorCommands.cpp
+++ b/Engine/Plugins/Experimental/Dataflow/Source/DataflowEditor/Private/Dataflow/DataflowEditorCommands.cpp
@@ -337,7 +337,7 @@ void FDataflowEditorCommands::OnAssetPropertyValueChanged(TObjectPtr<UDataflowBa
 	}
 }
 
-void FDataflowEditorCommands::OnPropertyValueChanged(UDataflow* OutDataflow, TSharedPtr<Dataflow::FEngineContext>& Context, Dataflow::FTimestamp& OutLastNodeTimestamp, const FPropertyChangedEvent& InPropertyChangedEvent, const TSet<UObject*>& SelectedNodes)
+void FDataflowEditorCommands::OnPropertyValueChanged(UDataflow* OutDataflow, TSharedPtr<Dataflow::FEngineContext>& Context, Dataflow::FTimestamp& OutLastNodeTimestamp, const FPropertyChangedEvent& InPropertyChangedEvent, const TSet<TObjectPtr<UObject> >& SelectedNodes)
 {
 	if (InPropertyChangedEvent.ChangeType == EPropertyChangeType::ValueSet)
 	{
@@ -404,7 +404,7 @@ void FDataflowEditorCommands::DeleteNodes(UDataflow* Graph, const FGraphPanelSel
 	}
 }
 
-void FDataflowEditorCommands::OnSelectedNodesChanged(TSharedPtr<IStructureDetailsView> PropertiesEditor, UObject* Asset, UDataflow* Graph, const TSet<UObject*>& NewSelection)
+void FDataflowEditorCommands::OnSelectedNodesChanged(TSharedPtr<IStructureDetailsView> PropertiesEditor, UObject* Asset, UDataflow* Graph, const TSet<TObjectPtr<UObject> >& NewSelection)
 {
 	PropertiesEditor->SetStructureData(nullptr);
 
@@ -412,7 +412,11 @@ void FDataflowEditorCommands::OnSelectedNodesChanged(TSharedPtr<IStructureDetail
 	{
 		if (const TSharedPtr<Dataflow::FGraph> DataflowGraph = Graph->GetDataflow())
 		{
-			FGraphPanelSelectionSet SelectedNodes = NewSelection;
+			auto AsRawPointers = [](const TSet<TObjectPtr<UObject> >& NewSelection) {
+				TSet<UObject*> Raw; for (UObject* Elem : NewSelection) Raw.Add(Elem);
+				return Raw;
+			};
+			FGraphPanelSelectionSet SelectedNodes = AsRawPointers(NewSelection);
 			if (SelectedNodes.Num())
 			{
 				TArray<UObject*> Objects;
diff --git a/Engine/Plugins/Experimental/Dataflow/Source/DataflowEditor/Private/Dataflow/DataflowEditorToolkit.cpp b/Engine/Plugins/Experimental/Dataflow/Source/DataflowEditor/Private/Dataflow/DataflowEditorToolkit.cpp
--- a/Engine/Plugins/Experimental/Dataflow/Source/DataflowEditor/Private/Dataflow/DataflowEditorToolkit.cpp
+++ b/Engine/Plugins/Experimental/Dataflow/Source/DataflowEditor/Private/Dataflow/DataflowEditorToolkit.cpp
@@ -535,8 +535,8 @@ void FDataflowEditorToolkit::OnNodeTitleCommitted(const FText& InNewText, ETextC
 
 void FDataflowEditorToolkit::OnNodeSelectionChanged(const TSet<UObject*>& InNewSelection)
 {
-	auto FindDataflowNodesInSet = [](const TSet<UObject*>& InSet) {
-		TSet<UObject*> Results;
+	auto FindDataflowNodesInSet = [](const TSet<TObjectPtr<UObject>>& InSet) {
+		TSet<TObjectPtr<UObject>> Results;
 		for (UObject* Item : InSet)
 		{
 			if (Cast<UDataflowEdNode>(Item))
@@ -619,22 +619,27 @@ void FDataflowEditorToolkit::OnNodeSelectionChanged(const TSet<UObject*>& InNewS
 
 	if (TObjectPtr<UDataflowBaseContent> EditorContent = GetDataflowContent(); EditorContent->GetDataflowAsset())
 	{
+		auto AsObjectPointers = [](const TSet<UObject*>& Set) {
+			TSet<TObjectPtr<UObject> > Objs; for (UObject* Elem : Set) Objs.Add(Elem);
+			return Objs;
+		};
+
 		// Only keep UDataflowEdNode from NewSelection
-		TSet<UObject*> NodeSelection = FindDataflowNodesInSet(InNewSelection);
+		TSet< TObjectPtr<UObject> > NodeSelection = FindDataflowNodesInSet(AsObjectPointers(InNewSelection));
 
 		if (!NodeSelection.Num())
 		{
 			// The selection is empty. 
 			ResetListeners();
-			SelectedDataflowNodes = TSet<UObject*>();
+			SelectedDataflowNodes = TSet<TObjectPtr<UObject>>();
 			if (PrimarySelection) bPrimarySelectionChanged = true;
 			PrimarySelection = nullptr;
 		}
 		else
 		{
-			TSet<UObject*> DeselectedNodes = SelectedDataflowNodes.Difference(NodeSelection);
-			TSet<UObject*> StillSelectedNodes = SelectedDataflowNodes.Intersect(NodeSelection);
-			TSet<UObject*> NewlySelectedNodes = NodeSelection.Difference(SelectedDataflowNodes);
+			TSet<TObjectPtr<UObject>> DeselectedNodes = SelectedDataflowNodes.Difference(NodeSelection);
+			TSet<TObjectPtr<UObject>> StillSelectedNodes = SelectedDataflowNodes.Intersect(NodeSelection);
+			TSet<TObjectPtr<UObject>> NewlySelectedNodes = NodeSelection.Difference(SelectedDataflowNodes);
 
 			// Something has been removed
 			if (DeselectedNodes.Num())
@@ -1287,4 +1292,11 @@ FLinearColor FDataflowEditorToolkit::GetWorldCentricTabColorScale() const
 	return FLinearColor(0.3f, 0.2f, 0.5f, 0.5f);
 }
 
+void FDataflowEditorToolkit::AddReferencedObjects(FReferenceCollector& Collector)
+{
+	Collector.AddReferencedObjects(SelectedDataflowNodes);
+	Collector.AddReferencedObject(PrimarySelection);
+}
+
+
 #undef LOCTEXT_NAMESPACE
diff --git a/Engine/Plugins/Experimental/Dataflow/Source/DataflowEditor/Private/Dataflow/DataflowGraphEditor.cpp b/Engine/Plugins/Experimental/Dataflow/Source/DataflowEditor/Private/Dataflow/DataflowGraphEditor.cpp
--- a/Engine/Plugins/Experimental/Dataflow/Source/DataflowEditor/Private/Dataflow/DataflowGraphEditor.cpp
+++ b/Engine/Plugins/Experimental/Dataflow/Source/DataflowEditor/Private/Dataflow/DataflowGraphEditor.cpp
@@ -198,7 +198,12 @@ void SDataflowGraphEditor::OnSelectedNodesChanged(const TSet<UObject*>& NewSelec
 
 	if (DataflowAsset.Get() && DetailsView)
 	{
-		FDataflowEditorCommands::OnSelectedNodesChanged(DetailsView, AssetOwner.Get(), DataflowAsset.Get(), NewSelection);
+		auto AsObjectPointers = [](const TSet<UObject*>& Set) {
+			TSet<TObjectPtr<UObject> > Objs; for (UObject* Elem : Set) Objs.Add(Elem);
+			return Objs;
+		};
+
+		FDataflowEditorCommands::OnSelectedNodesChanged(DetailsView, AssetOwner.Get(), DataflowAsset.Get(), AsObjectPointers(NewSelection) );
 	}
 }
 
diff --git a/Engine/Plugins/Experimental/Dataflow/Source/DataflowEditor/Private/DataflowEditorTools/DataflowEditorWeightMapPaintTool.cpp b/Engine/Plugins/Experimental/Dataflow/Source/DataflowEditor/Private/DataflowEditorTools/DataflowEditorWeightMapPaintTool.cpp
--- a/Engine/Plugins/Experimental/Dataflow/Source/DataflowEditor/Private/DataflowEditorTools/DataflowEditorWeightMapPaintTool.cpp
+++ b/Engine/Plugins/Experimental/Dataflow/Source/DataflowEditor/Private/DataflowEditorTools/DataflowEditorWeightMapPaintTool.cpp
@@ -2055,5 +2055,14 @@ void UDataflowEditorWeightMapPaintTool::UpdateVertexColorOverlay(const TSet<int>
 	}
 }
 
+void UDataflowEditorWeightMapPaintTool::AddReferencedObjects(UObject* InThis, FReferenceCollector& Collector)
+{
+	UDataflowEditorWeightMapPaintTool* This = CastChecked<UDataflowEditorWeightMapPaintTool>(InThis);
+	Collector.AddReferencedObject(This->PreviewMeshActor);
+	Collector.AddReferencedObject(This->DynamicMeshComponent);
+	Collector.AddReferencedObject(This->MeshElementsDisplay);
+	Collector.AddReferencedObject(This->DataflowEditorContextObject);
+	Super::AddReferencedObjects(InThis, Collector);
+}
 
 #undef LOCTEXT_NAMESPACE
diff --git a/Engine/Plugins/Experimental/Dataflow/Source/DataflowEditor/Public/Dataflow/DataflowEditorCommands.h b/Engine/Plugins/Experimental/Dataflow/Source/DataflowEditor/Public/Dataflow/DataflowEditorCommands.h
--- a/Engine/Plugins/Experimental/Dataflow/Source/DataflowEditor/Public/Dataflow/DataflowEditorCommands.h
+++ b/Engine/Plugins/Experimental/Dataflow/Source/DataflowEditor/Public/Dataflow/DataflowEditorCommands.h
@@ -116,13 +116,13 @@ public:
 	/*
 	*  OnPropertyValueChanged
 	*/
-	static void OnPropertyValueChanged(UDataflow* Graph, TSharedPtr<Dataflow::FEngineContext>& Context, Dataflow::FTimestamp& OutLastNodeTimestamp, const FPropertyChangedEvent& PropertyChangedEvent, const TSet<UObject*>& SelectedNodes = TSet<UObject*>());
+	static void OnPropertyValueChanged(UDataflow* Graph, TSharedPtr<Dataflow::FEngineContext>& Context, Dataflow::FTimestamp& OutLastNodeTimestamp, const FPropertyChangedEvent& PropertyChangedEvent, const TSet<TObjectPtr<UObject>>& NewSelection = TSet<TObjectPtr<UObject>>());
 	static void OnAssetPropertyValueChanged(TObjectPtr<UDataflowBaseContent> Content, const FPropertyChangedEvent& PropertyChangedEvent);
 
 	/*
 	*  OnSelectedNodesChanged
 	*/
-	static void OnSelectedNodesChanged(TSharedPtr<IStructureDetailsView> PropertiesEditor, UObject* Asset, UDataflow* Graph, const TSet<UObject*>& NewSelection);
+	static void OnSelectedNodesChanged(TSharedPtr<IStructureDetailsView> PropertiesEditor, UObject* Asset, UDataflow* Graph, const TSet<TObjectPtr<UObject>>& NewSelection);
 
 	/*
 	*  ToggleEnabledState
diff --git a/Engine/Plugins/Experimental/Dataflow/Source/DataflowEditor/Public/Dataflow/DataflowEditorToolkit.h b/Engine/Plugins/Experimental/Dataflow/Source/DataflowEditor/Public/Dataflow/DataflowEditorToolkit.h
--- a/Engine/Plugins/Experimental/Dataflow/Source/DataflowEditor/Public/Dataflow/DataflowEditorToolkit.h
+++ b/Engine/Plugins/Experimental/Dataflow/Source/DataflowEditor/Public/Dataflow/DataflowEditorToolkit.h
@@ -12,6 +12,7 @@
 #include "Dataflow/DataflowCollectionSpreadSheet.h"
 #include "Dataflow/DataflowConstructionViewport.h"
 #include "Dataflow/DataflowSimulationViewport.h"
+//#include "UObject\GCObject.h"
 
 class FEditorViewportTabContent;
 class IDetailsView;
@@ -42,7 +43,7 @@ namespace Dataflow
 	};
 }
 
-class DATAFLOWEDITOR_API FDataflowEditorToolkit final : public FBaseCharacterFXEditorToolkit, public FTickableEditorObject, public FNotifyHook
+class DATAFLOWEDITOR_API FDataflowEditorToolkit final : public FBaseCharacterFXEditorToolkit, public FTickableEditorObject, public FNotifyHook, public FGCObject
 {
 	using FBaseCharacterFXEditorToolkit::ObjectScene;
 
@@ -74,6 +75,14 @@ public:
 	/** Dataflow preview scenes accessor */
 	const TSharedPtr<FDataflowSimulationScene>& GetSimulationScene() const {return SimulationScene;}
 
+
+	// FSerializableObject interface
+	virtual void AddReferencedObjects(FReferenceCollector& Collector) override;
+	virtual FString GetReferencerName() const override
+	{
+		return TEXT("FDataflowEditorToolkit");
+	}
+	// End of FSerializableObject interface
 protected:
 
 	UDataflowEditor* DataflowEditor = nullptr;
@@ -174,7 +183,7 @@ private:
 
 	// The currently selected set of dataflow nodes. 
 	UPROPERTY()
-	TSet<UObject*> SelectedDataflowNodes;
+	TSet< TObjectPtr<UObject> > SelectedDataflowNodes;
 
 	// The most recently selected dataflow node.
 	UPROPERTY()
diff --git a/Engine/Plugins/Experimental/Dataflow/Source/DataflowEditor/Public/DataflowEditorTools/DataflowEditorWeightMapPaintTool.h b/Engine/Plugins/Experimental/Dataflow/Source/DataflowEditor/Public/DataflowEditorTools/DataflowEditorWeightMapPaintTool.h
--- a/Engine/Plugins/Experimental/Dataflow/Source/DataflowEditor/Public/DataflowEditorTools/DataflowEditorWeightMapPaintTool.h
+++ b/Engine/Plugins/Experimental/Dataflow/Source/DataflowEditor/Public/DataflowEditorTools/DataflowEditorWeightMapPaintTool.h
@@ -290,6 +290,10 @@ public:
 	void SetDataflowEditorContextObject(TObjectPtr<UDataflowContextObject> InDataflowEditorContextObject);
 	void SetEditorMode(UDataflowEditorMode* InMode) { Mode = InMode; }
 
+	//~ UObject interface
+
+	static void AddReferencedObjects(UObject* InThis, FReferenceCollector& Collector);
+
 public:
 
 
@@ -402,10 +406,10 @@ protected:
 	TObjectPtr<AInternalToolFrameworkActor> PreviewMeshActor = nullptr;
 
 	UPROPERTY()
-	TObjectPtr<UDynamicMeshComponent> DynamicMeshComponent;
+	TObjectPtr<UDynamicMeshComponent> DynamicMeshComponent = nullptr;
 
 	UPROPERTY()
-	TObjectPtr<UMeshElementsVisualizer> MeshElementsDisplay;
+	TObjectPtr<UMeshElementsVisualizer> MeshElementsDisplay = nullptr;
 
 	UPROPERTY()
 	TObjectPtr<UDataflowContextObject> DataflowEditorContextObject = nullptr;
diff --git a/Engine/Source/Runtime/Experimental/Dataflow/Engine/Private/Dataflow/DataflowContent.cpp b/Engine/Source/Runtime/Experimental/Dataflow/Engine/Private/Dataflow/DataflowContent.cpp
--- a/Engine/Source/Runtime/Experimental/Dataflow/Engine/Private/Dataflow/DataflowContent.cpp
+++ b/Engine/Source/Runtime/Experimental/Dataflow/Engine/Private/Dataflow/DataflowContent.cpp
@@ -13,6 +13,12 @@
 
 #include UE_INLINE_GENERATED_CPP_BY_NAME(DataflowContent)
 
+void UDataflowContextObject::AddReferencedObjects(UObject* InThis, FReferenceCollector& Collector)
+{
+	UDataflowContextObject* This = CastChecked<UDataflowContextObject>(InThis);
+	Collector.AddReferencedObject(This->PrimarySelectedNode);
+	Super::AddReferencedObjects(InThis, Collector);
+}
 
 //
 // UDataflowBaseContent
@@ -90,6 +96,14 @@ void UDataflowBaseContent::Serialize(FArchive& Ar)
 	DataflowContext->Serialize(Ar);
 }
 
+void UDataflowBaseContent::AddReferencedObjects(UObject* InThis, FReferenceCollector& Collector)
+{
+	UDataflowBaseContent* This = CastChecked<UDataflowBaseContent>(InThis);
+	Collector.AddReferencedObject(This->DataflowOwner);
+	Collector.AddReferencedObject(This->DataflowAsset);
+	Super::AddReferencedObjects(InThis, Collector);
+}
+
 //
 // UDataflowSkeletalContent
 //
@@ -234,3 +248,13 @@ FVector2f UDataflowSkeletalContent::GetSimulationRange() const
 	return FVector2f(0.0f, 0.0f);
 }
 
+void UDataflowSkeletalContent::AddReferencedObjects(UObject* InThis, FReferenceCollector& Collector)
+{
+	UDataflowSkeletalContent* This = CastChecked<UDataflowSkeletalContent>(InThis);
+	Collector.AddReferencedObject(This->SkeletalMesh);
+	Collector.AddReferencedObject(This->AnimationAsset);
+	Collector.AddReferencedObject(This->Skeleton);
+	Collector.AddReferencedObject(This->SkeletalMeshComponent);
+	Collector.AddReferencedObject(This->AnimationNodeInstance);
+	Super::AddReferencedObjects(InThis, Collector);
+}
diff --git a/Engine/Source/Runtime/Experimental/Dataflow/Engine/Public/Dataflow/DataflowContent.h b/Engine/Source/Runtime/Experimental/Dataflow/Engine/Public/Dataflow/DataflowContent.h
--- a/Engine/Source/Runtime/Experimental/Dataflow/Engine/Public/Dataflow/DataflowContent.h
+++ b/Engine/Source/Runtime/Experimental/Dataflow/Engine/Public/Dataflow/DataflowContent.h
@@ -62,6 +62,9 @@ public:
 		return nullptr;
 	}
 
+	//~ UObject interface
+	static void AddReferencedObjects(UObject* InThis, FReferenceCollector& Collector);
+
 protected:
 
 	/** Render collection to be used */
@@ -136,6 +139,9 @@ public:
 
 	virtual void Serialize(FArchive& Ar);
 
+	//~ UObject interface
+	static void AddReferencedObjects(UObject* InThis, FReferenceCollector& Collector);
+
 protected:
 
 	/** Data flow asset that we will edit */
@@ -208,6 +214,9 @@ public:
 	void SetAnimationAsset(const TObjectPtr<UAnimationAsset>& InAnimation);
 	const TObjectPtr<UAnimationAsset>& GetAnimationAsset() const { return AnimationAsset; }
 
+	//~ UObject interface
+	static void AddReferencedObjects(UObject* InThis, FReferenceCollector& Collector);
+
 protected:
 
 	/** Data flow skeletal mesh*/
@@ -216,7 +225,7 @@ protected:
 
 	/** Animation asset to be used to preview simulation */
 	UPROPERTY(EditAnywhere, Category = "Preview", Transient, SkipSerialization)
-	TObjectPtr<UAnimationAsset> AnimationAsset;
+	TObjectPtr<UAnimationAsset> AnimationAsset = nullptr;
 
 	/** Data flow skeleton */
 	UPROPERTY(EditAnywhere, Category = "Skeleton", Transient, SkipSerialization)
@@ -224,9 +233,9 @@ protected:
 
 	/** Skeletal mesh component used in the preview scene */
 	UPROPERTY(Transient, SkipSerialization)
-	TObjectPtr<USkeletalMeshComponent> SkeletalMeshComponent;
+	TObjectPtr<USkeletalMeshComponent> SkeletalMeshComponent = nullptr;
 
 	/** Anim node instance used with skelmesh component */
 	UPROPERTY(Transient, SkipSerialization)
-	TObjectPtr<UAnimSingleNodeInstance> AnimationNodeInstance;
+	TObjectPtr<UAnimSingleNodeInstance> AnimationNodeInstance = nullptr;
 };
