Dataflow : Skeleton View
- Added selection notification through the NodeView.
- Disabled the skeleton view with a cvar.
- Added selection feedback when construciton view selection changes.

diff --git a/Engine/Plugins/Experimental/Dataflow/Source/DataflowEditor/Private/Dataflow/DataflowConstructionViewportClient.cpp b/Engine/Plugins/Experimental/Dataflow/Source/DataflowEditor/Private/Dataflow/DataflowConstructionViewportClient.cpp
--- a/Engine/Plugins/Experimental/Dataflow/Source/DataflowEditor/Private/Dataflow/DataflowConstructionViewportClient.cpp
+++ b/Engine/Plugins/Experimental/Dataflow/Source/DataflowEditor/Private/Dataflow/DataflowConstructionViewportClient.cpp
@@ -130,7 +130,6 @@ void FDataflowConstructionViewportClient::ProcessClick(FSceneView& View, HHitPro
 		{
 			Component->PushSelectionToProxy();
 		}
-
 	};
 
 	auto SelectSingleNodeInGraph = [&](TObjectPtr<const UDataflowEdNode> Node)
@@ -166,6 +165,7 @@ void FDataflowConstructionViewportClient::ProcessClick(FSceneView& View, HHitPro
 		}
 	};
 
+	TArray<UPrimitiveComponent*> CurrentlySelectedComponents;
 	if (!IsInteractiveToolActive())
 	{
 		if (USelection* SelectedComponents = ModeTools->GetSelectedComponents())
@@ -182,8 +182,11 @@ void FDataflowConstructionViewportClient::ProcessClick(FSceneView& View, HHitPro
 			}
 
 			EnableToolForSelectedNode(SelectedComponents);
+
+			SelectedComponents->GetSelectedObjects<UPrimitiveComponent>(CurrentlySelectedComponents);
 		}
 	}
+	OnSelectionChangedMulticast.Broadcast(CurrentlySelectedComponents);
 }
 
 void FDataflowConstructionViewportClient::SetConstructionViewMode(Dataflow::EDataflowPatternVertexType InViewMode)
diff --git a/Engine/Plugins/Experimental/Dataflow/Source/DataflowEditor/Private/Dataflow/DataflowEditorToolkit.cpp b/Engine/Plugins/Experimental/Dataflow/Source/DataflowEditor/Private/Dataflow/DataflowEditorToolkit.cpp
--- a/Engine/Plugins/Experimental/Dataflow/Source/DataflowEditor/Private/Dataflow/DataflowEditorToolkit.cpp
+++ b/Engine/Plugins/Experimental/Dataflow/Source/DataflowEditor/Private/Dataflow/DataflowEditorToolkit.cpp
@@ -40,7 +40,6 @@
 #include "GameFramework/Actor.h"
 #include "GraphEditorActions.h"
 #include "ISkeletonTree.h"
-#include "ISkeletonEditorModule.h"
 #include "IStructureDetailsView.h"
 #include "Kismet/GameplayStatics.h"
 #include "Misc/MessageDialog.h"
@@ -54,6 +53,10 @@
 
 //DEFINE_LOG_CATEGORY_STATIC(EditorToolkitLog, Log, All);
 
+bool bDataflowEnableSkeletonView = false;
+FAutoConsoleVariableRef CVARDataflowEnableSkeletonView(TEXT("p.Dataflow.Editor.EnableSkeletonView"), bDataflowEnableSkeletonView,
+	TEXT("Deprecated Tool! Allows the Dataflow editor to create a skeleton view that reflects the hierarchy and selection state of the construction viewport.[def:false]"));
+
 const FName FDataflowEditorToolkit::GraphCanvasTabId(TEXT("DataflowEditor_GraphCanvas"));
 const FName FDataflowEditorToolkit::NodeDetailsTabId(TEXT("DataflowEditor_NodeDetails"));
 const FName FDataflowEditorToolkit::SkeletonViewTabId(TEXT("DataflowEditor_SkeletonView"));
@@ -351,6 +354,9 @@ void FDataflowEditorToolkit::PostInitAssetEditor()
 	DataflowMode->SetConstructionViewportClient(StaticCastWeakPtr<FDataflowConstructionViewportClient>(WeakConstructionViewportClient));
 	const TWeakPtr<FViewportClient> WeakSimulationViewportClient(SimulationViewportClient);
 	DataflowMode->SetSimulationViewportClient(StaticCastWeakPtr<FDataflowSimulationViewportClient>(WeakSimulationViewportClient));
+
+	FDataflowConstructionViewportClient* ConstructionViewportClient = static_cast<FDataflowConstructionViewportClient*>(ViewportClient.Get());
+	OnConstructionSelectionChangedDelegateHandle = ConstructionViewportClient->OnSelectionChangedMulticast.AddSP(this, &FDataflowEditorToolkit::OnConstructionViewSelectionChanged);
 }
 
 void FDataflowEditorToolkit::InitializeEdMode(UBaseCharacterFXEditorMode* EdMode)
@@ -723,6 +729,15 @@ void FDataflowEditorToolkit::OnNodeDeleted(const TSet<UObject*>& NewSelection)
 	}
 }
 
+void FDataflowEditorToolkit::OnConstructionViewSelectionChanged(const TArray<UPrimitiveComponent*>& SelectedComponents)
+{
+	for(IDataflowViewListener* Listener : ViewListeners)
+	{
+		Listener->OnConstructionViewSelectionChanged(SelectedComponents);
+	}
+}
+
+
 void FDataflowEditorToolkit::Tick(float DeltaTime)
 {
 	if (TObjectPtr<UDataflowBaseContent> EditorContent = GetDataflowContent())
@@ -896,27 +911,20 @@ TSharedRef<SDockTab> FDataflowEditorToolkit::SpawnTab_SkeletonView(const FSpawnT
 {
 	check(Args.GetTabId() == SkeletonViewTabId);
 
-
 	SkeletonEditorView = MakeShared<FDataflowSkeletonView>(DataflowEditor);
-
-	if (const TObjectPtr<UDataflowSkeletalContent> SkeletalContent = Cast<UDataflowSkeletalContent>(GetDataflowContent()))
-	{
-		if (SkeletalContent->GetDataflowAsset())
-		{
-			SkeletonEditorView->SetSkeleton(SkeletalContent->GetSkeleton());
-		}
-	}
 	ViewListeners.Add(SkeletonEditorView.Get());
 
 	FSkeletonTreeArgs SkeletonTreeArgs;
-	//TSharedRef<FDataflowSkeletonView> DataflowSkeletonViewRef = SkeletonEditorView.ToSharedRef();
+	SkeletonTreeArgs.bShowBlendProfiles = false;
+	SkeletonTreeArgs.bShowFilterMenu = true;
+	SkeletonTreeArgs.bShowDebugVisualizationOptions = false;
+	SkeletonTreeArgs.bAllowMeshOperations = false;
+	SkeletonTreeArgs.bAllowSkeletonOperations = false;
+	SkeletonTreeArgs.bHideBonesByDefault = false;
 	SkeletonTreeArgs.OnSelectionChanged = FOnSkeletonTreeSelectionChanged::CreateSP(SkeletonEditorView.ToSharedRef(), &FDataflowSkeletonView::SkeletonViewSelectionChanged);
 	SkeletonTreeArgs.ContextName = GetToolkitFName();
 
-	ISkeletonEditorModule& SkeletonEditorModule = FModuleManager::LoadModuleChecked<ISkeletonEditorModule>("SkeletonEditor");
-	TSharedPtr<ISkeletonTree> SkeletonEditor = SkeletonEditorModule.CreateSkeletonTree(SkeletonEditorView->GetSkeleton(), SkeletonTreeArgs);
-	SkeletonEditorView->SetSkeletonEditor(SkeletonEditor);
-
+	TSharedPtr<ISkeletonTree> SkeletonEditor = SkeletonEditorView->CreateEditor(SkeletonTreeArgs);
 	return SNew(SDockTab)
 		.Label(LOCTEXT("DataflowEditor_Skeleton_TabTitle", "Skeleton View"))
 		[
@@ -1147,10 +1155,13 @@ void FDataflowEditorToolkit::RegisterTabSpawners(const TSharedRef<FTabManager>&
 		.SetGroup(EditorMenuCategory.ToSharedRef())
 		.SetIcon(FSlateIcon(FAppStyle::GetAppStyleSetName(), "LevelEditor.Tabs.Details"));
 
-	InTabManager->RegisterTabSpawner(SkeletonViewTabId, FOnSpawnTab::CreateSP(this, &FDataflowEditorToolkit::SpawnTab_SkeletonView))
-		.SetDisplayName(LOCTEXT("DataflowSkeletonTab", "Skeleton View"))
-		.SetGroup(EditorMenuCategory.ToSharedRef())
-		.SetIcon(FSlateIcon(FAppStyle::GetAppStyleSetName(), "LevelEditor.Tabs.SkeletonHierarchy"));
+	if (bDataflowEnableSkeletonView)
+	{
+		InTabManager->RegisterTabSpawner(SkeletonViewTabId, FOnSpawnTab::CreateSP(this, &FDataflowEditorToolkit::SpawnTab_SkeletonView))
+			.SetDisplayName(LOCTEXT("DataflowSkeletonTab", "Skeleton View"))
+			.SetGroup(EditorMenuCategory.ToSharedRef())
+			.SetIcon(FSlateIcon(FAppStyle::GetAppStyleSetName(), "LevelEditor.Tabs.SkeletonHierarchy"));
+	}
 
 	InTabManager->RegisterTabSpawner(SelectionViewTabId_1, FOnSpawnTab::CreateSP(this, &FDataflowEditorToolkit::SpawnTab_SelectionView))
 		.SetDisplayName(LOCTEXT("DataflowSelectionViewTab1", "Selection View 1"))
diff --git a/Engine/Plugins/Experimental/Dataflow/Source/DataflowEditor/Private/Dataflow/DataflowSkeletonView.cpp b/Engine/Plugins/Experimental/Dataflow/Source/DataflowEditor/Private/Dataflow/DataflowSkeletonView.cpp
--- a/Engine/Plugins/Experimental/Dataflow/Source/DataflowEditor/Private/Dataflow/DataflowSkeletonView.cpp
+++ b/Engine/Plugins/Experimental/Dataflow/Source/DataflowEditor/Private/Dataflow/DataflowSkeletonView.cpp
@@ -9,19 +9,28 @@
 #include "IEditableSkeleton.h"
 #include "ISkeletonTree.h"
 #include "ISkeletonTreeItem.h"
+#include "ISkeletonEditorModule.h"
 #include "Engine/SkeletalMesh.h"
 
-//#include "Widgets/SCompoundWidget.h"
 
 FDataflowSkeletonView::FDataflowSkeletonView(UDataflowEditor* InDataflowEditor)
 	: DataflowEditor(InDataflowEditor)
 	, SkeletonEditor(nullptr)
 	, SkeletalMesh(NewObject<USkeletalMesh>())
-	, SelectedNode(nullptr)
 	, CollectionIndexRemap(TArray<int32>())
 
 {
-	SkeletalMesh->SetSkeleton(NewObject<USkeleton>(SkeletalMesh, NAME_Name));
+	check(DataflowEditor);
+
+	SetSkeleton(NewObject<USkeleton>(SkeletalMesh, NAME_Name));
+	if (const TObjectPtr<UDataflowSkeletalContent> SkeletalContent = Cast<UDataflowSkeletalContent>(DataflowEditor->GetDataflowContent()))
+	{
+		if (SkeletalContent->GetDataflowAsset())
+		{
+			SetSkeleton(SkeletalContent->GetSkeleton());
+		}
+	}
+
 }
 
 FDataflowSkeletonView::~FDataflowSkeletonView()
@@ -32,6 +41,15 @@ FDataflowSkeletonView::~FDataflowSkeletonView()
 	}
 }
 
+TSharedPtr<ISkeletonTree> FDataflowSkeletonView::CreateEditor(FSkeletonTreeArgs& InSkeletonTreeArgs)
+{
+	ISkeletonEditorModule& SkeletonEditorModule = FModuleManager::LoadModuleChecked<ISkeletonEditorModule>("SkeletonEditor");
+	SkeletonEditor = SkeletonEditorModule.CreateSkeletonTree(GetSkeleton(), InSkeletonTreeArgs);
+	SkeletonEditor->Refresh();
+	// add widget delegates (see FDataflowCollectionSpreadSheet)
+	return SkeletonEditor;
+}
+
 void FDataflowSkeletonView::SetSupportedOutputTypes()
 {
 	GetSupportedOutputTypes().Empty();
@@ -52,6 +70,10 @@ void FDataflowSkeletonView::SetSkeleton(USkeleton* Skeleton)
 		SkeletalMesh = NewObject<USkeletalMesh>();
 		SkeletalMesh->SetSkeleton(NewObject<USkeleton>(SkeletalMesh, NAME_Name));
 	}
+	if (SkeletonEditor)
+	{
+		SkeletonEditor->Refresh();
+	}
 }
 
 USkeleton* FDataflowSkeletonView::GetSkeleton()
@@ -65,8 +87,9 @@ USkeleton* FDataflowSkeletonView::GetSkeleton()
 
 void FDataflowSkeletonView::UpdateViewData()
 {
-	if(!ensure(DataflowEditor)) return;
+	if (!ensure(DataflowEditor)) return;
 
+	bool bNeedsDefaultSkeleton = true;
 	if (TObjectPtr<UDataflowEdNode> EdNode = GetSelectedNode())
 	{
 		if (GetSelectedNode()->IsBound())
@@ -85,7 +108,6 @@ void FDataflowSkeletonView::UpdateViewData()
 								FManagedArrayCollection DefaultCollection;
 								const FManagedArrayCollection& Result = Output->GetValue(*EditorContent->GetDataflowContext(), DefaultCollection);
 
-								SelectedNode = EdNode;
 								SkeletalMesh = NewObject<USkeletalMesh>();
 								TObjectPtr<USkeleton> Skeleton = NewObject<USkeleton>(SkeletalMesh, Node->Name);
 
@@ -99,6 +121,7 @@ void FDataflowSkeletonView::UpdateViewData()
 									SkeletonEditor->SetSkeletalMesh(SkeletalMesh);
 									SkeletonEditor->Refresh();
 								}
+								bNeedsDefaultSkeleton = false;
 							}
 						}
 					}
@@ -106,20 +129,42 @@ void FDataflowSkeletonView::UpdateViewData()
 			}
 		}
 	}
-}
 
-void FDataflowSkeletonView::SetSkeletonEditor(TSharedPtr<ISkeletonTree>& InSkeletonEditor)
-{
-	ensure(!SkeletonEditor);
 
-	SkeletonEditor = InSkeletonEditor;
+	if (bNeedsDefaultSkeleton)
+	{
+		SetSkeleton(NewObject<USkeleton>(SkeletalMesh, NAME_Name));
+		if (DataflowEditor)
+		{
+			if (const TObjectPtr<UDataflowSkeletalContent> SkeletalContent = Cast<UDataflowSkeletalContent>(DataflowEditor->GetDataflowContent()))
+			{
+				if (SkeletalContent->GetDataflowAsset())
+				{
+					SetSkeleton(SkeletalContent->GetSkeleton());
+				}
+			}
+		}
+		if (SkeletonEditor)
+		{
+			SkeletonEditor->Refresh();
+		}
+	}
+}
 
-	if (SkeletonEditor)
+void FDataflowSkeletonView::ConstructionViewSelectionChanged(const TArray<UPrimitiveComponent*>& InSelectedComponents)
+{
+	if (ensure(SkeletonEditor))
 	{
-		// add widget delegates (see FDataflowCollectionSpreadSheet)
+		SkeletonEditor->DeselectAll();
+		for (UPrimitiveComponent* Component : InSelectedComponents)
+		{
+			SkeletonEditor->SetSelectedBone(FName(Component->GetName()), ESelectInfo::Type::Direct);
+		}
+		SkeletonEditor->Refresh();
 	}
 }
 
+
 void FDataflowSkeletonView::SkeletonViewSelectionChanged(const TArrayView<TSharedPtr<ISkeletonTreeItem>>& InSelectedItems, ESelectInfo::Type InSelectInfo)
 {
 	if (SkeletonEditor)
@@ -135,7 +180,5 @@ void FDataflowSkeletonView::SkeletonViewSelectionChanged(const TArrayView<TShare
 void FDataflowSkeletonView::AddReferencedObjects(FReferenceCollector& Collector)
 {
 	FDataflowNodeView::AddReferencedObjects(Collector);
-
-	Collector.AddReferencedObject(SelectedNode);
 	Collector.AddReferencedObject(SkeletalMesh);
 }
diff --git a/Engine/Plugins/Experimental/Dataflow/Source/DataflowEditor/Private/Dataflow/DataflowView.cpp b/Engine/Plugins/Experimental/Dataflow/Source/DataflowEditor/Private/Dataflow/DataflowView.cpp
--- a/Engine/Plugins/Experimental/Dataflow/Source/DataflowEditor/Private/Dataflow/DataflowView.cpp
+++ b/Engine/Plugins/Experimental/Dataflow/Source/DataflowEditor/Private/Dataflow/DataflowView.cpp
@@ -53,6 +53,11 @@ void FDataflowNodeView::SetContext(TSharedPtr<Dataflow::FContext>& InContext)
 	Context = InContext;
 }
 
+void FDataflowNodeView::OnConstructionViewSelectionChanged(const TArray<UPrimitiveComponent*>& InComponents)
+{
+	ConstructionViewSelectionChanged(InComponents);
+}
+
 
 void FDataflowNodeView::OnSelectedNodeChanged(UDataflowEdNode* InNode)
 {
diff --git a/Engine/Plugins/Experimental/Dataflow/Source/DataflowEditor/Public/Dataflow/DataflowCollectionSpreadSheet.h b/Engine/Plugins/Experimental/Dataflow/Source/DataflowEditor/Public/Dataflow/DataflowCollectionSpreadSheet.h
--- a/Engine/Plugins/Experimental/Dataflow/Source/DataflowEditor/Public/Dataflow/DataflowCollectionSpreadSheet.h
+++ b/Engine/Plugins/Experimental/Dataflow/Source/DataflowEditor/Public/Dataflow/DataflowCollectionSpreadSheet.h
@@ -22,6 +22,7 @@ public:
 
 	virtual void SetSupportedOutputTypes() override;
 	virtual void UpdateViewData() override;
+	virtual void ConstructionViewSelectionChanged(const TArray<UPrimitiveComponent*>& InSelectedComponents) override {};
 
 	void SetCollectionSpreadSheet(TSharedPtr<SCollectionSpreadSheetWidget>& InCollectionSpreadSheet);
 
diff --git a/Engine/Plugins/Experimental/Dataflow/Source/DataflowEditor/Public/Dataflow/DataflowConstructionViewportClient.h b/Engine/Plugins/Experimental/Dataflow/Source/DataflowEditor/Public/Dataflow/DataflowConstructionViewportClient.h
--- a/Engine/Plugins/Experimental/Dataflow/Source/DataflowEditor/Public/Dataflow/DataflowConstructionViewportClient.h
+++ b/Engine/Plugins/Experimental/Dataflow/Source/DataflowEditor/Public/Dataflow/DataflowConstructionViewportClient.h
@@ -6,8 +6,10 @@
 #include "Dataflow/DataflowNodeParameters.h"
 #include "Dataflow/DataflowComponentSelectionState.h"
 #include "Dataflow/DataflowContent.h"
+#include "Delegates/Delegate.h"
 #include "InputBehaviorSet.h"
 
+
 class FDataflowEditorToolkit;
 class ADataflowActor;
 class UTransformProxy;
@@ -45,6 +47,9 @@ public:
 	virtual void AddReferencedObjects(FReferenceCollector& Collector) override;
 	virtual FString GetReferencerName() const override { return TEXT("FDataflowConstructionViewportClient"); }
 
+	DECLARE_MULTICAST_DELEGATE_OneParam(FOnSelectionChangedMulticast, const TArray<UPrimitiveComponent*>&)
+	FOnSelectionChangedMulticast OnSelectionChangedMulticast;
+
 private:
 
 	// FEditorViewportClient interface
diff --git a/Engine/Plugins/Experimental/Dataflow/Source/DataflowEditor/Public/Dataflow/DataflowEditorToolkit.h b/Engine/Plugins/Experimental/Dataflow/Source/DataflowEditor/Public/Dataflow/DataflowEditorToolkit.h
--- a/Engine/Plugins/Experimental/Dataflow/Source/DataflowEditor/Public/Dataflow/DataflowEditorToolkit.h
+++ b/Engine/Plugins/Experimental/Dataflow/Source/DataflowEditor/Public/Dataflow/DataflowEditorToolkit.h
@@ -95,7 +95,7 @@ protected:
 	void OnNodeDeleted(const TSet<UObject*>& NewSelection);
 	void OnNodeSingleClicked(UObject* ClickedNode) const;
 	void OnAssetPropertyValueChanged(const FPropertyChangedEvent& PropertyChangedEvent);
-	
+	void OnConstructionViewSelectionChanged(const TArray<UPrimitiveComponent*>& SelectedComponents);
 	// Callback to remove the closed one from the listener views
 	void OnTabClosed(TSharedRef<SDockTab> Tab);
 	
@@ -179,7 +179,8 @@ private:
 	FDelegateHandle OnSelectionChangedMulticastDelegateHandle;
     FDelegateHandle OnNodeDeletedMulticastDelegateHandle;
     FDelegateHandle OnFinishedChangingPropertiesDelegateHandle;
-    FDelegateHandle OnFinishedChangingAssetPropertiesDelegateHandle;
+	FDelegateHandle OnFinishedChangingAssetPropertiesDelegateHandle;
+	FDelegateHandle OnConstructionSelectionChangedDelegateHandle;
 
 	// The currently selected set of dataflow nodes. 
 	UPROPERTY()
diff --git a/Engine/Plugins/Experimental/Dataflow/Source/DataflowEditor/Public/Dataflow/DataflowSelectionView.h b/Engine/Plugins/Experimental/Dataflow/Source/DataflowEditor/Public/Dataflow/DataflowSelectionView.h
--- a/Engine/Plugins/Experimental/Dataflow/Source/DataflowEditor/Public/Dataflow/DataflowSelectionView.h
+++ b/Engine/Plugins/Experimental/Dataflow/Source/DataflowEditor/Public/Dataflow/DataflowSelectionView.h
@@ -22,6 +22,7 @@ public:
 
 	virtual void SetSupportedOutputTypes() override;
 	virtual void UpdateViewData() override;
+	virtual void ConstructionViewSelectionChanged(const TArray<UPrimitiveComponent*>& InSelectedComponents) override {};
 
 	void SetSelectionView(TSharedPtr<SSelectionViewWidget>& InSelectionView);
 
diff --git a/Engine/Plugins/Experimental/Dataflow/Source/DataflowEditor/Public/Dataflow/DataflowSkeletonView.h b/Engine/Plugins/Experimental/Dataflow/Source/DataflowEditor/Public/Dataflow/DataflowSkeletonView.h
--- a/Engine/Plugins/Experimental/Dataflow/Source/DataflowEditor/Public/Dataflow/DataflowSkeletonView.h
+++ b/Engine/Plugins/Experimental/Dataflow/Source/DataflowEditor/Public/Dataflow/DataflowSkeletonView.h
@@ -10,6 +10,7 @@ class ISkeletonTreeItem;
 class USkeleton;
 class USkeletalMesh;
 class UDataflowEdNode;
+struct FSkeletonTreeArgs;
 
 /**
 *
@@ -22,8 +23,8 @@ public:
 	FDataflowSkeletonView(UDataflowEditor* InDataflowEditor = nullptr);
 	~FDataflowSkeletonView();
 
-	/** Set the Skeleton Tree*/
-	void SetSkeletonEditor(TSharedPtr<ISkeletonTree>& InSkeletonTree);
+	/** Create the Skeleton Tree Editor*/
+	TSharedPtr<ISkeletonTree> CreateEditor(FSkeletonTreeArgs& InSkeletonTreeArgs);
 
 	/** Set selection types*/
 	virtual void SetSupportedOutputTypes() override;
@@ -37,23 +38,24 @@ public:
 	/** Update the view */
 	virtual void UpdateViewData() override;
 
+	/** Update the view based on changes in the construction view */
+	virtual void ConstructionViewSelectionChanged(const TArray<UPrimitiveComponent*>& InSelectedComponents) override;
+
 	/** Add GC managed objects*/
 	virtual void AddReferencedObjects(FReferenceCollector& Collector) override;
 
 	/** Selection View Callbacks */
 	void SkeletonViewSelectionChanged(const TArrayView<TSharedPtr<ISkeletonTreeItem>>& InSelectedItems, ESelectInfo::Type InSelectInfo);
 
+
 private:
 	UDataflowEditor* DataflowEditor = nullptr;
 
 	TSharedPtr<ISkeletonTree> SkeletonEditor;
 
-
 	/* Skeletal Mesh in the SkeletalViewer*/
 	TObjectPtr<USkeletalMesh> SkeletalMesh = nullptr;
 
-	TObjectPtr<UDataflowEdNode> SelectedNode = nullptr;
-
 	/* Rempping from the selected node to the SkeletalMesh*/
 	TArray<int32> CollectionIndexRemap;
 };
diff --git a/Engine/Plugins/Experimental/Dataflow/Source/DataflowEditor/Public/Dataflow/DataflowView.h b/Engine/Plugins/Experimental/Dataflow/Source/DataflowEditor/Public/Dataflow/DataflowView.h
--- a/Engine/Plugins/Experimental/Dataflow/Source/DataflowEditor/Public/Dataflow/DataflowView.h
+++ b/Engine/Plugins/Experimental/Dataflow/Source/DataflowEditor/Public/Dataflow/DataflowView.h
@@ -8,6 +8,7 @@
 #include "Dataflow/SelectionViewWidget.h"
 #include "Dataflow/DataflowSelection.h"
 
+class UPrimitiveComponent;
 
 /**
 *
@@ -17,6 +18,7 @@
 class IDataflowViewListener
 {
 public:
+	virtual void OnConstructionViewSelectionChanged(const TArray<UPrimitiveComponent*>& InSelectedComponents) = 0;
 	virtual void OnSelectedNodeChanged(UDataflowEdNode* InNode) = 0;  // nullptr is valid
 	virtual void OnNodeInvalidated(FDataflowNode* InvalidatedNode) = 0;
 };
@@ -45,6 +47,7 @@ public:
 	*/
 	virtual void UpdateViewData() = 0;
 	virtual void SetSupportedOutputTypes() = 0;
+	virtual void ConstructionViewSelectionChanged(const TArray<UPrimitiveComponent*>& InComponent) = 0;
 
 	/**
 	* Callback for PinnedDown change
@@ -59,6 +62,7 @@ public:
 	/**
 	* Virtual function overrides from IDataflowViewListener base class
 	*/
+	virtual void OnConstructionViewSelectionChanged(const TArray<UPrimitiveComponent*>& InNode) override;
 	virtual void OnSelectedNodeChanged(UDataflowEdNode* InNode) override;  // nullptr is valid
 	virtual void OnNodeInvalidated(FDataflowNode* InvalidatedNode) override;
 
