Dataflow : Dataflow View :
- Added access through the editor content.

diff --git a/Engine/Plugins/Experimental/Dataflow/Source/DataflowEditor/Private/Dataflow/AssetDefinition_DataflowAsset.cpp b/Engine/Plugins/Experimental/Dataflow/Source/DataflowEditor/Private/Dataflow/AssetDefinition_DataflowAsset.cpp
--- a/Engine/Plugins/Experimental/Dataflow/Source/DataflowEditor/Private/Dataflow/AssetDefinition_DataflowAsset.cpp
+++ b/Engine/Plugins/Experimental/Dataflow/Source/DataflowEditor/Private/Dataflow/AssetDefinition_DataflowAsset.cpp
@@ -184,7 +184,7 @@ EAssetCommandResult UAssetDefinition_DataflowAsset::OpenAssets(const FAssetOpenA
 		// Validate the asset
 		if (UDataflow* const DataflowAsset = CastChecked<UDataflow>(DataflowObjects[0]))
 		{
-			AssetEditor->InitializeContent(nullptr, DataflowAsset);
+			AssetEditor->InitializeContent(DataflowAsset);
 			return EAssetCommandResult::Handled;
 		}
 	}
diff --git a/Engine/Plugins/Experimental/Dataflow/Source/DataflowEditor/Private/Dataflow/DataflowCollectionSpreadSheet.cpp b/Engine/Plugins/Experimental/Dataflow/Source/DataflowEditor/Private/Dataflow/DataflowCollectionSpreadSheet.cpp
--- a/Engine/Plugins/Experimental/Dataflow/Source/DataflowEditor/Private/Dataflow/DataflowCollectionSpreadSheet.cpp
+++ b/Engine/Plugins/Experimental/Dataflow/Source/DataflowEditor/Private/Dataflow/DataflowCollectionSpreadSheet.cpp
@@ -1,9 +1,20 @@
 // Copyright Epic Games, Inc. All Rights Reserved.
 
 #include "Dataflow/DataflowCollectionSpreadSheet.h"
-#include "Templates/EnableIf.h"
+
+#include "Dataflow/DataflowCollectionSpreadSheetWidget.h"
+#include "Dataflow/DataflowContent.h"
 #include "Dataflow/DataflowSelection.h"
+#include "Templates/EnableIf.h"
 
+//#include "Dataflow/DataflowEdNode.h"
+
+
+FDataflowCollectionSpreadSheet::FDataflowCollectionSpreadSheet(TObjectPtr<UDataflowBaseContent> InContent)
+	: FDataflowNodeView(InContent)
+{
+
+}
 
 void FDataflowCollectionSpreadSheet::SetSupportedOutputTypes()
 {
@@ -25,17 +36,22 @@ void FDataflowCollectionSpreadSheet::UpdateViewData()
 				if (TSharedPtr<FDataflowNode> DataflowNode = GetSelectedNode()->DataflowGraph->FindBaseNode(GetSelectedNode()->DataflowNodeGuid))
 				{
 					TArray<FDataflowOutput*> Outputs = DataflowNode->GetOutputs();
-
-					for (FDataflowOutput* Output : Outputs)
+					if (const TObjectPtr<UDataflowBaseContent> Content = GetEditorContent())
 					{
-						FName Name = Output->GetName();
-						FName Type = Output->GetType();
-
-						if (Output->GetType() == "FManagedArrayCollection" && !Output->IsAnyType())
+						if (TSharedPtr<Dataflow::FEngineContext> Context = Content->GetDataflowContext())
 						{
-							const FManagedArrayCollection& Value = Output->GetValue<FManagedArrayCollection>(*GetContext(), FManagedArrayCollection());
-
-							CollectionSpreadSheet->GetCollectionTable()->GetCollectionInfoMap().Add(Name.ToString(), { Value });
+							for (FDataflowOutput* Output : Outputs)
+							{
+								FName Name = Output->GetName();
+								FName Type = Output->GetType();
+
+								if (Output->GetType() == "FManagedArrayCollection" && !Output->IsAnyType())
+								{
+									const FManagedArrayCollection& Value = Output->GetValue<FManagedArrayCollection>(*Context, FManagedArrayCollection());
+
+									CollectionSpreadSheet->GetCollectionTable()->GetCollectionInfoMap().Add(Name.ToString(), { Value });
+								}
+							}
 						}
 					}
 				}
diff --git a/Engine/Plugins/Experimental/Dataflow/Source/DataflowEditor/Private/Dataflow/DataflowEditor.cpp b/Engine/Plugins/Experimental/Dataflow/Source/DataflowEditor/Private/Dataflow/DataflowEditor.cpp
--- a/Engine/Plugins/Experimental/Dataflow/Source/DataflowEditor/Private/Dataflow/DataflowEditor.cpp
+++ b/Engine/Plugins/Experimental/Dataflow/Source/DataflowEditor/Private/Dataflow/DataflowEditor.cpp
@@ -26,69 +26,67 @@ void UDataflowEditor::Initialize(const TArray<TObjectPtr<UObject>>& InObjects)
 {
 	if(!InObjects.IsEmpty())
 	{
-		InitializeContent(nullptr, InObjects[0]);
+		InitializeContent(InObjects[0]);
 	}
 }
 
-void UDataflowEditor::InitializeContent(TObjectPtr<UDataflowBaseContent> BaseContent, const TObjectPtr<UObject>& ContentOwner)
+void UDataflowEditor::InitializeContent(const TObjectPtr<UObject>& ContentOwner)
 {
+	check(DataflowContent == nullptr);
+
 	TArray<TObjectPtr<UObject>> RequiredObjects = { ContentOwner };
 
-	DataflowContent = BaseContent;
-	if(!DataflowContent)
+	if (UDataflow* DataflowAsset = Cast<UDataflow>(ContentOwner))
 	{
-		if(UDataflow* DataflowAsset = Cast<UDataflow>(ContentOwner))
-		{
-			DataflowContent = DataflowContextDefinitionHelpers::CreateNewDataflowContext<UDataflowBaseContent>(ContentOwner);
+		DataflowContent = DataflowContextDefinitionHelpers::CreateNewDataflowContext<UDataflowBaseContent>(ContentOwner);
 
-			DataflowContent->SetDataflowAsset(DataflowAsset);
-			DataflowContent->SetDataflowTerminal(FString());
-		}
-		else
+		DataflowContent->SetDataflowAsset(DataflowAsset);
+		DataflowContent->SetDataflowTerminal(FString());
+	}
+	else
+	{
+		if (Private::HasDataflowAsset(ContentOwner))
 		{
-			if(Private::HasDataflowAsset(ContentOwner))
+			if (Private::HasSkeletalMesh(ContentOwner))
 			{
-				if(Private::HasSkeletalMesh(ContentOwner))
-				{
-					DataflowContent = DataflowContextDefinitionHelpers::CreateNewDataflowContext<UDataflowSkeletalContent>(ContentOwner);
-					const TObjectPtr<UDataflowSkeletalContent> SkeletalContent = Cast<UDataflowSkeletalContent>(DataflowContent);
-					
-					SkeletalContent->SetSkeletalMesh(Private::GetSkeletalMeshFrom(ContentOwner));
-                    SkeletalContent->SetSkeleton(Private::GetSkeletonFrom(ContentOwner));
-                    SkeletalContent->SetAnimationAsset(Private::GetAnimationAssetFrom(ContentOwner));
-				}
-				else
-				{
-					DataflowContent = DataflowContextDefinitionHelpers::CreateNewDataflowContext<UDataflowBaseContent>(ContentOwner);
-				}
-
-				DataflowContent->SetDataflowAsset(Private::GetDataflowAssetFrom(ContentOwner));
-				DataflowContent->SetDataflowTerminal(Private::GetDataflowTerminalFrom(ContentOwner));
-				RequiredObjects.Add(Private::GetDataflowAssetFrom(ContentOwner));
+				DataflowContent = DataflowContextDefinitionHelpers::CreateNewDataflowContext<UDataflowSkeletalContent>(ContentOwner);
+				const TObjectPtr<UDataflowSkeletalContent> SkeletalContent = Cast<UDataflowSkeletalContent>(DataflowContent);
+
+				SkeletalContent->SetSkeletalMesh(Private::GetSkeletalMeshFrom(ContentOwner));
+				SkeletalContent->SetSkeleton(Private::GetSkeletonFrom(ContentOwner));
+				SkeletalContent->SetAnimationAsset(Private::GetAnimationAssetFrom(ContentOwner));
+			}
+			else
+			{
+				DataflowContent = DataflowContextDefinitionHelpers::CreateNewDataflowContext<UDataflowBaseContent>(ContentOwner);
 			}
+
+			DataflowContent->SetDataflowAsset(Private::GetDataflowAssetFrom(ContentOwner));
+			DataflowContent->SetDataflowTerminal(Private::GetDataflowTerminalFrom(ContentOwner));
+			RequiredObjects.Add(Private::GetDataflowAssetFrom(ContentOwner));
 		}
 	}
 
 	if (!DataflowContent) return;
 	RequiredObjects.Add(DataflowContent);
 
-	if(const TObjectPtr<UDataflowSkeletalContent> SkeletalContent = Cast<UDataflowSkeletalContent>(DataflowContent))
+	if (const TObjectPtr<UDataflowSkeletalContent> SkeletalContent = Cast<UDataflowSkeletalContent>(DataflowContent))
 	{
-		if(!SkeletalContent->GetSkeletalMesh())
+		if (!SkeletalContent->GetSkeletalMesh())
 		{
 			const FName SkeletalMeshName = MakeUniqueObjectName(SkeletalContent->GetDataflowAsset(), UDataflow::StaticClass(), FName("USkeletalMesh"));
 			USkeletalMesh* SkeletalMesh = NewObject<USkeletalMesh>(SkeletalContent->GetDataflowAsset(), SkeletalMeshName);
 
 			USkeleton* Skeleton = SkeletalContent->GetSkeleton();
-			if(!Skeleton)
+			if (!Skeleton)
 			{
 				const FName SkeletonName = MakeUniqueObjectName(SkeletalContent->GetDataflowAsset(), UDataflow::StaticClass(), FName("USkeleton"));
-                Skeleton = NewObject<USkeleton>(SkeletalContent->GetDataflowAsset(), SkeletonName);
+				Skeleton = NewObject<USkeleton>(SkeletalContent->GetDataflowAsset(), SkeletonName);
 			}
 			SkeletalMesh->SetSkeleton(Skeleton);
 			SkeletalContent->SetSkeletalMesh(SkeletalMesh);
 		}
-		else if(!SkeletalContent->GetSkeleton())
+		else if (!SkeletalContent->GetSkeleton())
 		{
 			SkeletalContent->SetSkeleton(SkeletalContent->GetSkeletalMesh()->GetSkeleton());
 		}
diff --git a/Engine/Plugins/Experimental/Dataflow/Source/DataflowEditor/Private/Dataflow/DataflowEditorToolkit.cpp b/Engine/Plugins/Experimental/Dataflow/Source/DataflowEditor/Private/Dataflow/DataflowEditorToolkit.cpp
--- a/Engine/Plugins/Experimental/Dataflow/Source/DataflowEditor/Private/Dataflow/DataflowEditorToolkit.cpp
+++ b/Engine/Plugins/Experimental/Dataflow/Source/DataflowEditor/Private/Dataflow/DataflowEditorToolkit.cpp
@@ -8,6 +8,7 @@
 #include "Dataflow/DataflowCore.h"
 #include "Dataflow/DataflowEditor.h"
 #include "Dataflow/DataflowContent.h"
+#include "Dataflow/DataflowCollectionSpreadSheetWidget.h"
 #include "Dataflow/DataflowConstructionScene.h"
 #include "Dataflow/DataflowConstructionViewportClient.h"
 #include "Dataflow/DataflowEditorCollectionComponent.h"
@@ -49,6 +50,7 @@
 #include "Styling/SlateStyleRegistry.h"
 #include "Widgets/Docking/SDockTab.h"
 
+
 #define LOCTEXT_NAMESPACE "DataflowEditorToolkit"
 
 //DEFINE_LOG_CATEGORY_STATIC(EditorToolkitLog, Log, All);
@@ -910,8 +912,10 @@ TSharedRef<SDockTab> FDataflowEditorToolkit::SpawnTab_NodeDetails(const FSpawnTa
 TSharedRef<SDockTab> FDataflowEditorToolkit::SpawnTab_SkeletonView(const FSpawnTabArgs& Args)
 {
 	check(Args.GetTabId() == SkeletonViewTabId);
+	check(DataflowEditor);
+	check(DataflowEditor->GetDataflowContent());
 
-	SkeletonEditorView = MakeShared<FDataflowSkeletonView>(DataflowEditor);
+	SkeletonEditorView = MakeShared<FDataflowSkeletonView>(DataflowEditor->GetDataflowContent());
 	ViewListeners.Add(SkeletonEditorView.Get());
 
 	FSkeletonTreeArgs SkeletonTreeArgs;
@@ -935,10 +939,12 @@ TSharedRef<SDockTab> FDataflowEditorToolkit::SpawnTab_SkeletonView(const FSpawnT
 TSharedRef<SDockTab> FDataflowEditorToolkit::SpawnTab_SelectionView(const FSpawnTabArgs& Args)
 {
 	//	check(Args.GetTabId().TabType == SelectionViewTabId_1);
+	check(DataflowEditor);
+	check(DataflowEditor->GetDataflowContent());
 
-	if (Args.GetTabId() == SelectionViewTabId_1)
+		if (Args.GetTabId() == SelectionViewTabId_1)
 	{
-		DataflowSelectionView_1 = MakeShared<FDataflowSelectionView>(FDataflowSelectionView());
+		DataflowSelectionView_1 = MakeShared<FDataflowSelectionView>(FDataflowSelectionView(DataflowEditor->GetDataflowContent()));
 		if (DataflowSelectionView_1.IsValid())
 		{
 			ViewListeners.Add(DataflowSelectionView_1.Get());
@@ -946,7 +952,7 @@ TSharedRef<SDockTab> FDataflowEditorToolkit::SpawnTab_SelectionView(const FSpawn
 	}
 	else if (Args.GetTabId() == SelectionViewTabId_2)
 	{
-		DataflowSelectionView_2 = MakeShared<FDataflowSelectionView>(FDataflowSelectionView());
+		DataflowSelectionView_2 = MakeShared<FDataflowSelectionView>(FDataflowSelectionView(DataflowEditor->GetDataflowContent()));
 		if (DataflowSelectionView_2.IsValid())
 		{
 			ViewListeners.Add(DataflowSelectionView_2.Get());
@@ -954,7 +960,7 @@ TSharedRef<SDockTab> FDataflowEditorToolkit::SpawnTab_SelectionView(const FSpawn
 	}
 	else if (Args.GetTabId() == SelectionViewTabId_3)
 	{
-		DataflowSelectionView_3 = MakeShared<FDataflowSelectionView>(FDataflowSelectionView());
+		DataflowSelectionView_3 = MakeShared<FDataflowSelectionView>(FDataflowSelectionView(DataflowEditor->GetDataflowContent()));
 		if (DataflowSelectionView_3.IsValid())
 		{
 			ViewListeners.Add(DataflowSelectionView_3.Get());
@@ -962,7 +968,7 @@ TSharedRef<SDockTab> FDataflowEditorToolkit::SpawnTab_SelectionView(const FSpawn
 	}
 	else if (Args.GetTabId() == SelectionViewTabId_4)
 	{
-		DataflowSelectionView_4 = MakeShared<FDataflowSelectionView>(FDataflowSelectionView());
+		DataflowSelectionView_4 = MakeShared<FDataflowSelectionView>(FDataflowSelectionView(DataflowEditor->GetDataflowContent()));
 		if (DataflowSelectionView_4.IsValid())
 		{
 			ViewListeners.Add(DataflowSelectionView_4.Get());
@@ -972,9 +978,9 @@ TSharedRef<SDockTab> FDataflowEditorToolkit::SpawnTab_SelectionView(const FSpawn
 	TSharedPtr<SSelectionViewWidget> SelectionViewWidget;
 
 	TSharedRef<SDockTab> DockableTab = SNew(SDockTab)
-		[
-			SAssignNew(SelectionViewWidget, SSelectionViewWidget)
-		];
+	[
+		SAssignNew(SelectionViewWidget, SSelectionViewWidget)
+	];
 
 	if (SelectionViewWidget)
 	{
@@ -983,42 +989,18 @@ TSharedRef<SDockTab> FDataflowEditorToolkit::SpawnTab_SelectionView(const FSpawn
 			if (Args.GetTabId() == SelectionViewTabId_1)
 			{
 				DataflowSelectionView_1->SetSelectionView(SelectionViewWidget);
-
-				// Set the Context on the interface
-				if (TSharedPtr<Dataflow::FContext> CurrentContext = EditorContent->GetDataflowContext())
-				{
-					DataflowSelectionView_1->SetContext(CurrentContext);
-				}
 			}
 			else if (Args.GetTabId() == SelectionViewTabId_2)
 			{
 				DataflowSelectionView_2->SetSelectionView(SelectionViewWidget);
-
-				// Set the Context on the interface
-				if (TSharedPtr<Dataflow::FContext> CurrentContext = EditorContent->GetDataflowContext())
-				{
-					DataflowSelectionView_2->SetContext(CurrentContext);
-				}
 			}
 			else if (Args.GetTabId() == SelectionViewTabId_3)
 			{
 				DataflowSelectionView_3->SetSelectionView(SelectionViewWidget);
-
-				// Set the Context on the interface
-				if (TSharedPtr<Dataflow::FContext> CurrentContext = EditorContent->GetDataflowContext())
-				{
-					DataflowSelectionView_3->SetContext(CurrentContext);
-				}
 			}
 			else if (Args.GetTabId() == SelectionViewTabId_4)
 			{
 				DataflowSelectionView_4->SetSelectionView(SelectionViewWidget);
-
-				// Set the Context on the interface
-				if (TSharedPtr<Dataflow::FContext> CurrentContext = EditorContent->GetDataflowContext())
-				{
-					DataflowSelectionView_4->SetContext(CurrentContext);
-				}
 			}
 		}
 	}
@@ -1030,9 +1012,12 @@ TSharedRef<SDockTab> FDataflowEditorToolkit::SpawnTab_SelectionView(const FSpawn
 
 TSharedRef<SDockTab> FDataflowEditorToolkit::SpawnTab_CollectionSpreadSheet(const FSpawnTabArgs& Args)
 {
+	check(DataflowEditor);
+	check(DataflowEditor->GetDataflowContent());
+
 	if (Args.GetTabId() == CollectionSpreadSheetTabId_1)
 	{
-		DataflowCollectionSpreadSheet_1 = MakeShared<FDataflowCollectionSpreadSheet>(FDataflowCollectionSpreadSheet());
+		DataflowCollectionSpreadSheet_1 = MakeShared<FDataflowCollectionSpreadSheet>(FDataflowCollectionSpreadSheet(DataflowEditor->GetDataflowContent()));
 		if (DataflowCollectionSpreadSheet_1.IsValid())
 		{
 			ViewListeners.Add(DataflowCollectionSpreadSheet_1.Get());
@@ -1040,7 +1025,7 @@ TSharedRef<SDockTab> FDataflowEditorToolkit::SpawnTab_CollectionSpreadSheet(cons
 	}
 	else if (Args.GetTabId() == CollectionSpreadSheetTabId_2)
 	{
-		DataflowCollectionSpreadSheet_2 = MakeShared<FDataflowCollectionSpreadSheet>(FDataflowCollectionSpreadSheet());
+		DataflowCollectionSpreadSheet_2 = MakeShared<FDataflowCollectionSpreadSheet>(FDataflowCollectionSpreadSheet(DataflowEditor->GetDataflowContent()));
 		if (DataflowCollectionSpreadSheet_2.IsValid())
 		{
 			ViewListeners.Add(DataflowCollectionSpreadSheet_2.Get());
@@ -1048,7 +1033,7 @@ TSharedRef<SDockTab> FDataflowEditorToolkit::SpawnTab_CollectionSpreadSheet(cons
 	}
 	else if (Args.GetTabId() == CollectionSpreadSheetTabId_3)
 	{
-		DataflowCollectionSpreadSheet_3 = MakeShared<FDataflowCollectionSpreadSheet>(FDataflowCollectionSpreadSheet());
+		DataflowCollectionSpreadSheet_3 = MakeShared<FDataflowCollectionSpreadSheet>(FDataflowCollectionSpreadSheet(DataflowEditor->GetDataflowContent()));
 		if (DataflowCollectionSpreadSheet_3.IsValid())
 		{
 			ViewListeners.Add(DataflowCollectionSpreadSheet_3.Get());
@@ -1056,7 +1041,7 @@ TSharedRef<SDockTab> FDataflowEditorToolkit::SpawnTab_CollectionSpreadSheet(cons
 	}
 	else if (Args.GetTabId() == CollectionSpreadSheetTabId_4)
 	{
-		DataflowCollectionSpreadSheet_4 = MakeShared<FDataflowCollectionSpreadSheet>(FDataflowCollectionSpreadSheet());
+		DataflowCollectionSpreadSheet_4 = MakeShared<FDataflowCollectionSpreadSheet>(FDataflowCollectionSpreadSheet(DataflowEditor->GetDataflowContent()));
 		if (DataflowCollectionSpreadSheet_4.IsValid())
 		{
 			ViewListeners.Add(DataflowCollectionSpreadSheet_4.Get());
@@ -1077,42 +1062,18 @@ TSharedRef<SDockTab> FDataflowEditorToolkit::SpawnTab_CollectionSpreadSheet(cons
 			if (Args.GetTabId() == CollectionSpreadSheetTabId_1)
 			{
 				DataflowCollectionSpreadSheet_1->SetCollectionSpreadSheet(CollectionSpreadSheetWidget);
-
-				// Set the Context on the interface
-				if (TSharedPtr<Dataflow::FContext> CurrentContext = EditorContent->GetDataflowContext())
-				{
-					DataflowCollectionSpreadSheet_1->SetContext(CurrentContext);
-				}
 			}
 			else if (Args.GetTabId() == CollectionSpreadSheetTabId_2)
 			{
 				DataflowCollectionSpreadSheet_2->SetCollectionSpreadSheet(CollectionSpreadSheetWidget);
-
-				// Set the Context on the interface
-				if (TSharedPtr<Dataflow::FContext> CurrentContext = EditorContent->GetDataflowContext())
-				{
-					DataflowCollectionSpreadSheet_2->SetContext(CurrentContext);
-				}
 			}
 			else if (Args.GetTabId() == CollectionSpreadSheetTabId_3)
 			{
 				DataflowCollectionSpreadSheet_3->SetCollectionSpreadSheet(CollectionSpreadSheetWidget);
-
-				// Set the Context on the interface
-				if (TSharedPtr<Dataflow::FContext> CurrentContext = EditorContent->GetDataflowContext())
-				{
-					DataflowCollectionSpreadSheet_3->SetContext(CurrentContext);
-				}
 			}
 			else if (Args.GetTabId() == CollectionSpreadSheetTabId_4)
 			{
 				DataflowCollectionSpreadSheet_4->SetCollectionSpreadSheet(CollectionSpreadSheetWidget);
-
-				// Set the Context on the interface
-				if (TSharedPtr<Dataflow::FContext> CurrentContext = EditorContent->GetDataflowContext())
-				{
-					DataflowCollectionSpreadSheet_4->SetContext(CurrentContext);
-				}
 			}
 		}
 	}
diff --git a/Engine/Plugins/Experimental/Dataflow/Source/DataflowEditor/Private/Dataflow/DataflowSelectionView.cpp b/Engine/Plugins/Experimental/Dataflow/Source/DataflowEditor/Private/Dataflow/DataflowSelectionView.cpp
--- a/Engine/Plugins/Experimental/Dataflow/Source/DataflowEditor/Private/Dataflow/DataflowSelectionView.cpp
+++ b/Engine/Plugins/Experimental/Dataflow/Source/DataflowEditor/Private/Dataflow/DataflowSelectionView.cpp
@@ -1,11 +1,25 @@
 // Copyright Epic Games, Inc. All Rights Reserved.
 
 #include "Dataflow/DataflowSelectionView.h"
+
 #include "Templates/EnableIf.h"
+#include "Dataflow/DataflowContent.h"
+//#include "Dataflow/DataflowEdNode.h"
 #include "Dataflow/DataflowSelection.h"
+#include "Dataflow/SelectionViewWidget.h"
+
+//#include "Widgets/SCompoundWidget.h"
+//#include "Widgets/Views/SListView.h"
+
 
 #define LOCTEXT_NAMESPACE "DataflowSelectionView"
 
+FDataflowSelectionView::FDataflowSelectionView(TObjectPtr<UDataflowBaseContent> InContent)
+	: FDataflowNodeView(InContent)
+{
+
+}
+
 void FDataflowSelectionView::SetSupportedOutputTypes()
 {
 	GetSupportedOutputTypes().Empty();
@@ -27,30 +41,36 @@ void FDataflowSelectionView::UpdateViewData()
 			{
 				if (TSharedPtr<FDataflowNode> DataflowNode = GetSelectedNode()->DataflowGraph->FindBaseNode(GetSelectedNode()->DataflowNodeGuid))
 				{
-					TArray<FDataflowOutput*> Outputs = DataflowNode->GetOutputs();
-
-					for (FDataflowOutput* Output : Outputs)
+					if (const TObjectPtr<UDataflowBaseContent> Content = GetEditorContent())
 					{
-						FName Name = Output->GetName();
-						FName Type = Output->GetType();
-
-						if (Output->GetType() == "FDataflowTransformSelection")
-						{
-							const FDataflowTransformSelection& Value = Output->GetValue<FDataflowTransformSelection>(*GetContext(), FDataflowTransformSelection());
-
-							SelectionView->GetSelectionTable()->GetSelectionInfoMap().Add(Name.ToString(), { Type.ToString(), TBitArray<>(Value.GetBitArray()) });
-						}
-						else if (Output->GetType() == "FDataflowVertexSelection")
+						if (TSharedPtr<Dataflow::FEngineContext> Context = Content->GetDataflowContext())
 						{
-							const FDataflowVertexSelection& Value = Output->GetValue<FDataflowVertexSelection>(*GetContext(), FDataflowVertexSelection());
-
-							SelectionView->GetSelectionTable()->GetSelectionInfoMap().Add(Name.ToString(), { Type.ToString(), TBitArray<>(Value.GetBitArray()) });
-						}
-						else if (Output->GetType() == "FDataflowFaceSelection")
-						{
-							const FDataflowFaceSelection& Value = Output->GetValue<FDataflowFaceSelection>(*GetContext(), FDataflowFaceSelection());
-
-							SelectionView->GetSelectionTable()->GetSelectionInfoMap().Add(Name.ToString(), { Type.ToString(), TBitArray<>(Value.GetBitArray()) });
+							TArray<FDataflowOutput*> Outputs = DataflowNode->GetOutputs();
+
+							for (FDataflowOutput* Output : Outputs)
+							{
+								FName Name = Output->GetName();
+								FName Type = Output->GetType();
+
+								if (Output->GetType() == "FDataflowTransformSelection")
+								{
+									const FDataflowTransformSelection& Value = Output->GetValue<FDataflowTransformSelection>(*Context, FDataflowTransformSelection());
+
+									SelectionView->GetSelectionTable()->GetSelectionInfoMap().Add(Name.ToString(), { Type.ToString(), TBitArray<>(Value.GetBitArray()) });
+								}
+								else if (Output->GetType() == "FDataflowVertexSelection")
+								{
+									const FDataflowVertexSelection& Value = Output->GetValue<FDataflowVertexSelection>(*Context, FDataflowVertexSelection());
+
+									SelectionView->GetSelectionTable()->GetSelectionInfoMap().Add(Name.ToString(), { Type.ToString(), TBitArray<>(Value.GetBitArray()) });
+								}
+								else if (Output->GetType() == "FDataflowFaceSelection")
+								{
+									const FDataflowFaceSelection& Value = Output->GetValue<FDataflowFaceSelection>(*Context, FDataflowFaceSelection());
+
+									SelectionView->GetSelectionTable()->GetSelectionInfoMap().Add(Name.ToString(), { Type.ToString(), TBitArray<>(Value.GetBitArray()) });
+								}
+							}
 						}
 					}
 				}
diff --git a/Engine/Plugins/Experimental/Dataflow/Source/DataflowEditor/Private/Dataflow/DataflowSkeletonView.cpp b/Engine/Plugins/Experimental/Dataflow/Source/DataflowEditor/Private/Dataflow/DataflowSkeletonView.cpp
--- a/Engine/Plugins/Experimental/Dataflow/Source/DataflowEditor/Private/Dataflow/DataflowSkeletonView.cpp
+++ b/Engine/Plugins/Experimental/Dataflow/Source/DataflowEditor/Private/Dataflow/DataflowSkeletonView.cpp
@@ -13,17 +13,16 @@
 #include "Engine/SkeletalMesh.h"
 
 
-FDataflowSkeletonView::FDataflowSkeletonView(UDataflowEditor* InDataflowEditor)
-	: DataflowEditor(InDataflowEditor)
+FDataflowSkeletonView::FDataflowSkeletonView(TObjectPtr<UDataflowBaseContent> InContent)
+	: FDataflowNodeView(InContent)
 	, SkeletonEditor(nullptr)
 	, SkeletalMesh(NewObject<USkeletalMesh>())
 	, CollectionIndexRemap(TArray<int32>())
-
 {
-	check(DataflowEditor);
+	check(InContent);
 
 	SetSkeleton(NewObject<USkeleton>(SkeletalMesh, NAME_Name));
-	if (const TObjectPtr<UDataflowSkeletalContent> SkeletalContent = Cast<UDataflowSkeletalContent>(DataflowEditor->GetDataflowContent()))
+	if (const TObjectPtr<UDataflowSkeletalContent> SkeletalContent = Cast<UDataflowSkeletalContent>(InContent))
 	{
 		if (SkeletalContent->GetDataflowAsset())
 		{
@@ -87,8 +86,6 @@ USkeleton* FDataflowSkeletonView::GetSkeleton()
 
 void FDataflowSkeletonView::UpdateViewData()
 {
-	if (!ensure(DataflowEditor)) return;
-
 	bool bNeedsDefaultSkeleton = true;
 	if (TObjectPtr<UDataflowEdNode> EdNode = GetSelectedNode())
 	{
@@ -98,31 +95,27 @@ void FDataflowSkeletonView::UpdateViewData()
 			{
 				if (FDataflowOutput* Output = Node->FindOutput(FName("Collection")))
 				{
-					if (TObjectPtr<UDataflowBaseContent> EditorContent = DataflowEditor->GetDataflowContent())
+					if (const TObjectPtr<UDataflowSkeletalContent> SkeletalContent = Cast<UDataflowSkeletalContent>(GetEditorContent()))
 					{
-						if (EditorContent->GetDataflowAsset())
+						if (TSharedPtr<Dataflow::FEngineContext> Context = SkeletalContent->GetDataflowContext())
 						{
-							Dataflow::FTimestamp TimeStamp = EditorContent->GetLastModifiedTimestamp();
-							if (EditorContent->GetDataflowContext())
+							FManagedArrayCollection DefaultCollection;
+							const FManagedArrayCollection& Result = Output->GetValue(*Context, DefaultCollection);
+
+							SkeletalMesh = NewObject<USkeletalMesh>();
+							TObjectPtr<USkeleton> Skeleton = NewObject<USkeleton>(SkeletalMesh, Node->Name);
+
+							FGeometryCollectionEngineConversion::ConvertCollectionToSkeleton(Result, Skeleton, CollectionIndexRemap);
+							SkeletalMesh->SetSkeleton(Skeleton);
+							SkeletalMesh->SetRefSkeleton(Skeleton->GetReferenceSkeleton());
+
+							if (SkeletonEditor)
 							{
-								FManagedArrayCollection DefaultCollection;
-								const FManagedArrayCollection& Result = Output->GetValue(*EditorContent->GetDataflowContext(), DefaultCollection);
-
-								SkeletalMesh = NewObject<USkeletalMesh>();
-								TObjectPtr<USkeleton> Skeleton = NewObject<USkeleton>(SkeletalMesh, Node->Name);
-
-								FGeometryCollectionEngineConversion::ConvertCollectionToSkeleton(Result, Skeleton, CollectionIndexRemap);
-								SkeletalMesh->SetSkeleton(Skeleton);
-								SkeletalMesh->SetRefSkeleton(Skeleton->GetReferenceSkeleton());
-
-								if (SkeletonEditor)
-								{
-									SkeletonEditor->GetEditableSkeleton()->RecreateBoneTree(SkeletalMesh);
-									SkeletonEditor->SetSkeletalMesh(SkeletalMesh);
-									SkeletonEditor->Refresh();
-								}
-								bNeedsDefaultSkeleton = false;
+								SkeletonEditor->GetEditableSkeleton()->RecreateBoneTree(SkeletalMesh);
+								SkeletonEditor->SetSkeletalMesh(SkeletalMesh);
+								SkeletonEditor->Refresh();
 							}
+							bNeedsDefaultSkeleton = false;
 						}
 					}
 				}
@@ -134,14 +127,11 @@ void FDataflowSkeletonView::UpdateViewData()
 	if (bNeedsDefaultSkeleton)
 	{
 		SetSkeleton(NewObject<USkeleton>(SkeletalMesh, NAME_Name));
-		if (DataflowEditor)
+		if (const TObjectPtr<UDataflowSkeletalContent> SkeletalContent = Cast<UDataflowSkeletalContent>(GetEditorContent()))
 		{
-			if (const TObjectPtr<UDataflowSkeletalContent> SkeletalContent = Cast<UDataflowSkeletalContent>(DataflowEditor->GetDataflowContent()))
+			if (SkeletalContent->GetDataflowAsset())
 			{
-				if (SkeletalContent->GetDataflowAsset())
-				{
-					SetSkeleton(SkeletalContent->GetSkeleton());
-				}
+				SetSkeleton(SkeletalContent->GetSkeleton());
 			}
 		}
 		if (SkeletonEditor)
diff --git a/Engine/Plugins/Experimental/Dataflow/Source/DataflowEditor/Private/Dataflow/DataflowView.cpp b/Engine/Plugins/Experimental/Dataflow/Source/DataflowEditor/Private/Dataflow/DataflowView.cpp
--- a/Engine/Plugins/Experimental/Dataflow/Source/DataflowEditor/Private/Dataflow/DataflowView.cpp
+++ b/Engine/Plugins/Experimental/Dataflow/Source/DataflowEditor/Private/Dataflow/DataflowView.cpp
@@ -1,11 +1,20 @@
 // Copyright Epic Games, Inc. All Rights Reserved.
 
 #include "Dataflow/DataflowView.h"
-#include "Templates/EnableIf.h"
+
+#include "Dataflow/DataflowEditor.h"
 #include "Dataflow/DataflowSelection.h"
+#include "Templates/EnableIf.h"
 
 #define LOCTEXT_NAMESPACE "DataflowView"
 
+FDataflowNodeView::FDataflowNodeView(TObjectPtr<UDataflowBaseContent> InContent)
+	: FGCObject()
+	, EditorContent(InContent)
+{
+}
+
+
 FDataflowNodeView::~FDataflowNodeView()
 {
 	if (SelectedNode)
@@ -20,6 +29,14 @@ FDataflowNodeView::~FDataflowNodeView()
 	}
 }
 
+TObjectPtr<UDataflowBaseContent> FDataflowNodeView::GetEditorContent()
+{
+	if (ensure(EditorContent))
+	{
+		return EditorContent;
+	}
+	return nullptr;
+}
 
 bool FDataflowNodeView::SelectedNodeHaveSupportedOutputTypes(UDataflowEdNode* InNode)
 {
@@ -48,10 +65,6 @@ bool FDataflowNodeView::SelectedNodeHaveSupportedOutputTypes(UDataflowEdNode* In
 }
 
 
-void FDataflowNodeView::SetContext(TSharedPtr<Dataflow::FContext>& InContext)
-{
-	Context = InContext;
-}
 
 void FDataflowNodeView::OnConstructionViewSelectionChanged(const TArray<UPrimitiveComponent*>& InComponents)
 {
@@ -119,6 +132,10 @@ void FDataflowNodeView::OnNodeInvalidated(FDataflowNode* InvalidatedNode)
 void FDataflowNodeView::AddReferencedObjects(FReferenceCollector& Collector)
 {
 	Collector.AddReferencedObject(SelectedNode);
+	if (EditorContent)
+	{
+		Collector.AddReferencedObject(EditorContent);
+	}
 }
 
 
diff --git a/Engine/Plugins/Experimental/Dataflow/Source/DataflowEditor/Public/Dataflow/DataflowCollectionSpreadSheet.h b/Engine/Plugins/Experimental/Dataflow/Source/DataflowEditor/Public/Dataflow/DataflowCollectionSpreadSheet.h
--- a/Engine/Plugins/Experimental/Dataflow/Source/DataflowEditor/Public/Dataflow/DataflowCollectionSpreadSheet.h
+++ b/Engine/Plugins/Experimental/Dataflow/Source/DataflowEditor/Public/Dataflow/DataflowCollectionSpreadSheet.h
@@ -2,13 +2,12 @@
 
 #pragma once
 
-#include "Widgets/SCompoundWidget.h"
-#include "Widgets/Views/SListView.h"
-#include "Dataflow/DataflowEdNode.h"
-#include "Dataflow/DataflowCollectionSpreadSheetWidget.h"
-#include "Dataflow/DataflowSelection.h"
+#include "CoreMinimal.h"
 #include "Dataflow/DataflowView.h"
 
+class UDataflowEditor;
+class UPrimitiveComponent;
+class SCollectionSpreadSheetWidget;
 
 /**
 *
@@ -18,6 +17,8 @@
 class FDataflowCollectionSpreadSheet : public FDataflowNodeView
 {
 public:
+
+	FDataflowCollectionSpreadSheet(TObjectPtr<UDataflowBaseContent> InContent = nullptr);
 	~FDataflowCollectionSpreadSheet();
 
 	virtual void SetSupportedOutputTypes() override;
diff --git a/Engine/Plugins/Experimental/Dataflow/Source/DataflowEditor/Public/Dataflow/DataflowEditor.h b/Engine/Plugins/Experimental/Dataflow/Source/DataflowEditor/Public/Dataflow/DataflowEditor.h
--- a/Engine/Plugins/Experimental/Dataflow/Source/DataflowEditor/Public/Dataflow/DataflowEditor.h
+++ b/Engine/Plugins/Experimental/Dataflow/Source/DataflowEditor/Public/Dataflow/DataflowEditor.h
@@ -45,7 +45,7 @@ public:
 	virtual void Initialize(const TArray<TObjectPtr<UObject>>& InObjects) override;
 
 	/** Initialize an editor with a given content and an owner */
-	void InitializeContent(TObjectPtr<UDataflowBaseContent> BaseContent, const TObjectPtr<UObject>& ContentOwner);
+	void InitializeContent(const TObjectPtr<UObject>& ContentOwner);
 
 	TObjectPtr<UDataflowBaseContent> GetDataflowContent() { return DataflowContent; }
 	const TObjectPtr<UDataflowBaseContent> GetDataflowContent() const { return DataflowContent; }
diff --git a/Engine/Plugins/Experimental/Dataflow/Source/DataflowEditor/Public/Dataflow/DataflowSelectionView.h b/Engine/Plugins/Experimental/Dataflow/Source/DataflowEditor/Public/Dataflow/DataflowSelectionView.h
--- a/Engine/Plugins/Experimental/Dataflow/Source/DataflowEditor/Public/Dataflow/DataflowSelectionView.h
+++ b/Engine/Plugins/Experimental/Dataflow/Source/DataflowEditor/Public/Dataflow/DataflowSelectionView.h
@@ -2,13 +2,12 @@
 
 #pragma once
 
-#include "Widgets/SCompoundWidget.h"
-#include "Widgets/Views/SListView.h"
-#include "Dataflow/DataflowEdNode.h"
-#include "Dataflow/SelectionViewWidget.h"
-#include "Dataflow/DataflowSelection.h"
+#include "CoreMinimal.h"
 #include "Dataflow/DataflowView.h"
 
+class UDataflowEditor;
+class SSelectionViewWidget;
+class UPrimitiveComponent;
 
 /**
 *
@@ -18,6 +17,7 @@
 class FDataflowSelectionView : public FDataflowNodeView
 {
 public:
+	FDataflowSelectionView(TObjectPtr<UDataflowBaseContent> InContent = nullptr);
 	~FDataflowSelectionView();
 
 	virtual void SetSupportedOutputTypes() override;
diff --git a/Engine/Plugins/Experimental/Dataflow/Source/DataflowEditor/Public/Dataflow/DataflowSkeletonView.h b/Engine/Plugins/Experimental/Dataflow/Source/DataflowEditor/Public/Dataflow/DataflowSkeletonView.h
--- a/Engine/Plugins/Experimental/Dataflow/Source/DataflowEditor/Public/Dataflow/DataflowSkeletonView.h
+++ b/Engine/Plugins/Experimental/Dataflow/Source/DataflowEditor/Public/Dataflow/DataflowSkeletonView.h
@@ -20,7 +20,7 @@ struct FSkeletonTreeArgs;
 class FDataflowSkeletonView : public FDataflowNodeView
 {
 public:
-	FDataflowSkeletonView(UDataflowEditor* InDataflowEditor = nullptr);
+	FDataflowSkeletonView(TObjectPtr<UDataflowBaseContent> InContent = nullptr);
 	~FDataflowSkeletonView();
 
 	/** Create the Skeleton Tree Editor*/
@@ -49,8 +49,6 @@ public:
 
 
 private:
-	UDataflowEditor* DataflowEditor = nullptr;
-
 	TSharedPtr<ISkeletonTree> SkeletonEditor;
 
 	/* Skeletal Mesh in the SkeletalViewer*/
diff --git a/Engine/Plugins/Experimental/Dataflow/Source/DataflowEditor/Public/Dataflow/DataflowView.h b/Engine/Plugins/Experimental/Dataflow/Source/DataflowEditor/Public/Dataflow/DataflowView.h
--- a/Engine/Plugins/Experimental/Dataflow/Source/DataflowEditor/Public/Dataflow/DataflowView.h
+++ b/Engine/Plugins/Experimental/Dataflow/Source/DataflowEditor/Public/Dataflow/DataflowView.h
@@ -9,7 +9,8 @@
 #include "Dataflow/DataflowSelection.h"
 
 class UPrimitiveComponent;
-
+class UDataflowEditor;
+class UDataflowBaseContent;
 /**
 *
 * Base listener class to interface between the DataflowToolkit and Dataflow views
@@ -29,9 +30,10 @@ public:
 * FDataflowNodeView class implements common functions for single node based Dataflow views
 *
 */
-class FDataflowNodeView : public IDataflowViewListener, public FGCObject
+class FDataflowNodeView : public FGCObject, public IDataflowViewListener
 {
 public:
+	FDataflowNodeView(TObjectPtr<UDataflowBaseContent> InContent = nullptr);
 	virtual ~FDataflowNodeView();
 
 	UDataflowEdNode* GetSelectedNode() const { return SelectedNode; }
@@ -39,8 +41,7 @@ public:
 
 	TArray<FString>& GetSupportedOutputTypes() { return SupportedOutputTypes; }
 
-	TSharedPtr<Dataflow::FContext> GetContext() { return Context; }
-	void SetContext(TSharedPtr<Dataflow::FContext>& InContext);
+	TObjectPtr<UDataflowBaseContent> GetEditorContent();
 
 	/**
 	* Virtual functions to overwrite in view widget classes
@@ -72,9 +73,12 @@ public:
 	virtual FString GetReferencerName() const override { return TEXT("FDataflowNodeView"); }
 
 private:
+	TObjectPtr<UDataflowBaseContent> EditorContent = nullptr;
+
 	TObjectPtr<UDataflowEdNode> SelectedNode = nullptr;
-	TSharedPtr<Dataflow::FContext> Context;
+
 	bool bIsPinnedDown = false;
+
 	bool bIsRefreshLocked = false;
 
 	TArray<FString> SupportedOutputTypes;
