Dataflow : Set connected properties as read only in the node details panel

diff --git a/Engine/Plugins/Experimental/Dataflow/Source/DataflowEditor/Private/Dataflow/DataflowEditorCommands.cpp b/Engine/Plugins/Experimental/Dataflow/Source/DataflowEditor/Private/Dataflow/DataflowEditorCommands.cpp
--- a/Engine/Plugins/Experimental/Dataflow/Source/DataflowEditor/Private/Dataflow/DataflowEditorCommands.cpp
+++ b/Engine/Plugins/Experimental/Dataflow/Source/DataflowEditor/Private/Dataflow/DataflowEditorCommands.cpp
@@ -17,6 +17,7 @@
 #include "EdGraphNode_Comment.h"
 #include "EdGraph/EdGraphNode.h"
 #include "Editor.h"
+#include "IDetailsView.h"
 #include "IStructureDetailsView.h"
 #include "Serialization/ObjectWriter.h"
 #include "Serialization/ObjectReader.h"
@@ -426,6 +427,17 @@ void FDataflowEditorCommands::OnSelectedNodesChanged(TSharedPtr<IStructureDetail
 					{
 						if (TSharedPtr<FDataflowNode> DataflowNode = DataflowGraph->FindBaseNode(EdNode->GetDataflowNodeGuid()))
 						{
+							FIsPropertyReadOnly Delegate = FIsPropertyReadOnly::CreateLambda(
+								[DataflowNode](const FPropertyAndParent& PropertyAndParent)
+								{
+									if (const FDataflowInput* Input = DataflowNode->FindInput(PropertyAndParent.Property.GetFName()))
+									{
+										return (Input->GetConnection() != nullptr);
+									}
+									return false;
+								});
+							PropertiesEditor->GetDetailsView()->SetIsPropertyReadOnlyDelegate(Delegate);
+
 							TSharedPtr<FStructOnScope> Struct(DataflowNode->NewStructOnScope());
 							PropertiesEditor->SetStructureData(Struct);
 						}
