Dataflow Editor:
- Fixed crash caused by a Comment Node properties change deleting the context.
- Fixed a missing node invalidation when array properties have elements added, moved, duplicated, or removed.

#tests Verified that the crash didn't happened. Tested a few ways to invalidate nodes (forcing evaluation, changing properties, reloading cloth assets).
#rnx

diff --git a/Engine/Plugins/Experimental/Dataflow/Source/DataflowEditor/Private/Dataflow/DataflowEditorCommands.cpp b/Engine/Plugins/Experimental/Dataflow/Source/DataflowEditor/Private/Dataflow/DataflowEditorCommands.cpp
--- a/Engine/Plugins/Experimental/Dataflow/Source/DataflowEditor/Private/Dataflow/DataflowEditorCommands.cpp
+++ b/Engine/Plugins/Experimental/Dataflow/Source/DataflowEditor/Private/Dataflow/DataflowEditorCommands.cpp
@@ -340,43 +340,41 @@ void FDataflowEditorCommands::OnAssetPropertyValueChanged(TObjectPtr<UDataflowBa
 
 void FDataflowEditorCommands::OnPropertyValueChanged(UDataflow* OutDataflow, TSharedPtr<Dataflow::FEngineContext>& Context, Dataflow::FTimestamp& OutLastNodeTimestamp, const FPropertyChangedEvent& InPropertyChangedEvent, const TSet<TObjectPtr<UObject> >& SelectedNodes)
 {
-	if (InPropertyChangedEvent.ChangeType == EPropertyChangeType::ValueSet)
+	switch (InPropertyChangedEvent.ChangeType)
 	{
-		TSharedPtr<const FDataflowNode> UpdatedNode = nullptr;
-		if (OutDataflow && InPropertyChangedEvent.Property && InPropertyChangedEvent.Property->GetOwnerUObject())
+	case EPropertyChangeType::ValueSet:
+	case EPropertyChangeType::ArrayAdd:
+	case EPropertyChangeType::ArrayRemove:
+	case EPropertyChangeType::ArrayClear:
+	case EPropertyChangeType::ArrayMove:
+	case EPropertyChangeType::Duplicate:
+		if (ensure(OutDataflow && InPropertyChangedEvent.Property && InPropertyChangedEvent.Property->GetOwnerUObject()))
 		{
-			//			OutDataflow->MarkPackageDirty();
-			OutDataflow->Modify();
+			OutDataflow->Modify();  // Modify must be called even if SelectedNodes is empty because comment nodes aren't part of the selection set but still have properties
 
-			for (UObject* SelectedNode : SelectedNodes)
+			for (UObject* const SelectedNode : SelectedNodes)
 			{
-				if (UDataflowEdNode* Node = Cast<UDataflowEdNode>(SelectedNode))
+				UDataflowEdNode* const Node = CastChecked<UDataflowEdNode>(SelectedNode);  // SelectedNodes is assumed to be the graph panel selection filtered to UDataflowEdNode
+
+				if (TSharedPtr<FDataflowNode> DataflowNode = Node->GetDataflowNode())
 				{
-					if (TSharedPtr<FDataflowNode> DataflowNode = Node->GetDataflowNode())
-					{
-						UpdatedNode = DataflowNode;
-						DataflowNode->Invalidate();
+					DataflowNode->Invalidate();
+					OutLastNodeTimestamp = Dataflow::FTimestamp::Invalid;
 
-						// Reflect the active state on the drawing of the node
-						if (DataflowNode->bActive != Node->IsNodeEnabled())
-						{
-							Node->SetEnabledState(DataflowNode->bActive ? ENodeEnabledState::Enabled : ENodeEnabledState::Disabled);
-						}
+					// Reflect the active state on the drawing of the node
+					if (DataflowNode->bActive != Node->IsNodeEnabled())
+					{
+						Node->SetEnabledState(DataflowNode->bActive ? ENodeEnabledState::Enabled : ENodeEnabledState::Disabled);
 					}
 				}
 			}
 		}
-
-		if (!UpdatedNode && Context)
-		{
-			// Some base properties dont link back to the parent, so just clobber the cache for now. 
-			Context.Reset();
-		}
-		OutLastNodeTimestamp = Dataflow::FTimestamp::Invalid;
+		break;
+	default:
+		break;
 	}
 }
 
-
 void FDataflowEditorCommands::DeleteNodes(UDataflow* Graph, const FGraphPanelSelectionSet& SelectedNodes)
 {
 	if (ensureMsgf(Graph != nullptr, TEXT("Warning : Failed to find valid graph.")))
