Dataflow editor: set node output visibility based on node selection. Set wireframe view based on toggle button on the node

diff --git a/Engine/Plugins/Experimental/Dataflow/Source/DataflowEditor/Private/Dataflow/DataflowConstructionScene.cpp b/Engine/Plugins/Experimental/Dataflow/Source/DataflowEditor/Private/Dataflow/DataflowConstructionScene.cpp
--- a/Engine/Plugins/Experimental/Dataflow/Source/DataflowEditor/Private/Dataflow/DataflowConstructionScene.cpp
+++ b/Engine/Plugins/Experimental/Dataflow/Source/DataflowEditor/Private/Dataflow/DataflowConstructionScene.cpp
@@ -14,9 +14,6 @@
 
 #define LOCTEXT_NAMESPACE "FDataflowConstructionScene"
 
-bool bDataflowShowWireframeInConstructionView = false;
-FAutoConsoleVariableRef CVARDataflowShowWireframeInConstructionView(TEXT("p.Dataflow.Editor.Construction.ShowWireframe"), bDataflowShowWireframeInConstructionView, TEXT("Show the wireframe model in the dataflows construction view[def:true]"));
-
 //
 // Construction Scene
 //
@@ -27,6 +24,7 @@ FDataflowConstructionScene::FDataflowConstructionScene(FPreviewScene::Constructi
 
 FDataflowConstructionScene::~FDataflowConstructionScene()
 {
+	ResetWireframeMeshElementsVisualizer();
 	ResetDynamicMeshComponents();
 }
 
@@ -145,8 +143,8 @@ void FDataflowConstructionScene::UpdateDynamicMeshComponents()
 
 					Target->Render(Facade, DataflowContext);
 
-					int32 NumGeometry = Facade.NumGeometry();
-					for (int32 MeshIndex = 0; MeshIndex < NumGeometry; MeshIndex++)
+					const int32 NumGeometry = Facade.NumGeometry();
+					for (int32 MeshIndex = 0; MeshIndex < NumGeometry; ++MeshIndex)
 					{
 						FDynamicMesh3 DynamicMesh;
 						Dataflow::Conversion::RenderingFacadeToDynamicMesh(Facade, MeshIndex, DynamicMesh);
@@ -163,6 +161,56 @@ void FDataflowConstructionScene::UpdateDynamicMeshComponents()
 					}
 				}
 			}
+
+			// Add hidden DynamicMeshComponents for any targets that we want to render in wireframe
+			// 
+			// Note: UMeshElementsVisualizers need source meshes to pull from. We add invisible dynamic mesh components to the existing DynamicMeshComponents collection
+			// for this purpose, but could have instead created a separate collection of meshes for wireframe rendering. We are choosing to keep all the scene DynamicMeshComponents 
+			// in one place and using separate structures to dictate how they are used (MeshComponentsForWireframeRendering in this case), in case visualization requirements 
+			// change in the future.
+			//
+
+			MeshComponentsForWireframeRendering.Reset();
+			for (TObjectPtr<const UDataflowEdNode> Target : DataflowAsset->GetWireframeRenderTargets())
+			{
+				if (Target)
+				{
+					TSharedPtr<FManagedArrayCollection> RenderCollection(new FManagedArrayCollection);
+					GeometryCollection::Facades::FRenderingFacade Facade(*RenderCollection);
+					Facade.DefineSchema();
+
+					Target->Render(Facade, DataflowContext);
+
+					const int32 NumGeometry = Facade.NumGeometry();
+					for (int32 MeshIndex = 0; MeshIndex < NumGeometry; ++MeshIndex)
+					{
+						FDataflowRenderKey WireframeDynamicMeshKey{ Target, MeshIndex };
+
+						if (DynamicMeshComponents.Contains(WireframeDynamicMeshKey))
+						{
+							UDynamicMeshComponent* const ExistingMeshComponent = DynamicMeshComponents[WireframeDynamicMeshKey];
+							MeshComponentsForWireframeRendering.Add(ExistingMeshComponent);
+						}
+						else
+						{
+							FDynamicMesh3 DynamicMesh;
+							Dataflow::Conversion::RenderingFacadeToDynamicMesh(Facade, MeshIndex, DynamicMesh);
+
+							if (DynamicMesh.VertexCount())
+							{
+								if (Target == DataflowContent->GetPrimarySelectedNode())
+								{
+									DataflowContent->SetPrimaryRenderCollection(RenderCollection);
+								}
+								const FString MeshName = Facade.GetGeometryName()[MeshIndex];
+								UDynamicMeshComponent* const NewDynamicMeshComponent = AddDynamicMeshComponent(WireframeDynamicMeshKey, MeshName, MoveTemp(DynamicMesh), {});
+								NewDynamicMeshComponent->SetVisibility(false);
+								MeshComponentsForWireframeRendering.Add(NewDynamicMeshComponent);
+							}
+						}
+					}
+				}
+			}
 		}
 	}
 }
@@ -187,9 +235,12 @@ void FDataflowConstructionScene::ResetDynamicMeshComponents()
 
 TObjectPtr<UDynamicMeshComponent>& FDataflowConstructionScene::AddDynamicMeshComponent(FDataflowRenderKey InKey, const FString& MeshName, UE::Geometry::FDynamicMesh3&& DynamicMesh, const TArray<UMaterialInterface*>& MaterialSet)
 {
-	TObjectPtr<UDataflowEditorCollectionComponent> DynamicMeshComponent = NewObject<UDataflowEditorCollectionComponent>(RootSceneActor, FName(MeshName));
+	const FName UniqueObjectName = MakeUniqueObjectName(RootSceneActor, UDataflowEditorCollectionComponent::StaticClass(), FName(MeshName));
+	
+	TObjectPtr<UDataflowEditorCollectionComponent> DynamicMeshComponent = NewObject<UDataflowEditorCollectionComponent>(RootSceneActor, UniqueObjectName);
+
 	DynamicMeshComponent->MeshIndex = InKey.Value;
-	DynamicMeshComponent->Node = InKey.Key;;
+	DynamicMeshComponent->Node = InKey.Key;
 	DynamicMeshComponent->SetMesh(MoveTemp(DynamicMesh));
 	
 	// @todo(Material) This is just to have a material, we should transfer the materials from the assets if they have them. 
@@ -222,8 +273,6 @@ TObjectPtr<UDynamicMeshComponent>& FDataflowConstructionScene::AddDynamicMeshCom
 
 void FDataflowConstructionScene::AddWireframeMeshElementsVisualizer()
 {
-	if(!bDataflowShowWireframeInConstructionView) return;
-
 	ensure(WireframeElements.Num()==0);
 	for(FRenderElement Elem : DynamicMeshComponents)
 	{
@@ -235,6 +284,8 @@ void FDataflowConstructionScene::AddWireframeMeshElementsVisualizer()
 			WireframeElements.Add(DynamicMeshComponent, WireframeDraw);
 
 			WireframeDraw->CreateInWorld(GetWorld(), FTransform::Identity);
+			checkf(WireframeDraw->Settings, TEXT("Expected UMeshElementsVisualizer::Settings to exist after CreateInWorld"));
+
 			WireframeDraw->Settings->DepthBias = 2.0;
 			WireframeDraw->Settings->bAdjustDepthBiasUsingMeshSize = false;
 			WireframeDraw->Settings->bShowWireframe = true;
@@ -254,10 +305,9 @@ void FDataflowConstructionScene::AddWireframeMeshElementsVisualizer()
 					{
 						WireframeDraw->NotifyMeshChanged();
 					}));
-
-				const bool bRestSpaceMeshVisible = RenderElement.Value->GetVisibleFlag();
-				WireframeDraw->Settings->bVisible = bRestSpaceMeshVisible && bConstructionViewWireframe;
 			}
+
+			WireframeDraw->Settings->bVisible = false;
 			PropertyObjectsToTick.Add(WireframeDraw->Settings);
 		}
 	}
@@ -314,6 +364,11 @@ void FDataflowConstructionScene::UpdateConstructionScene()
 	// Attach a wireframe renderer to the DynamicMeshComponents
 	UpdateWireframeMeshElementsVisualizer();
 
+	for (const UDynamicMeshComponent* const DynamicMeshComponent : MeshComponentsForWireframeRendering)
+	{
+		WireframeElements[DynamicMeshComponent]->Settings->bVisible = true;
+	}
+
 	if (TObjectPtr<UDataflowBaseContent> DataflowContent = GetDataflowContent())
 	{
 		DataflowContent->SetIsDirty(false);
diff --git a/Engine/Plugins/Experimental/Dataflow/Source/DataflowEditor/Private/Dataflow/DataflowEditorToolkit.cpp b/Engine/Plugins/Experimental/Dataflow/Source/DataflowEditor/Private/Dataflow/DataflowEditorToolkit.cpp
--- a/Engine/Plugins/Experimental/Dataflow/Source/DataflowEditor/Private/Dataflow/DataflowEditorToolkit.cpp
+++ b/Engine/Plugins/Experimental/Dataflow/Source/DataflowEditor/Private/Dataflow/DataflowEditorToolkit.cpp
@@ -632,6 +632,14 @@ void FDataflowEditorToolkit::OnNodeSelectionChanged(const TSet<UObject*>& InNewS
 			return Objs;
 		};
 
+		for (UObject* const PreviouslySelectedNode : SelectedDataflowNodes)
+		{
+			if (UDataflowEdNode* PreviouslySelectedEdNode = Cast<UDataflowEdNode>(PreviouslySelectedNode))
+			{
+				PreviouslySelectedEdNode->SetShouldRenderNode(false);
+			}
+		}
+
 		// Only keep UDataflowEdNode from NewSelection
 		TSet< TObjectPtr<UObject> > NodeSelection = FindDataflowNodesInSet(AsObjectPointers(InNewSelection));
 
@@ -680,6 +688,14 @@ void FDataflowEditorToolkit::OnNodeSelectionChanged(const TSet<UObject*>& InNewS
 			SelectedDataflowNodes = NodeSelection;
 		}
 
+		for (UObject* const SelectedNode : NodeSelection)
+		{
+			if (UDataflowEdNode* SelectedEdNode = Cast<UDataflowEdNode>(SelectedNode))
+			{
+				SelectedEdNode->SetShouldRenderNode(true);
+			}
+		}
+
 		if (bPrimarySelectionChanged)
 		{
 			EditorContent->SetPrimarySelectedNode(nullptr);
diff --git a/Engine/Plugins/Experimental/Dataflow/Source/DataflowEditor/Private/Dataflow/DataflowSNode.cpp b/Engine/Plugins/Experimental/Dataflow/Source/DataflowEditor/Private/Dataflow/DataflowSNode.cpp
--- a/Engine/Plugins/Experimental/Dataflow/Source/DataflowEditor/Private/Dataflow/DataflowSNode.cpp
+++ b/Engine/Plugins/Experimental/Dataflow/Source/DataflowEditor/Private/Dataflow/DataflowSNode.cpp
@@ -54,7 +54,7 @@ void SDataflowEdNode::Construct(const FArguments& InArgs, UDataflowEdNode* InNod
 		.Style(&CheckBoxStyle)
 		.IsChecked_Lambda([this]()-> ECheckBoxState
 			{
-				if (DataflowGraphNode && DataflowGraphNode->DoAssetRender())
+				if (DataflowGraphNode && DataflowGraphNode->ShouldWireframeRenderNode())
 				{
 					return ECheckBoxState::Checked;
 				}
@@ -65,9 +65,9 @@ void SDataflowEdNode::Construct(const FArguments& InArgs, UDataflowEdNode* InNod
 				if (DataflowGraphNode)
 				{
 					if (NewState == ECheckBoxState::Checked)
-						DataflowGraphNode->SetAssetRender(true);
+						DataflowGraphNode->SetShouldWireframeRenderNode(true);
 					else
-						DataflowGraphNode->SetAssetRender(false);
+						DataflowGraphNode->SetShouldWireframeRenderNode(false);
 				}
 			});
 
diff --git a/Engine/Plugins/Experimental/Dataflow/Source/DataflowEditor/Public/Dataflow/DataflowConstructionScene.h b/Engine/Plugins/Experimental/Dataflow/Source/DataflowEditor/Public/Dataflow/DataflowConstructionScene.h
--- a/Engine/Plugins/Experimental/Dataflow/Source/DataflowEditor/Public/Dataflow/DataflowConstructionScene.h
+++ b/Engine/Plugins/Experimental/Dataflow/Source/DataflowEditor/Public/Dataflow/DataflowConstructionScene.h
@@ -82,6 +82,9 @@ private:
 	typedef TPair<TObjectPtr<UDynamicMeshComponent>, TObjectPtr<UMeshElementsVisualizer>> FRenderWireElement;
 	TMap<TObjectPtr<UDynamicMeshComponent>,TObjectPtr<UMeshElementsVisualizer> > WireframeElements;
 
+	/** Source meshes for Wireframe rendering */
+	TArray<UDynamicMeshComponent*> MeshComponentsForWireframeRendering;
+
 	/** Property objects (visible or not) that get ticked. [ticked] */
 	TArray<TObjectPtr<UInteractiveToolPropertySet>> PropertyObjectsToTick;
 
diff --git a/Engine/Source/Runtime/Experimental/Dataflow/Engine/Private/Dataflow/DataflowEdNode.cpp b/Engine/Source/Runtime/Experimental/Dataflow/Engine/Private/Dataflow/DataflowEdNode.cpp
--- a/Engine/Source/Runtime/Experimental/Dataflow/Engine/Private/Dataflow/DataflowEdNode.cpp
+++ b/Engine/Source/Runtime/Experimental/Dataflow/Engine/Private/Dataflow/DataflowEdNode.cpp
@@ -33,7 +33,7 @@ UDataflowEdNode::UDataflowEdNode(const FObjectInitializer& ObjectInitializer)
 #endif // WITH_EDITOR && !UE_BUILD_SHIPPING
 }
 
-void UDataflowEdNode::SetAssetRender(bool bInRender)
+void UDataflowEdNode::SetShouldRenderNode(bool bInRender)
 {
 	bRenderInAssetEditor = bInRender;
 	if (IsBound())
@@ -42,14 +42,38 @@ void UDataflowEdNode::SetAssetRender(bool bInRender)
 		if (UDataflow* DataflowObject = Cast< UDataflow>(GetGraph()))
 		{
 			if (bRenderInAssetEditor)
+			{
 				DataflowObject->AddRenderTarget(this);
+			}
 			else
+			{
 				DataflowObject->RemoveRenderTarget(this);
+			}
 		}
 #endif
 	}
 }
 
+void UDataflowEdNode::SetShouldWireframeRenderNode(bool bInRender)
+{
+	bRenderWireframeInAssetEditor = bInRender;
+	if (IsBound())
+	{
+#if WITH_EDITOR && !(UE_BUILD_SHIPPING || UE_BUILD_TEST)
+		if (UDataflow* const DataflowObject = Cast<UDataflow>(GetGraph()))
+		{
+			if (bRenderWireframeInAssetEditor)
+			{
+				DataflowObject->AddWireframeRenderTarget(this);
+			}
+			else
+			{
+				DataflowObject->RemoveWireframeRenderTarget(this);
+			}
+		}
+#endif
+	}
+}
 
 TSharedPtr<FDataflowNode> UDataflowEdNode::GetDataflowNode()
 {
diff --git a/Engine/Source/Runtime/Experimental/Dataflow/Engine/Private/Dataflow/DataflowObject.cpp b/Engine/Source/Runtime/Experimental/Dataflow/Engine/Private/Dataflow/DataflowObject.cpp
--- a/Engine/Source/Runtime/Experimental/Dataflow/Engine/Private/Dataflow/DataflowObject.cpp
+++ b/Engine/Source/Runtime/Experimental/Dataflow/Engine/Private/Dataflow/DataflowObject.cpp
@@ -91,10 +91,14 @@ void UDataflow::PostLoad()
 		{
 			if (DataflowEdNode)
 			{
-				if (DataflowEdNode->DoAssetRender())
+				if (DataflowEdNode->ShouldRenderNode())
 				{
 					RenderTargets.Add(DataflowEdNode);
 				}
+				if (DataflowEdNode->ShouldWireframeRenderNode())
+				{
+					WireframeRenderTargets.Add(DataflowEdNode);
+				}
 			}
 		}
 	}
@@ -104,20 +108,34 @@ void UDataflow::PostLoad()
 	UObject::PostLoad();
 }
 
-void UDataflow::AddRenderTarget(TObjectPtr<UDataflowEdNode> InNode)
+void UDataflow::AddRenderTarget(TObjectPtr<const UDataflowEdNode> InNode)
 {
 	LastModifiedRenderTarget = Dataflow::FTimestamp::Current();
-	InNode->bRenderInAssetEditor = true;
+	check(InNode->ShouldRenderNode());
 	RenderTargets.AddUnique(InNode);
 }
 
-void UDataflow::RemoveRenderTarget(TObjectPtr<UDataflowEdNode> InNode)
+void UDataflow::RemoveRenderTarget(TObjectPtr<const UDataflowEdNode> InNode)
 {
 	LastModifiedRenderTarget = Dataflow::FTimestamp::Current();
-	InNode->bRenderInAssetEditor = false;
+	check(!InNode->ShouldRenderNode());
 	RenderTargets.Remove(InNode);
 }
 
+void UDataflow::AddWireframeRenderTarget(TObjectPtr<const UDataflowEdNode> InNode)
+{
+	LastModifiedRenderTarget = Dataflow::FTimestamp::Current();
+	check(InNode->ShouldWireframeRenderNode());
+	WireframeRenderTargets.AddUnique(InNode);
+}
+
+void UDataflow::RemoveWireframeRenderTarget(TObjectPtr<const UDataflowEdNode> InNode)
+{
+	LastModifiedRenderTarget = Dataflow::FTimestamp::Current();
+	check(!InNode->ShouldWireframeRenderNode());
+	WireframeRenderTargets.Remove(InNode);
+}
+
 
 void UDataflow::Serialize(FArchive& Ar)
 {
diff --git a/Engine/Source/Runtime/Experimental/Dataflow/Engine/Public/Dataflow/DataflowEdNode.h b/Engine/Source/Runtime/Experimental/Dataflow/Engine/Public/Dataflow/DataflowEdNode.h
--- a/Engine/Source/Runtime/Experimental/Dataflow/Engine/Public/Dataflow/DataflowEdNode.h
+++ b/Engine/Source/Runtime/Experimental/Dataflow/Engine/Public/Dataflow/DataflowEdNode.h
@@ -90,12 +90,26 @@ public:
 	//
 	// Node Rendering
 	//
-	DATAFLOWENGINE_API void SetAssetRender(bool bInRender);
-	bool DoAssetRender() { return bRenderInAssetEditor; }
+	DATAFLOWENGINE_API void SetShouldRenderNode(bool bInRender);
+	bool ShouldRenderNode() const { return bRenderInAssetEditor; }
+
+	DATAFLOWENGINE_API void SetShouldWireframeRenderNode(bool bInRender);
+	bool ShouldWireframeRenderNode() const { return bRenderWireframeInAssetEditor; }
+
 	DATAFLOWENGINE_API TArray<Dataflow::FRenderingParameter> GetRenderParameters() const;
 	DATAFLOWENGINE_API virtual bool Render(GeometryCollection::Facades::FRenderingFacade& RenderData, const TSharedPtr<Dataflow::FContext> Context) const;
 
+private:
+
 	UPROPERTY()
 	bool bRenderInAssetEditor = false;
+
+	UPROPERTY()
+	bool bRenderWireframeInAssetEditor = false;
+
+	void OnPinConnection(const UEdGraphPin& Pin, const FEdGraphPinType& Type);
+	void OnPinDisconnection(const UEdGraphPin& Pin);
+	void SetAnyTypePinsToType(const FEdGraphPinType& Type);
+
 };
 
diff --git a/Engine/Source/Runtime/Experimental/Dataflow/Engine/Public/Dataflow/DataflowObject.h b/Engine/Source/Runtime/Experimental/Dataflow/Engine/Public/Dataflow/DataflowObject.h
--- a/Engine/Source/Runtime/Experimental/Dataflow/Engine/Public/Dataflow/DataflowObject.h
+++ b/Engine/Source/Runtime/Experimental/Dataflow/Engine/Public/Dataflow/DataflowObject.h
@@ -61,6 +61,7 @@ class UDataflow : public UEdGraph
 
 	Dataflow::FTimestamp LastModifiedRenderTarget = Dataflow::FTimestamp::Invalid; 
 	TArray< TObjectPtr<const UDataflowEdNode> > RenderTargets; // Not Serialized
+	TArray< TObjectPtr<const UDataflowEdNode> > WireframeRenderTargets; // Not Serialized
 	TSharedPtr<Dataflow::FGraph, ESPMode::ThreadSafe> Dataflow;
 	DATAFLOWENGINE_API void PostEditCallback();
 
@@ -113,9 +114,14 @@ public:
 	//
 	// Render Targets
 	//
-	DATAFLOWENGINE_API void AddRenderTarget(TObjectPtr<UDataflowEdNode>);
-	DATAFLOWENGINE_API void RemoveRenderTarget(TObjectPtr<UDataflowEdNode>);
+	DATAFLOWENGINE_API void AddRenderTarget(TObjectPtr<const UDataflowEdNode>);
+	DATAFLOWENGINE_API void RemoveRenderTarget(TObjectPtr<const UDataflowEdNode>);
 	const TArray< TObjectPtr<const UDataflowEdNode> >& GetRenderTargets() const { return RenderTargets; }
+
+	DATAFLOWENGINE_API void AddWireframeRenderTarget(TObjectPtr<const UDataflowEdNode>);
+	DATAFLOWENGINE_API void RemoveWireframeRenderTarget(TObjectPtr<const UDataflowEdNode>);
+	const TArray< TObjectPtr<const UDataflowEdNode> >& GetWireframeRenderTargets() const { return WireframeRenderTargets; }
+
 	const Dataflow::FTimestamp& GetRenderingTimestamp() const { return LastModifiedRenderTarget; }
 
 };
