3- Multiple dataflow terminal nodes

diff --git a/Engine/Plugins/Experimental/ChaosFlesh/Source/ChaosFleshEngine/Private/ChaosFlesh/FleshAsset.cpp b/Engine/Plugins/Experimental/ChaosFlesh/Source/ChaosFleshEngine/Private/ChaosFlesh/FleshAsset.cpp
--- a/Engine/Plugins/Experimental/ChaosFlesh/Source/ChaosFleshEngine/Private/ChaosFlesh/FleshAsset.cpp
+++ b/Engine/Plugins/Experimental/ChaosFlesh/Source/ChaosFleshEngine/Private/ChaosFlesh/FleshAsset.cpp
@@ -5,6 +5,8 @@
 =============================================================================*/
 #include "ChaosFlesh/FleshAsset.h"
 #include "ChaosFlesh/FleshCollection.h"
+#include "Components/SkeletalMeshComponent.h"
+#include "Engine/SkeletalMesh.h"
 #include "GeometryCollection/TransformCollection.h"
 
 #include UE_INLINE_GENERATED_CPP_BY_NAME(FleshAsset)
@@ -70,3 +72,54 @@ void UFleshAsset::Serialize(FArchive& Ar)
 	FleshCollection->Serialize(ChaosAr);
 }
 
+TObjectPtr<UDataflowBaseContent> UFleshAsset::CreateDataflowContent()
+{
+	TObjectPtr<UDataflowSkeletalContent> SkeletalContent = NewObject<UDataflowSkeletalContent>();
+
+	SkeletalContent->SetDataflowOwner(this);
+	SkeletalContent->SetTerminalAsset(this);
+
+	UpdateDataflowContent(SkeletalContent);
+	
+	return SkeletalContent;
+}
+
+void UFleshAsset::UpdateDataflowContent(const TObjectPtr<UDataflowBaseContent>& DataflowContent) const
+{
+	if(const TObjectPtr<UDataflowSkeletalContent> SkeletalContent = Cast<UDataflowSkeletalContent>(DataflowContent))
+	{
+		SkeletalContent->SetDataflowAsset(DataflowAsset);
+		SkeletalContent->SetDataflowTerminal(DataflowTerminal);
+		
+		SkeletalContent->SetSkeletalMesh(SkeletalMesh);
+		SkeletalContent->SetSkeleton(Skeleton);
+	}
+}
+
+#if WITH_EDITOR
+void UFleshAsset::PostEditChangeProperty(FPropertyChangedEvent& PropertyChangedEvent)
+{
+	Super::PostEditChangeProperty(PropertyChangedEvent);
+
+	const FName PropertyName = PropertyChangedEvent.Property->GetFName();
+    
+	if (PropertyName == GET_MEMBER_NAME_CHECKED(UFleshAsset, SkeletalMesh))
+	{
+		if(SkeletalMesh && (SkeletalMesh->GetSkeleton() != Skeleton))
+		{
+			Skeleton = SkeletalMesh->GetSkeleton();
+		}
+	}
+
+	if (PropertyName == GET_MEMBER_NAME_CHECKED(UFleshAsset, Skeleton))
+	{
+		if(SkeletalMesh && (SkeletalMesh->GetSkeleton() != Skeleton))
+		{
+			SkeletalMesh = nullptr;
+		}
+	}
+	InvalidateDataflowContents();
+}
+#endif //if WITH_EDITOR
+
+
diff --git a/Engine/Plugins/Experimental/ChaosFlesh/Source/ChaosFleshEngine/Public/ChaosFlesh/FleshAsset.h b/Engine/Plugins/Experimental/ChaosFlesh/Source/ChaosFleshEngine/Public/ChaosFlesh/FleshAsset.h
--- a/Engine/Plugins/Experimental/ChaosFlesh/Source/ChaosFleshEngine/Public/ChaosFlesh/FleshAsset.h
+++ b/Engine/Plugins/Experimental/ChaosFlesh/Source/ChaosFleshEngine/Public/ChaosFlesh/FleshAsset.h
@@ -7,6 +7,7 @@
 #include "ChaosFlesh/FleshCollection.h"
 #include "Dataflow/DataflowNodeParameters.h"
 #include "Dataflow/DataflowEngineTypes.h"
+#include "Dataflow/DataflowContent.h"
 #include "Engine/StaticMesh.h"
 #include "UObject/ObjectMacros.h"
 
@@ -50,7 +51,7 @@ private:
 *
 */
 UCLASS(customconstructor)
-class CHAOSFLESHENGINE_API UFleshAsset : public UObject
+class CHAOSFLESHENGINE_API UFleshAsset : public UObject, public IDataflowContentOwner
 {
 	GENERATED_UCLASS_BODY()
 	friend class FFleshAssetEdit;
@@ -80,6 +81,12 @@ public:
 	UFleshAsset(const FObjectInitializer& ObjectInitializer = FObjectInitializer::Get());
 
 
+#if WITH_EDITOR
+	/** Post edit change property */
+	virtual void PostEditChangeProperty(FPropertyChangedEvent& PropertyChangedEvent) override;
+#endif //if WITH_EDITOR
+
+
 	/**Editing the collection should only be through the edit object.*/
 	void SetCollection(FFleshCollection* InCollection);
 	const FFleshCollection* GetCollection() const { return FleshCollection.Get(); }
@@ -95,6 +102,10 @@ public:
 
 	void Serialize(FArchive& Ar);
 
+	/** IDataflowContentOwner interface */
+	virtual TObjectPtr<UDataflowBaseContent> CreateDataflowContent() override;
+	virtual void UpdateDataflowContent(const TObjectPtr<UDataflowBaseContent>& DataflowContent) const override;
+
 	//
 	// Dataflow
 	//
diff --git a/Engine/Plugins/Experimental/ChaosFlesh/Source/ChaosFleshNodes/Private/Dataflow/ChaosFleshFleshAssetTerminalNode.cpp b/Engine/Plugins/Experimental/ChaosFlesh/Source/ChaosFleshNodes/Private/Dataflow/ChaosFleshFleshAssetTerminalNode.cpp
--- a/Engine/Plugins/Experimental/ChaosFlesh/Source/ChaosFleshNodes/Private/Dataflow/ChaosFleshFleshAssetTerminalNode.cpp
+++ b/Engine/Plugins/Experimental/ChaosFlesh/Source/ChaosFleshNodes/Private/Dataflow/ChaosFleshFleshAssetTerminalNode.cpp
@@ -7,10 +7,10 @@
 
 void FFleshAssetTerminalDataflowNode::SetAssetValue(TObjectPtr<UObject> Asset, Dataflow::FContext& Context) const
 {
-	if (UFleshAsset* FleshAsset = Cast<UFleshAsset>(Asset.Get()))
+	if (UFleshAsset* InFleshAsset = Cast<UFleshAsset>(Asset.Get()))
 	{
 		const FManagedArrayCollection& InCollection = GetValue<FManagedArrayCollection>(Context, &Collection);
-		FleshAsset->SetCollection(InCollection.NewCopy<FFleshCollection>());
+		InFleshAsset->SetCollection(InCollection.NewCopy<FFleshCollection>());
 	}
 }
 
diff --git a/Engine/Plugins/Experimental/ChaosFlesh/Source/ChaosFleshNodes/Public/Dataflow/ChaosFleshFleshAssetTerminalNode.h b/Engine/Plugins/Experimental/ChaosFlesh/Source/ChaosFleshNodes/Public/Dataflow/ChaosFleshFleshAssetTerminalNode.h
--- a/Engine/Plugins/Experimental/ChaosFlesh/Source/ChaosFleshNodes/Public/Dataflow/ChaosFleshFleshAssetTerminalNode.h
+++ b/Engine/Plugins/Experimental/ChaosFlesh/Source/ChaosFleshNodes/Public/Dataflow/ChaosFleshFleshAssetTerminalNode.h
@@ -28,6 +28,12 @@ public:
 		RegisterInputConnection(&Collection);
 		RegisterOutputConnection(&Collection, &Collection);
 	}
+	
+	UPROPERTY(EditAnywhere, Category = "Dataflow", meta = (DisplayName = "FleshAsset"))
+	TObjectPtr<UFleshAsset> FleshAsset = nullptr;
+
+	/** Return the terminal asset */
+	virtual TObjectPtr<UObject> GetTerminalAsset() const override {return FleshAsset;}
 
 	virtual void SetAssetValue(TObjectPtr<UObject> Asset, Dataflow::FContext& Context) const override;
 	virtual void Evaluate(Dataflow::FContext& Context, const FDataflowOutput* Out) const override;
diff --git a/Engine/Plugins/Experimental/Dataflow/Source/DataflowEditor/Private/Dataflow/AssetDefinition_DataflowAsset.cpp b/Engine/Plugins/Experimental/Dataflow/Source/DataflowEditor/Private/Dataflow/AssetDefinition_DataflowAsset.cpp
--- a/Engine/Plugins/Experimental/Dataflow/Source/DataflowEditor/Private/Dataflow/AssetDefinition_DataflowAsset.cpp
+++ b/Engine/Plugins/Experimental/Dataflow/Source/DataflowEditor/Private/Dataflow/AssetDefinition_DataflowAsset.cpp
@@ -184,7 +184,7 @@ EAssetCommandResult UAssetDefinition_DataflowAsset::OpenAssets(const FAssetOpenA
 		// Validate the asset
 		if (UDataflow* const DataflowAsset = CastChecked<UDataflow>(DataflowObjects[0]))
 		{
-			AssetEditor->InitializeContent(DataflowAsset);
+			AssetEditor->Initialize({DataflowAsset});
 			return EAssetCommandResult::Handled;
 		}
 	}
diff --git a/Engine/Plugins/Experimental/Dataflow/Source/DataflowEditor/Private/Dataflow/AssetDefinition_DataflowContext.cpp b/Engine/Plugins/Experimental/Dataflow/Source/DataflowEditor/Private/Dataflow/AssetDefinition_DataflowContext.cpp
--- a/Engine/Plugins/Experimental/Dataflow/Source/DataflowEditor/Private/Dataflow/AssetDefinition_DataflowContext.cpp
+++ b/Engine/Plugins/Experimental/Dataflow/Source/DataflowEditor/Private/Dataflow/AssetDefinition_DataflowContext.cpp
@@ -166,23 +166,23 @@ namespace DataflowContextDefinitionHelpers
 	}
 
 	/*
-	* CreateNewDataflowContext
+	* CreateNewDataflowContent
 	*
 	*/	
 	template<class T>
-	TObjectPtr<T> CreateNewDataflowContext(const TObjectPtr<UObject>& ContentOwner)
+	TObjectPtr<T> CreateNewDataflowContent(const TObjectPtr<UObject>& ContentOwner)
 	{
 		check(ContentOwner.Get());
 		bool bNeedsNewAsset = true;
 		TObjectPtr<UObject> Asset = nullptr;
-
+		
+		UDataflow* DataflowAsset = Private::GetDataflowAssetFrom(ContentOwner);
 		if (bDataflowEnableContextCaching)
 		{
 			// Setup an asset that lives in the content broswer.
 			FContentBrowserModule& ContentBrowserModule = FModuleManager::LoadModuleChecked<FContentBrowserModule>("ContentBrowser");
 
 			UClass* DataflowClass = T::StaticClass();
-			UDataflow* DataflowAsset = Private::GetDataflowAssetFrom(ContentOwner);
 
 			const FString AssetPackageName = ContentOwner->GetOutermost()->GetName();
 			const FString AssetDefaultPath = FPackageName::GetLongPackagePath(AssetPackageName);
@@ -220,7 +220,8 @@ namespace DataflowContextDefinitionHelpers
 
 				if (UDataflowBaseContent* BaseContent = Cast< UDataflowBaseContent>(Asset.Get()))
 				{
-					BaseContent->BuildBaseContent(ContentOwner);
+					BaseContent->SetDataflowOwner(ContentOwner);
+					BaseContent->SetDataflowAsset(DataflowAsset);
 				}
 			}
 		}
@@ -229,7 +230,8 @@ namespace DataflowContextDefinitionHelpers
 			Asset = NewObject<T>(ContentOwner, T::StaticClass());
 			if (UDataflowBaseContent* BaseContent = Cast< UDataflowBaseContent>(Asset.Get()))
 			{
-				BaseContent->BuildBaseContent(ContentOwner);
+				BaseContent->SetDataflowOwner(ContentOwner);
+				BaseContent->SetDataflowAsset(DataflowAsset);
 			}
 		}
 
@@ -241,8 +243,8 @@ namespace DataflowContextDefinitionHelpers
 		return Cast<T>(Asset.Get());
 	}
 
-	template TObjectPtr<UDataflowBaseContent> CreateNewDataflowContext(const TObjectPtr<UObject>& ContentOwner);
-	template TObjectPtr<UDataflowSkeletalContent> CreateNewDataflowContext(const TObjectPtr<UObject>& ContentOwner);
+	template TObjectPtr<UDataflowBaseContent> CreateNewDataflowContent(const TObjectPtr<UObject>& ContentOwner);
+	template TObjectPtr<UDataflowSkeletalContent> CreateNewDataflowContent(const TObjectPtr<UObject>& ContentOwner);
 }
 
 
diff --git a/Engine/Plugins/Experimental/Dataflow/Source/DataflowEditor/Private/Dataflow/DataflowConstructionScene.cpp b/Engine/Plugins/Experimental/Dataflow/Source/DataflowEditor/Private/Dataflow/DataflowConstructionScene.cpp
--- a/Engine/Plugins/Experimental/Dataflow/Source/DataflowEditor/Private/Dataflow/DataflowConstructionScene.cpp
+++ b/Engine/Plugins/Experimental/Dataflow/Source/DataflowEditor/Private/Dataflow/DataflowConstructionScene.cpp
@@ -78,18 +78,18 @@ FORCEINLINE Dataflow::FTimestamp LatestTimestamp(const UDataflow* Dataflow, cons
 
 void FDataflowConstructionScene::TickDataflowScene(const float DeltaSeconds)
 {
-	if (TObjectPtr<UDataflowBaseContent> DataflowContent = GetDataflowContent())
+	if (const TObjectPtr<UDataflowBaseContent>& EditorContent = GetEditorContent())
 	{
-		if (const TSharedPtr<Dataflow::FContext> DataflowContext = DataflowContent->GetDataflowContext())
+		if (const TSharedPtr<Dataflow::FContext> DataflowContext = EditorContent->GetDataflowContext())
 		{
-			if (const UDataflow* Dataflow = DataflowContent->GetDataflowAsset())
+			if (const UDataflow* Dataflow = EditorContent->GetDataflowAsset())
 			{
 				const Dataflow::FTimestamp SystemTimestamp = LatestTimestamp(Dataflow, DataflowContext.Get());
-				if (SystemTimestamp >= DataflowContent->GetLastModifiedTimestamp() || DataflowContent->IsDirty())
+				if (SystemTimestamp >= EditorContent->GetLastModifiedTimestamp() || EditorContent->IsDirty())
 				{
-					DataflowContent->SetLastModifiedTimestamp(SystemTimestamp.Value + 1);
+					EditorContent->SetLastModifiedTimestamp(SystemTimestamp.Value + 1);
 
-					if (DataflowContent->IsDirty())
+					if (EditorContent->IsDirty())
 					{
 						UpdateConstructionScene();
 					}
@@ -127,10 +127,10 @@ void FDataflowConstructionScene::UpdateDynamicMeshComponents()
 	// list of UPrimitiveComponents for rendering.
 	ResetDynamicMeshComponents();
 
-	if (TObjectPtr<UDataflowBaseContent> DataflowContent = GetDataflowContent())
+	if (const TObjectPtr<UDataflowBaseContent>& EditorContent = GetEditorContent())
 	{
-		const TObjectPtr<UDataflow>& DataflowAsset = DataflowContent->GetDataflowAsset();
-		const TSharedPtr<Dataflow::FEngineContext>& DataflowContext = DataflowContent->GetDataflowContext();
+		const TObjectPtr<UDataflow>& DataflowAsset = EditorContent->GetDataflowAsset();
+		const TSharedPtr<Dataflow::FEngineContext>& DataflowContext = EditorContent->GetDataflowContext();
 		if(DataflowAsset && DataflowContext)
 		{
 			for (TObjectPtr<const UDataflowEdNode> Target : DataflowAsset->GetRenderTargets())
@@ -151,9 +151,9 @@ void FDataflowConstructionScene::UpdateDynamicMeshComponents()
 
 						if (DynamicMesh.VertexCount())
 						{
-							if (Target == DataflowContent->GetPrimarySelectedNode())
+							if (Target == EditorContent->GetPrimarySelectedNode())
 							{
-								DataflowContent->SetPrimaryRenderCollection(RenderCollection);
+								EditorContent->SetPrimaryRenderCollection(RenderCollection);
 							}
 							const FString MeshName = Facade.GetGeometryName()[MeshIndex];
 							AddDynamicMeshComponent({Target, MeshIndex}, MeshName, MoveTemp(DynamicMesh), {});
@@ -198,9 +198,9 @@ void FDataflowConstructionScene::UpdateDynamicMeshComponents()
 
 							if (DynamicMesh.VertexCount())
 							{
-								if (Target == DataflowContent->GetPrimarySelectedNode())
+								if (Target == EditorContent->GetPrimarySelectedNode())
 								{
-									DataflowContent->SetPrimaryRenderCollection(RenderCollection);
+									EditorContent->SetPrimaryRenderCollection(RenderCollection);
 								}
 								const FString MeshName = Facade.GetGeometryName()[MeshIndex];
 								UDynamicMeshComponent* const NewDynamicMeshComponent = AddDynamicMeshComponent(WireframeDynamicMeshKey, MeshName, MoveTemp(DynamicMesh), {});
@@ -244,10 +244,10 @@ TObjectPtr<UDynamicMeshComponent>& FDataflowConstructionScene::AddDynamicMeshCom
 	DynamicMeshComponent->SetMesh(MoveTemp(DynamicMesh));
 	
 	// @todo(Material) This is just to have a material, we should transfer the materials from the assets if they have them. 
-	TObjectPtr<UDataflowBaseContent> DataflowContent = GetDataflowContent();
-	if (DataflowContent && DataflowContent->GetDataflowAsset() && DataflowContent->GetDataflowAsset()->Material)
+	const  TObjectPtr<UDataflowBaseContent>& EditorContent = GetEditorContent();
+	if (EditorContent && EditorContent->GetDataflowAsset() && EditorContent->GetDataflowAsset()->Material)
 	{
-		DynamicMeshComponent->ConfigureMaterialSet({ DataflowContent->GetDataflowAsset()->Material });
+		DynamicMeshComponent->ConfigureMaterialSet({ EditorContent->GetDataflowAsset()->Material });
 	}
 	else
 	{
@@ -357,7 +357,7 @@ void FDataflowConstructionScene::UpdateConstructionScene()
 	ResetConstructionScene();
 
 	// The preview scene for the construction view will be
-	// cleared and rebuilt from scratch. This will genrate a 
+	// cleared and rebuilt from scratch. This will generate a 
 	// list of UPrimitiveComponents for rendering.
 	UpdateDynamicMeshComponents();
 	
@@ -369,9 +369,14 @@ void FDataflowConstructionScene::UpdateConstructionScene()
 		WireframeElements[DynamicMeshComponent]->Settings->bVisible = true;
 	}
 
-	if (TObjectPtr<UDataflowBaseContent> DataflowContent = GetDataflowContent())
+	if (const TObjectPtr<UDataflowBaseContent>& EditorContent = GetEditorContent())
+	{
+		EditorContent->SetIsDirty(false);
+	}
+
+	for(const TObjectPtr<UDataflowBaseContent>& TerminalContent : GetTerminalContents())
 	{
-		DataflowContent->SetIsDirty(false);
+		TerminalContent->SetIsDirty(false);
 	}
 }
 
diff --git a/Engine/Plugins/Experimental/Dataflow/Source/DataflowEditor/Private/Dataflow/DataflowConstructionViewport.cpp b/Engine/Plugins/Experimental/Dataflow/Source/DataflowEditor/Private/Dataflow/DataflowConstructionViewport.cpp
--- a/Engine/Plugins/Experimental/Dataflow/Source/DataflowEditor/Private/Dataflow/DataflowConstructionViewport.cpp
+++ b/Engine/Plugins/Experimental/Dataflow/Source/DataflowEditor/Private/Dataflow/DataflowConstructionViewport.cpp
@@ -80,15 +80,5 @@ void SDataflowConstructionViewport::OnFloatingButtonClicked()
 {
 }
 
-float SDataflowConstructionViewport:: GetViewMinInput() const
-{
-	return static_cast<FDataflowPreviewSceneBase*>(Client->GetPreviewScene())->GetDataflowContent()->GetSimulationRange()[0];
-}
-
-float SDataflowConstructionViewport::GetViewMaxInput() const
-{
-	return static_cast<FDataflowPreviewSceneBase*>(Client->GetPreviewScene())->GetDataflowContent()->GetSimulationRange()[1];
-}
-
 
 #undef LOCTEXT_NAMESPACE
diff --git a/Engine/Plugins/Experimental/Dataflow/Source/DataflowEditor/Private/Dataflow/DataflowEditor.cpp b/Engine/Plugins/Experimental/Dataflow/Source/DataflowEditor/Private/Dataflow/DataflowEditor.cpp
--- a/Engine/Plugins/Experimental/Dataflow/Source/DataflowEditor/Private/Dataflow/DataflowEditor.cpp
+++ b/Engine/Plugins/Experimental/Dataflow/Source/DataflowEditor/Private/Dataflow/DataflowEditor.cpp
@@ -26,79 +26,144 @@ void UDataflowEditor::Initialize(const TArray<TObjectPtr<UObject>>& InObjects)
 {
 	if(!InObjects.IsEmpty())
 	{
-		InitializeContent(InObjects[0]);
+		const TObjectPtr<UObject> ContentOwner = InObjects[0];
+		TArray<TObjectPtr<UObject>> RequiredObjects = { ContentOwner };
+		
+		if(!EditorContent)
+		{
+			if(UDataflow* DataflowAsset = Cast<UDataflow>(ContentOwner))
+			{
+				EditorContent = DataflowContextDefinitionHelpers::CreateNewDataflowContent<UDataflowBaseContent>(ContentOwner);
+				EditorContent->SetDataflowOwner(DataflowAsset);
+				EditorContent->SetDataflowAsset(DataflowAsset);
+			}
+			else
+			{
+				if(IDataflowContentOwner* EditorContentOwner = Cast<IDataflowContentOwner>(ContentOwner))
+				{
+					EditorContent = EditorContentOwner->BuildDataflowContent();
+					RequiredObjects.Add(EditorContent->GetDataflowAsset());
+				}
+			}
+		}
+		RequiredObjects.Add(EditorContent);
+
+		// Update the editor datas (skeleton information for the viewer)
+		UpdateEditorContent();
+
+		// Update and build the terminal contents
+		UpdateTerminalContents(Dataflow::FTimestamp::Invalid);
+		
+		// Potentially we could add additional objects to edit here (fields, meshes....)
+		// If these objects have a matching factory we would be able to use geometry tools
+		UBaseCharacterFXEditor::Initialize(RequiredObjects);
 	}
 }
 
-void UDataflowEditor::InitializeContent(const TObjectPtr<UObject>& ContentOwner)
+void UDataflowEditor::UpdateEditorContent()
 {
-	check(DataflowContent == nullptr);
-
-	TArray<TObjectPtr<UObject>> RequiredObjects = { ContentOwner };
-
-	if (UDataflow* DataflowAsset = Cast<UDataflow>(ContentOwner))
+	if(const TObjectPtr<UDataflowSkeletalContent> SkeletalContent = Cast<UDataflowSkeletalContent>(EditorContent))
 	{
-		DataflowContent = DataflowContextDefinitionHelpers::CreateNewDataflowContext<UDataflowBaseContent>(ContentOwner);
+		// Add a skeleton mesh in case we don't have one
+		if(!SkeletalContent->GetSkeleton() && !SkeletalContent->GetSkeletalMesh())
+		{
+			const FName SkeletalMeshName = MakeUniqueObjectName(SkeletalContent->GetDataflowAsset(), UDataflow::StaticClass(), FName("USkeletalMesh"));
+			USkeletalMesh* SkeletalMesh = NewObject<USkeletalMesh>(SkeletalContent->GetDataflowAsset(), SkeletalMeshName);
 
-		DataflowContent->SetDataflowAsset(DataflowAsset);
-		DataflowContent->SetDataflowTerminal(FString());
+			const FName SkeletonName = MakeUniqueObjectName(SkeletalContent->GetDataflowAsset(), UDataflow::StaticClass(), FName("USkeleton"));
+			USkeleton* Skeleton = NewObject<USkeleton>(SkeletalContent->GetDataflowAsset(), SkeletonName);
+			
+			SkeletalMesh->SetSkeleton(Skeleton);
+			SkeletalContent->SetSkeletalMesh(SkeletalMesh);
+		}
 	}
-	else
+	if(EditorContent && EditorContent->GetDataflowAsset())
 	{
-		if (Private::HasDataflowAsset(ContentOwner))
+		EditorContent->GetDataflowAsset()->Schema = UDataflowSchema::StaticClass();
+	}
+}
+
+void UDataflowEditor::RemoveTerminalContents(const TSharedPtr<Dataflow::FGraph>& DataflowGraph, ValidTerminalsType& ValidTerminals)
+{
+	for(int32 ContentIndex = TerminalContents.Num()-1; ContentIndex >= 0; --ContentIndex)
+	{
+		if(TSharedPtr<FDataflowNode> TerminalNode = DataflowGraph->FindTerminalNode(
+			FName(TerminalContents[ContentIndex]->GetDataflowTerminal())))
 		{
-			if (Private::HasSkeletalMesh(ContentOwner))
+			if(Cast<IDataflowContentOwner>(TerminalNode->AsType<FDataflowTerminalNode>()->GetTerminalAsset()))
 			{
-				DataflowContent = DataflowContextDefinitionHelpers::CreateNewDataflowContext<UDataflowSkeletalContent>(ContentOwner);
-				const TObjectPtr<UDataflowSkeletalContent> SkeletalContent = Cast<UDataflowSkeletalContent>(DataflowContent);
-
-				SkeletalContent->SetSkeletalMesh(Private::GetSkeletalMeshFrom(ContentOwner));
-				SkeletalContent->SetSkeleton(Private::GetSkeletonFrom(ContentOwner));
-				SkeletalContent->SetAnimationAsset(Private::GetAnimationAssetFrom(ContentOwner));
+				ValidTerminals.Add(TerminalNode, TerminalContents[ContentIndex]);
 			}
 			else
 			{
-				DataflowContent = DataflowContextDefinitionHelpers::CreateNewDataflowContext<UDataflowBaseContent>(ContentOwner);
+				TerminalContents.RemoveAt(ContentIndex);
+				bHasTerminalsDirty = true;
 			}
-
-			DataflowContent->SetDataflowAsset(Private::GetDataflowAssetFrom(ContentOwner));
-			DataflowContent->SetDataflowTerminal(Private::GetDataflowTerminalFrom(ContentOwner));
-			RequiredObjects.Add(Private::GetDataflowAssetFrom(ContentOwner));
+		}
+		else
+		{
+			// Removal of all the contents with invalid nodes
+			TerminalContents.RemoveAt(ContentIndex);
+			bHasTerminalsDirty = true;
 		}
 	}
+}
 
-	if (!DataflowContent) return;
-	RequiredObjects.Add(DataflowContent);
-
-	if (const TObjectPtr<UDataflowSkeletalContent> SkeletalContent = Cast<UDataflowSkeletalContent>(DataflowContent))
+void UDataflowEditor::AddTerminalContents(const TSharedPtr<Dataflow::FGraph>& DataflowGraph, ValidTerminalsType& ValidTerminals)
+{
+	auto BuildTerminalContent = [this](const TObjectPtr<UDataflowBaseContent>& TerminalContent,
+					const FString& TerminalName, const TObjectPtr<UObject>& TerminalAsset)
 	{
-		if (!SkeletalContent->GetSkeletalMesh())
+		// Set the context (owner/asset) onto the terminal content
+		TerminalContent->SetDataflowContext(EditorContent->GetDataflowContext());
+				
+		TerminalContent->SetDataflowTerminal(TerminalName);
+		TerminalContent->SetTerminalAsset(TerminalAsset);
+	};
+			
+	for(const TSharedPtr<FDataflowNode>& DataflowNode : DataflowGraph->GetTerminalNodes())
+	{
+		if(const FDataflowTerminalNode* TerminalNode = DataflowNode->AsType<FDataflowTerminalNode>())
 		{
-			const FName SkeletalMeshName = MakeUniqueObjectName(SkeletalContent->GetDataflowAsset(), UDataflow::StaticClass(), FName("USkeletalMesh"));
-			USkeletalMesh* SkeletalMesh = NewObject<USkeletalMesh>(SkeletalContent->GetDataflowAsset(), SkeletalMeshName);
-
-			USkeleton* Skeleton = SkeletalContent->GetSkeleton();
-			if (!Skeleton)
+			if(IDataflowContentOwner* TerminalOwner = Cast<IDataflowContentOwner>(TerminalNode->GetTerminalAsset()))
 			{
-				const FName SkeletonName = MakeUniqueObjectName(SkeletalContent->GetDataflowAsset(), UDataflow::StaticClass(), FName("USkeleton"));
-				Skeleton = NewObject<USkeleton>(SkeletalContent->GetDataflowAsset(), SkeletonName);
+				const TObjectPtr<UDataflowBaseContent>* TerminalContent = ValidTerminals.Find(DataflowNode);
+				if(!TerminalContent)
+				{
+					TerminalContents.Add(TerminalOwner->BuildDataflowContent());
+					TerminalContent = &TerminalContents.Last();
+							
+					(*TerminalContent)->SetLastModifiedTimestamp(EditorContent->GetLastModifiedTimestamp());
+					bHasTerminalsDirty = true;
+				}
+				if(TerminalNode->GetTerminalAsset() != (*TerminalContent)->GetTerminalAsset())
+				{
+					BuildTerminalContent(*TerminalContent,
+							TerminalNode->GetName().ToString(), TerminalNode->GetTerminalAsset());
+					bHasTerminalsDirty = true;
+				}
 			}
-			SkeletalMesh->SetSkeleton(Skeleton);
-			SkeletalContent->SetSkeletalMesh(SkeletalMesh);
-		}
-		else if (!SkeletalContent->GetSkeleton())
-		{
-			SkeletalContent->SetSkeleton(SkeletalContent->GetSkeletalMesh()->GetSkeleton());
 		}
 	}
+}
 
-	if(DataflowContent && DataflowContent->GetDataflowAsset())
+void UDataflowEditor::UpdateTerminalContents(const Dataflow::FTimestamp TimeStamp)
+{
+	bHasTerminalsDirty = false;
+	
+	// update of the terminal contents only if no terminal asset on the main editor content
+	if(EditorContent && !EditorContent->GetTerminalAsset() && EditorContent->GetDataflowAsset())
 	{
-		DataflowContent->GetDataflowAsset()->Schema = UDataflowSchema::StaticClass();
+		if(const TSharedPtr<Dataflow::FGraph> DataflowGraph = EditorContent->GetDataflowAsset()->GetDataflow())
+		{
+			ValidTerminalsType ValidTerminals;
+
+			// Remove invalid terminals
+			RemoveTerminalContents(DataflowGraph, ValidTerminals);
+
+			// Add valid terminals
+			AddTerminalContents(DataflowGraph, ValidTerminals);
+		}
 	}
-	
-	// Potentially we could add additional objects to edit here (fields, meshes....)
-	// If these objects have a matching factory we would be able to use geometry tools
-	UBaseCharacterFXEditor::Initialize(RequiredObjects);
 }
 
diff --git a/Engine/Plugins/Experimental/Dataflow/Source/DataflowEditor/Private/Dataflow/DataflowEditorMode.cpp b/Engine/Plugins/Experimental/Dataflow/Source/DataflowEditor/Private/Dataflow/DataflowEditorMode.cpp
--- a/Engine/Plugins/Experimental/Dataflow/Source/DataflowEditor/Private/Dataflow/DataflowEditorMode.cpp
+++ b/Engine/Plugins/Experimental/Dataflow/Source/DataflowEditor/Private/Dataflow/DataflowEditorMode.cpp
@@ -334,7 +334,7 @@ void UDataflowEditorMode::SetDataflowSimlationScene(FDataflowSimulationScene* In
 void UDataflowEditorMode::CreateToolTargets(const TArray<TObjectPtr<UObject>>& AssetsIn)
 {
 	ToolTargets.Reset();
-	if (TObjectPtr<UDataflowBaseContent> EditorContent = ConstructionScene->GetDataflowContent())
+	if (const TObjectPtr<UDataflowBaseContent>& EditorContent = ConstructionScene->GetEditorContent())
 	{
 		if (UToolTarget* Target = GetInteractiveToolsContext()->TargetManager->BuildTarget(EditorContent, GetToolTargetRequirements()))
 		{
@@ -659,14 +659,14 @@ void UDataflowEditorMode::InitializeContextObject()
 {
 	check(ConstructionScene);
 
-	if (TObjectPtr<UDataflowBaseContent> DataflowContent = ConstructionScene->GetDataflowContent())
+	if (const TObjectPtr<UDataflowBaseContent>& EditorContent = ConstructionScene->GetEditorContent())
 	{
-		UEditorInteractiveToolsContext* const ConstructionToolsContext = GetInteractiveToolsContext();
+		const UEditorInteractiveToolsContext* const ConstructionToolsContext = GetInteractiveToolsContext();
 
 		UDataflowContextObject* ContextObject = ConstructionToolsContext->ContextObjectStore->FindContext<UDataflowContextObject>();
 		if (!ContextObject)
 		{
-			ContextObject = DataflowContent;
+			ContextObject = EditorContent;
 			ConstructionToolsContext->ContextObjectStore->AddContextObject(ContextObject);
 		}
 
@@ -757,9 +757,9 @@ UEdGraphNode* UDataflowEditorMode::CreateNewNode(const FName& NewNodeTypeName)
 		return nullptr;
 	}
 
-	if (TObjectPtr<UDataflowBaseContent> EditorContent = ConstructionScene->GetDataflowContent())
+	if (const TObjectPtr<UDataflowBaseContent>& EditorContent = ConstructionScene->GetEditorContent())
 	{
-		if (TObjectPtr<UDataflow> DataflowGraph = EditorContent->GetDataflowAsset())
+		if (const TObjectPtr<UDataflow>& DataflowGraph = EditorContent->GetDataflowAsset())
 		{
 			const TSharedPtr<FAssetSchemaAction_Dataflow_CreateNode_DataflowEdNode> NodeAction =
 				FAssetSchemaAction_Dataflow_CreateNode_DataflowEdNode::CreateAction(DataflowGraph, NewNodeTypeName);
@@ -774,9 +774,9 @@ UEdGraphNode* UDataflowEditorMode::CreateNewNode(const FName& NewNodeTypeName)
 
 UEdGraphNode* UDataflowEditorMode::CreateAndConnectNewNode(const FName& NewNodeTypeName, UEdGraphNode& UpstreamNode, const FName& ConnectionTypeName, const FName& NewNodeConnectionName)
 {
-	if (TObjectPtr<UDataflowBaseContent> EditorContent = ConstructionScene->GetDataflowContent())
+	if (const TObjectPtr<UDataflowBaseContent>& EditorContent = ConstructionScene->GetEditorContent())
 	{
-		if (TObjectPtr<UDataflow> DataflowGraph = EditorContent->GetDataflowAsset())
+		if (const TObjectPtr<UDataflow>& DataflowGraph = EditorContent->GetDataflowAsset())
 		{
 			// First find the specified output of the upstream node, plus any pins it's connected to
 
diff --git a/Engine/Plugins/Experimental/Dataflow/Source/DataflowEditor/Private/Dataflow/DataflowEditorModeToolkit.cpp b/Engine/Plugins/Experimental/Dataflow/Source/DataflowEditor/Private/Dataflow/DataflowEditorModeToolkit.cpp
--- a/Engine/Plugins/Experimental/Dataflow/Source/DataflowEditor/Private/Dataflow/DataflowEditorModeToolkit.cpp
+++ b/Engine/Plugins/Experimental/Dataflow/Source/DataflowEditor/Private/Dataflow/DataflowEditorModeToolkit.cpp
@@ -276,9 +276,9 @@ void FDataflowEditorModeToolkit::OnToolEnded(UInteractiveToolManager* Manager, U
 
 	if (FDataflowConstructionScene* ConstructionScene = GetConstructionScene())
 	{
-		if (TObjectPtr<UDataflowBaseContent> Content = ConstructionScene->GetDataflowContent())
+		if (const TObjectPtr<UDataflowBaseContent>& EditorContent = ConstructionScene->GetEditorContent())
 		{
-			Content->SetIsDirty(true);
+			EditorContent->SetIsDirty(true);
 		}
 	}
 
diff --git a/Engine/Plugins/Experimental/Dataflow/Source/DataflowEditor/Private/Dataflow/DataflowEditorPreviewSceneBase.cpp b/Engine/Plugins/Experimental/Dataflow/Source/DataflowEditor/Private/Dataflow/DataflowEditorPreviewSceneBase.cpp
--- a/Engine/Plugins/Experimental/Dataflow/Source/DataflowEditor/Private/Dataflow/DataflowEditorPreviewSceneBase.cpp
+++ b/Engine/Plugins/Experimental/Dataflow/Source/DataflowEditor/Private/Dataflow/DataflowEditorPreviewSceneBase.cpp
@@ -19,32 +19,42 @@ FDataflowPreviewSceneBase::FDataflowPreviewSceneBase(FPreviewScene::Construction
 	: FAdvancedPreviewScene(ConstructionValues)
 	, DataflowEditor(InEditor)
 {
+	RootSceneActor = GetWorld()->SpawnActor<AActor>(AActor::StaticClass());
+	
 	check(DataflowEditor);
 	SetFloorVisibility(bDataflowShowFloorDefault, true);
-
-	RootSceneActor = GetWorld()->SpawnActor<AActor>(AActor::StaticClass());
 }
 
 FDataflowPreviewSceneBase::~FDataflowPreviewSceneBase()
 {}
 
-TObjectPtr<UDataflowBaseContent> FDataflowPreviewSceneBase::GetDataflowContent() 
+TObjectPtr<UDataflowBaseContent>& FDataflowPreviewSceneBase::GetEditorContent() 
+{ 
+	return DataflowEditor->GetEditorContent();
+}
+
+const TObjectPtr<UDataflowBaseContent>& FDataflowPreviewSceneBase::GetEditorContent() const 
+{ 
+	return DataflowEditor->GetEditorContent();
+}
+
+TArray<TObjectPtr<UDataflowBaseContent>>& FDataflowPreviewSceneBase::GetTerminalContents() 
 { 
-	return DataflowEditor->GetDataflowContent();
+	return DataflowEditor->GetTerminalContents();
 }
 
-const TObjectPtr<UDataflowBaseContent> FDataflowPreviewSceneBase::GetDataflowContent() const 
+const TArray<TObjectPtr<UDataflowBaseContent>>& FDataflowPreviewSceneBase::GetTerminalContents() const 
 { 
-	return DataflowEditor->GetDataflowContent();
+	return DataflowEditor->GetTerminalContents();
 }
 
 void FDataflowPreviewSceneBase::AddReferencedObjects(FReferenceCollector& Collector)
 {
 	FAdvancedPreviewScene::AddReferencedObjects(Collector);
 	Collector.AddReferencedObject(RootSceneActor);
-	if (GetDataflowContent())
+	if (const TObjectPtr<UDataflowBaseContent> EditorContent = GetEditorContent())
 	{
-		GetDataflowContent()->AddContentObjects(Collector);
+		EditorContent->AddContentObjects(Collector);
 	}
 }
 
diff --git a/Engine/Plugins/Experimental/Dataflow/Source/DataflowEditor/Private/Dataflow/DataflowEditorToolkit.cpp b/Engine/Plugins/Experimental/Dataflow/Source/DataflowEditor/Private/Dataflow/DataflowEditorToolkit.cpp
--- a/Engine/Plugins/Experimental/Dataflow/Source/DataflowEditor/Private/Dataflow/DataflowEditorToolkit.cpp
+++ b/Engine/Plugins/Experimental/Dataflow/Source/DataflowEditor/Private/Dataflow/DataflowEditorToolkit.cpp
@@ -272,14 +272,24 @@ FEditorModeID FDataflowEditorToolkit::GetEditorModeId() const
 	return UDataflowEditorMode::EM_DataflowEditorModeId;
 }
 
-TObjectPtr<UDataflowBaseContent> FDataflowEditorToolkit::GetDataflowContent()
+TObjectPtr<UDataflowBaseContent>& FDataflowEditorToolkit::GetEditorContent()
 {
-	return DataflowEditor->DataflowContent;
+	return DataflowEditor->GetEditorContent();
 }
 
-const TObjectPtr<UDataflowBaseContent> FDataflowEditorToolkit::GetDataflowContent() const
+const TObjectPtr<UDataflowBaseContent>& FDataflowEditorToolkit::GetEditorContent() const
 {
-	return DataflowEditor->DataflowContent;
+	return DataflowEditor->GetEditorContent();
+}
+
+TArray<TObjectPtr<UDataflowBaseContent>>& FDataflowEditorToolkit::GetTerminalContents()
+{
+	return DataflowEditor->GetTerminalContents();
+}
+
+const TArray<TObjectPtr<UDataflowBaseContent>>& FDataflowEditorToolkit::GetTerminalContents() const
+{
+	return DataflowEditor->GetTerminalContents();
 }
 
 bool FDataflowEditorToolkit::OnRequestClose(EAssetEditorCloseReason InCloseReason)
@@ -401,9 +411,9 @@ void FDataflowEditorToolkit::GetSaveableObjects(TArray<UObject*>& OutObjects) co
 {
 	FBaseCharacterFXEditorToolkit::GetSaveableObjects(OutObjects);
 
-	if (ensure(GetDataflowContent()))
+	if (ensure(GetEditorContent()))
 	{
-		if (UDataflow* DataflowAsset = GetDataflowContent()->GetDataflowAsset())
+		if (UDataflow* DataflowAsset = GetEditorContent()->GetDataflowAsset())
 		{
 			check(DataflowAsset->IsAsset());
 			OutObjects.Add(DataflowAsset);
@@ -420,12 +430,12 @@ void FDataflowEditorToolkit::CreateWidgets()
 {
 	FBaseCharacterFXEditorToolkit::CreateWidgets();
 
-	if (TObjectPtr<UDataflowBaseContent> DataflowContent = GetDataflowContent())
+	if (const TObjectPtr<UDataflowBaseContent>& EditorContent = GetEditorContent())
 	{
-		if (UDataflow* DataflowAsset = DataflowContent->GetDataflowAsset())
+		if (UDataflow* DataflowAsset = EditorContent->GetDataflowAsset())
 		{
-			NodeDetailsEditor = CreateNodeDetailsEditorWidget(DataflowContent->GetDataflowOwner());
-			AssetDetailsEditor = CreateAssetDetailsEditorWidget(DataflowContent->GetDataflowOwner());
+			NodeDetailsEditor = CreateNodeDetailsEditorWidget(EditorContent->GetDataflowOwner());
+			AssetDetailsEditor = CreateAssetDetailsEditorWidget({EditorContent->GetDataflowOwner(),EditorContent->GetDataflowAsset()});
 			GraphEditor = CreateGraphEditorWidget(DataflowAsset, NodeDetailsEditor);
 			CreateSimulationViewportClient();
 		}
@@ -505,7 +515,7 @@ void FDataflowEditorToolkit::CreateSimulationViewportClient()
 
 void FDataflowEditorToolkit::OnPropertyValueChanged(const FPropertyChangedEvent& PropertyChangedEvent)
 {
-	if (TObjectPtr<UDataflowBaseContent> EditorContent = GetDataflowContent())
+	if (const TObjectPtr<UDataflowBaseContent>& EditorContent = GetEditorContent())
 	{
 		ensure(EditorContent);
 		if (UDataflow* DataflowAsset = EditorContent->GetDataflowAsset())
@@ -523,7 +533,7 @@ void FDataflowEditorToolkit::OnPropertyValueChanged(const FPropertyChangedEvent&
 
 void FDataflowEditorToolkit::OnAssetPropertyValueChanged(const FPropertyChangedEvent& PropertyChangedEvent)
 {
-	if (TObjectPtr<UDataflowBaseContent> EditorContent = GetDataflowContent())
+	if (const TObjectPtr<UDataflowBaseContent>& EditorContent = GetEditorContent())
 	{
 		ensure(EditorContent);	FDataflowEditorCommands::OnAssetPropertyValueChanged(EditorContent, PropertyChangedEvent);
 	}
@@ -625,7 +635,7 @@ void FDataflowEditorToolkit::OnNodeSelectionChanged(const TSet<UObject*>& InNewS
 	// Despite this function's name, we might not have actually changed which node is selected
 	bool bPrimarySelectionChanged = false;
 
-	if (TObjectPtr<UDataflowBaseContent> EditorContent = GetDataflowContent(); EditorContent->GetDataflowAsset())
+	if (const TObjectPtr<UDataflowBaseContent>& EditorContent = GetEditorContent(); EditorContent->GetDataflowAsset())
 	{
 		auto AsObjectPointers = [](const TSet<UObject*>& Set) {
 			TSet<TObjectPtr<UObject> > Objs; for (UObject* Elem : Set) Objs.Add(Elem);
@@ -758,7 +768,7 @@ void FDataflowEditorToolkit::OnConstructionViewSelectionChanged(const TArray<UPr
 
 void FDataflowEditorToolkit::Tick(float DeltaTime)
 {
-	if (TObjectPtr<UDataflowBaseContent> EditorContent = GetDataflowContent())
+	if (const TObjectPtr<UDataflowBaseContent> EditorContent = GetEditorContent())
 	{
 		if (EditorContent->GetDataflowAsset())
 		{
@@ -769,9 +779,28 @@ void FDataflowEditorToolkit::Tick(float DeltaTime)
 				TimeStamp = Dataflow::FTimestamp::Invalid;
 			}
 
-			// OnTick evaluation only pulls the termnial nodes. The other evaluations can be specific nodes. 
-			FDataflowEditorCommands::EvaluateTerminalNode(*EditorContent->GetDataflowContext().Get(), TimeStamp, EditorContent->GetDataflowAsset(),
-				nullptr, nullptr, EditorContent->GetDataflowOwner(), EditorContent->GetDataflowTerminal());
+			// Update the list of dataflow teminal contents 
+			DataflowEditor->UpdateTerminalContents(TimeStamp);
+			
+			// OnTick evaluation only pulls the terminal nodes. The other evaluations can be specific nodes.
+			// We only evaluate multiple terminal nodes if the dataflow owner is a UDataflow (Owner == Asset)
+			if(!GetTerminalContents().IsEmpty() && (EditorContent->GetDataflowOwner() == EditorContent->GetDataflowAsset()) && EditorContent->GetDataflowContext())
+			{
+				const Dataflow::FTimestamp InitTimeStamp = TimeStamp;
+				for(const TObjectPtr<UDataflowBaseContent>& TerminalContent : GetTerminalContents())
+				{
+					Dataflow::FTimestamp TerminalTimeStamp = InitTimeStamp;
+					FDataflowEditorCommands::EvaluateTerminalNode(*EditorContent->GetDataflowContext().Get(), TerminalTimeStamp, EditorContent->GetDataflowAsset(),
+						nullptr, nullptr, TerminalContent->GetTerminalAsset(), TerminalContent->GetDataflowTerminal());
+
+					TimeStamp = FMath::Max(TimeStamp, TerminalTimeStamp);
+				}
+			}
+			else
+			{
+				FDataflowEditorCommands::EvaluateTerminalNode(*EditorContent->GetDataflowContext().Get(), TimeStamp, EditorContent->GetDataflowAsset(),
+					nullptr, nullptr, EditorContent->GetTerminalAsset(), EditorContent->GetDataflowTerminal());
+			}
 			EditorContent->SetLastModifiedTimestamp(TimeStamp);
 		}
 	}
@@ -789,25 +818,23 @@ TSharedRef<SDataflowGraphEditor> FDataflowEditorToolkit::CreateGraphEditorWidget
 
 	FDataflowEditorCommands::FGraphEvaluationCallback Evaluate = [&](FDataflowNode* Node, FDataflowOutput* Out)
 	{
-		if (TObjectPtr<UDataflowBaseContent> EditorContent = GetDataflowContent())
+		if (const TObjectPtr<UDataflowBaseContent>& EditorContent = GetEditorContent())
 		{
 			if (EditorContent->GetDataflowAsset())
 			{
-				if (!EditorContent->GetDataflowContext())
-				{
-					EditorContent->SetDataflowContext(MakeShared<Dataflow::FEngineContext>(EditorContent->GetDataflowOwner(), EditorContent->GetDataflowAsset(), Dataflow::FTimestamp::Invalid));
-				}
 				Node->Invalidate();
 				Dataflow::FTimestamp TimeStamp = Dataflow::FTimestamp::Invalid;
 				
 				FDataflowEditorCommands::EvaluateTerminalNode(*EditorContent->GetDataflowContext().Get(), TimeStamp, EditorContent->GetDataflowAsset(),
-					Node, Out, EditorContent->GetDataflowOwner(), EditorContent->GetDataflowTerminal());
+					Node, Out, EditorContent->GetTerminalAsset(), EditorContent->GetDataflowTerminal());
 
 				EditorContent->SetLastModifiedTimestamp(TimeStamp);
 			}
 		}
 	};
-
+	
+	DataflowEditor->UpdateTerminalContents(Dataflow::FTimestamp::Invalid);
+	
 	SGraphEditor::FGraphEditorEvents InEvents;
 	InEvents.OnVerifyTextCommit = FOnNodeVerifyTextCommit::CreateSP(this, &FDataflowEditorToolkit::OnNodeVerifyTitleCommit);
 	InEvents.OnTextCommitted = FOnNodeTextCommitted::CreateSP(this, &FDataflowEditorToolkit::OnNodeTitleCommitted);
@@ -857,9 +884,9 @@ TSharedPtr<IStructureDetailsView> FDataflowEditorToolkit::CreateNodeDetailsEdito
 	return LocalDetailsView;
 }
 
-TSharedPtr<IDetailsView> FDataflowEditorToolkit::CreateAssetDetailsEditorWidget(UObject* ObjectToEdit)
+TSharedPtr<IDetailsView> FDataflowEditorToolkit::CreateAssetDetailsEditorWidget(const TArray<UObject*>& ObjectsToEdit)
 {
-	ensure(ObjectToEdit);
+	ensure(ObjectsToEdit.Num() > 0);
 	FPropertyEditorModule& PropertyEditorModule = FModuleManager::LoadModuleChecked<FPropertyEditorModule>(TEXT("PropertyEditor"));
 
 	FDetailsViewArgs DetailsViewArgs;
@@ -869,10 +896,11 @@ TSharedPtr<IDetailsView> FDataflowEditorToolkit::CreateAssetDetailsEditorWidget(
 		DetailsViewArgs.bUpdatesFromSelection = false;
 		DetailsViewArgs.NameAreaSettings = FDetailsViewArgs::HideNameArea;
 		DetailsViewArgs.NotifyHook = this;
+		DetailsViewArgs.bAllowMultipleTopLevelObjects = true;
 	}
 
 	TSharedPtr<IDetailsView> LocalDetailsView = PropertyEditorModule.CreateDetailView(DetailsViewArgs);
-	LocalDetailsView->SetObject(ObjectToEdit);
+	LocalDetailsView->SetObjects(ObjectsToEdit, true);
 
 	OnFinishedChangingAssetPropertiesDelegateHandle = LocalDetailsView->OnFinishedChangingProperties().AddSP(this, &FDataflowEditorToolkit::OnAssetPropertyValueChanged);
 
@@ -929,9 +957,9 @@ TSharedRef<SDockTab> FDataflowEditorToolkit::SpawnTab_SkeletonView(const FSpawnT
 {
 	check(Args.GetTabId() == SkeletonViewTabId);
 	check(DataflowEditor);
-	check(DataflowEditor->GetDataflowContent());
+	check(DataflowEditor->GetEditorContent());
 
-	SkeletonEditorView = MakeShared<FDataflowSkeletonView>(DataflowEditor->GetDataflowContent());
+	SkeletonEditorView = MakeShared<FDataflowSkeletonView>(DataflowEditor->GetEditorContent());
 	ViewListeners.Add(SkeletonEditorView.Get());
 
 	FSkeletonTreeArgs SkeletonTreeArgs;
@@ -956,11 +984,11 @@ TSharedRef<SDockTab> FDataflowEditorToolkit::SpawnTab_SelectionView(const FSpawn
 {
 	//	check(Args.GetTabId().TabType == SelectionViewTabId_1);
 	check(DataflowEditor);
-	check(DataflowEditor->GetDataflowContent());
+	check(DataflowEditor->GetEditorContent());
 
 		if (Args.GetTabId() == SelectionViewTabId_1)
 	{
-		DataflowSelectionView_1 = MakeShared<FDataflowSelectionView>(FDataflowSelectionView(DataflowEditor->GetDataflowContent()));
+		DataflowSelectionView_1 = MakeShared<FDataflowSelectionView>(FDataflowSelectionView(DataflowEditor->GetEditorContent()));
 		if (DataflowSelectionView_1.IsValid())
 		{
 			ViewListeners.Add(DataflowSelectionView_1.Get());
@@ -968,7 +996,7 @@ TSharedRef<SDockTab> FDataflowEditorToolkit::SpawnTab_SelectionView(const FSpawn
 	}
 	else if (Args.GetTabId() == SelectionViewTabId_2)
 	{
-		DataflowSelectionView_2 = MakeShared<FDataflowSelectionView>(FDataflowSelectionView(DataflowEditor->GetDataflowContent()));
+		DataflowSelectionView_2 = MakeShared<FDataflowSelectionView>(FDataflowSelectionView(DataflowEditor->GetEditorContent()));
 		if (DataflowSelectionView_2.IsValid())
 		{
 			ViewListeners.Add(DataflowSelectionView_2.Get());
@@ -976,7 +1004,7 @@ TSharedRef<SDockTab> FDataflowEditorToolkit::SpawnTab_SelectionView(const FSpawn
 	}
 	else if (Args.GetTabId() == SelectionViewTabId_3)
 	{
-		DataflowSelectionView_3 = MakeShared<FDataflowSelectionView>(FDataflowSelectionView(DataflowEditor->GetDataflowContent()));
+		DataflowSelectionView_3 = MakeShared<FDataflowSelectionView>(FDataflowSelectionView(DataflowEditor->GetEditorContent()));
 		if (DataflowSelectionView_3.IsValid())
 		{
 			ViewListeners.Add(DataflowSelectionView_3.Get());
@@ -984,7 +1012,7 @@ TSharedRef<SDockTab> FDataflowEditorToolkit::SpawnTab_SelectionView(const FSpawn
 	}
 	else if (Args.GetTabId() == SelectionViewTabId_4)
 	{
-		DataflowSelectionView_4 = MakeShared<FDataflowSelectionView>(FDataflowSelectionView(DataflowEditor->GetDataflowContent()));
+		DataflowSelectionView_4 = MakeShared<FDataflowSelectionView>(FDataflowSelectionView(DataflowEditor->GetEditorContent()));
 		if (DataflowSelectionView_4.IsValid())
 		{
 			ViewListeners.Add(DataflowSelectionView_4.Get());
@@ -1000,7 +1028,7 @@ TSharedRef<SDockTab> FDataflowEditorToolkit::SpawnTab_SelectionView(const FSpawn
 
 	if (SelectionViewWidget)
 	{
-		if (const TObjectPtr<UDataflowBaseContent> EditorContent = GetDataflowContent())
+		if (const TObjectPtr<UDataflowBaseContent>& EditorContent = GetEditorContent())
 		{
 			if (Args.GetTabId() == SelectionViewTabId_1)
 			{
@@ -1029,11 +1057,11 @@ TSharedRef<SDockTab> FDataflowEditorToolkit::SpawnTab_SelectionView(const FSpawn
 TSharedRef<SDockTab> FDataflowEditorToolkit::SpawnTab_CollectionSpreadSheet(const FSpawnTabArgs& Args)
 {
 	check(DataflowEditor);
-	check(DataflowEditor->GetDataflowContent());
+	check(DataflowEditor->GetEditorContent());
 
 	if (Args.GetTabId() == CollectionSpreadSheetTabId_1)
 	{
-		DataflowCollectionSpreadSheet_1 = MakeShared<FDataflowCollectionSpreadSheet>(FDataflowCollectionSpreadSheet(DataflowEditor->GetDataflowContent()));
+		DataflowCollectionSpreadSheet_1 = MakeShared<FDataflowCollectionSpreadSheet>(FDataflowCollectionSpreadSheet(DataflowEditor->GetEditorContent()));
 		if (DataflowCollectionSpreadSheet_1.IsValid())
 		{
 			ViewListeners.Add(DataflowCollectionSpreadSheet_1.Get());
@@ -1041,7 +1069,7 @@ TSharedRef<SDockTab> FDataflowEditorToolkit::SpawnTab_CollectionSpreadSheet(cons
 	}
 	else if (Args.GetTabId() == CollectionSpreadSheetTabId_2)
 	{
-		DataflowCollectionSpreadSheet_2 = MakeShared<FDataflowCollectionSpreadSheet>(FDataflowCollectionSpreadSheet(DataflowEditor->GetDataflowContent()));
+		DataflowCollectionSpreadSheet_2 = MakeShared<FDataflowCollectionSpreadSheet>(FDataflowCollectionSpreadSheet(DataflowEditor->GetEditorContent()));
 		if (DataflowCollectionSpreadSheet_2.IsValid())
 		{
 			ViewListeners.Add(DataflowCollectionSpreadSheet_2.Get());
@@ -1049,7 +1077,7 @@ TSharedRef<SDockTab> FDataflowEditorToolkit::SpawnTab_CollectionSpreadSheet(cons
 	}
 	else if (Args.GetTabId() == CollectionSpreadSheetTabId_3)
 	{
-		DataflowCollectionSpreadSheet_3 = MakeShared<FDataflowCollectionSpreadSheet>(FDataflowCollectionSpreadSheet(DataflowEditor->GetDataflowContent()));
+		DataflowCollectionSpreadSheet_3 = MakeShared<FDataflowCollectionSpreadSheet>(FDataflowCollectionSpreadSheet(DataflowEditor->GetEditorContent()));
 		if (DataflowCollectionSpreadSheet_3.IsValid())
 		{
 			ViewListeners.Add(DataflowCollectionSpreadSheet_3.Get());
@@ -1057,7 +1085,7 @@ TSharedRef<SDockTab> FDataflowEditorToolkit::SpawnTab_CollectionSpreadSheet(cons
 	}
 	else if (Args.GetTabId() == CollectionSpreadSheetTabId_4)
 	{
-		DataflowCollectionSpreadSheet_4 = MakeShared<FDataflowCollectionSpreadSheet>(FDataflowCollectionSpreadSheet(DataflowEditor->GetDataflowContent()));
+		DataflowCollectionSpreadSheet_4 = MakeShared<FDataflowCollectionSpreadSheet>(FDataflowCollectionSpreadSheet(DataflowEditor->GetEditorContent()));
 		if (DataflowCollectionSpreadSheet_4.IsValid())
 		{
 			ViewListeners.Add(DataflowCollectionSpreadSheet_4.Get());
@@ -1073,7 +1101,7 @@ TSharedRef<SDockTab> FDataflowEditorToolkit::SpawnTab_CollectionSpreadSheet(cons
 
 	if (CollectionSpreadSheetWidget)
 	{
-		if (const TObjectPtr<UDataflowBaseContent> EditorContent = GetDataflowContent())
+		if (const TObjectPtr<UDataflowBaseContent>& EditorContent = GetEditorContent())
 		{
 			if (Args.GetTabId() == CollectionSpreadSheetTabId_1)
 			{
@@ -1246,7 +1274,7 @@ FName FDataflowEditorToolkit::GetToolkitFName() const
 
 FText FDataflowEditorToolkit::GetToolkitName() const
 {
-	if (const TObjectPtr<UDataflowBaseContent> EditorContent = GetDataflowContent())
+	if (const TObjectPtr<UDataflowBaseContent>& EditorContent = GetEditorContent())
 	{
 		if (EditorContent->GetDataflowOwner())
 		{
diff --git a/Engine/Plugins/Experimental/Dataflow/Source/DataflowEditor/Private/Dataflow/DataflowSimulationPanel.cpp b/Engine/Plugins/Experimental/Dataflow/Source/DataflowEditor/Private/Dataflow/DataflowSimulationPanel.cpp
--- a/Engine/Plugins/Experimental/Dataflow/Source/DataflowEditor/Private/Dataflow/DataflowSimulationPanel.cpp
+++ b/Engine/Plugins/Experimental/Dataflow/Source/DataflowEditor/Private/Dataflow/DataflowSimulationPanel.cpp
@@ -415,9 +415,10 @@ UAnimSingleNodeInstance* SDataflowSimulationPanel::GetPreviewAnimationInstance()
 {
 	if (const TSharedPtr<const FDataflowSimulationScene> PreviewScene = SimulationScene.Pin())
 	{
-		if(const TObjectPtr<UDataflowSkeletalContent> DataflowContent = Cast<UDataflowSkeletalContent>(PreviewScene->GetDataflowContent()))
+		if(const TObjectPtr<UDataflowSkeletalContent> DataflowContent = Cast<UDataflowSkeletalContent>(PreviewScene->GetEditorContent()))
 		{
-			return DataflowContent->GetAnimationInstance();
+			// temporary nullptr while the simulation nodes are pushed 
+			return nullptr;
 		}
 	}
 
@@ -428,9 +429,10 @@ const UAnimSingleNodeInstance* SDataflowSimulationPanel::GetPreviewAnimationInst
 {
 	if (const TSharedPtr<const FDataflowSimulationScene> PreviewScene = SimulationScene.Pin())
 	{
-		if(const TObjectPtr<UDataflowSkeletalContent> DataflowContent = Cast<UDataflowSkeletalContent>(PreviewScene->GetDataflowContent()))
+		if(const TObjectPtr<UDataflowSkeletalContent> DataflowContent = Cast<UDataflowSkeletalContent>(PreviewScene->GetEditorContent()))
 		{
-			return DataflowContent->GetAnimationInstance();
+			// temporary nullptr while the simulation nodes are pushed 
+			return nullptr;
 		}
 	}
 
@@ -441,9 +443,9 @@ void SDataflowSimulationPanel::UpdatePreviewAnimationInstance()
 {
 	if (const TSharedPtr<FDataflowSimulationScene> PreviewScene = SimulationScene.Pin())
 	{
-		if(TObjectPtr<UDataflowSkeletalContent> DataflowContent = Cast<UDataflowSkeletalContent>(PreviewScene->GetDataflowContent()))
+		if(TObjectPtr<UDataflowSkeletalContent> DataflowContent = Cast<UDataflowSkeletalContent>(PreviewScene->GetEditorContent()))
 		{
-			DataflowContent->UpdateAnimationInstance();
+			// do nothing while the simulation nodes are pushed 
 		}
 	}
 }
diff --git a/Engine/Plugins/Experimental/Dataflow/Source/DataflowEditor/Private/Dataflow/DataflowSimulationScene.cpp b/Engine/Plugins/Experimental/Dataflow/Source/DataflowEditor/Private/Dataflow/DataflowSimulationScene.cpp
--- a/Engine/Plugins/Experimental/Dataflow/Source/DataflowEditor/Private/Dataflow/DataflowSimulationScene.cpp
+++ b/Engine/Plugins/Experimental/Dataflow/Source/DataflowEditor/Private/Dataflow/DataflowSimulationScene.cpp
@@ -14,34 +14,30 @@
 FDataflowSimulationScene::FDataflowSimulationScene(FPreviewScene::ConstructionValues ConstructionValues, UDataflowEditor* InEditor)
 	: FDataflowPreviewSceneBase(ConstructionValues, InEditor)
 {
-	if (TObjectPtr<UDataflowBaseContent> DataflowContent = GetDataflowContent())
+	if(RootSceneActor)
 	{
-		DataflowContent->RegisterWorldContent(this, RootSceneActor);
-	}
-
-	TInlineComponentArray<UPrimitiveComponent*> PrimComponents;
-	RootSceneActor->GetComponents(PrimComponents);
-
-	for(UPrimitiveComponent* PrimComponent : PrimComponents)
-	{
-		PrimComponent->SelectionOverrideDelegate =
-			UPrimitiveComponent::FSelectionOverride::CreateRaw(this, &FDataflowPreviewSceneBase::IsComponentSelected);
+		TInlineComponentArray<UPrimitiveComponent*> PrimComponents;
+		RootSceneActor->GetComponents(PrimComponents);
+
+		for(UPrimitiveComponent* PrimComponent : PrimComponents)
+		{
+			PrimComponent->SelectionOverrideDelegate =
+				UPrimitiveComponent::FSelectionOverride::CreateRaw(this, &FDataflowPreviewSceneBase::IsComponentSelected);
+		}
 	}
 }
 
 FDataflowSimulationScene::~FDataflowSimulationScene()
 {
-	TInlineComponentArray<UPrimitiveComponent*> PrimComponents;
-	RootSceneActor->GetComponents(PrimComponents);
-
-	for(UPrimitiveComponent* PrimComponent : PrimComponents)
+	if(RootSceneActor)
 	{
-		PrimComponent->SelectionOverrideDelegate.Unbind();
-	}
+		TInlineComponentArray<UPrimitiveComponent*> PrimComponents;
+		RootSceneActor->GetComponents(PrimComponents);
 
-	if (TObjectPtr<UDataflowBaseContent> DataflowContent = GetDataflowContent())
-	{
-		DataflowContent->UnregisterWorldContent(this);
+		for(UPrimitiveComponent* PrimComponent : PrimComponents)
+		{
+			PrimComponent->SelectionOverrideDelegate.Unbind();
+		}
 	}
 }
 
diff --git a/Engine/Plugins/Experimental/Dataflow/Source/DataflowEditor/Private/Dataflow/DataflowSimulationViewport.cpp b/Engine/Plugins/Experimental/Dataflow/Source/DataflowEditor/Private/Dataflow/DataflowSimulationViewport.cpp
--- a/Engine/Plugins/Experimental/Dataflow/Source/DataflowEditor/Private/Dataflow/DataflowSimulationViewport.cpp
+++ b/Engine/Plugins/Experimental/Dataflow/Source/DataflowEditor/Private/Dataflow/DataflowSimulationViewport.cpp
@@ -102,12 +102,12 @@ void SDataflowSimulationViewport::OnFloatingButtonClicked()
 
 float SDataflowSimulationViewport:: GetViewMinInput() const
 {
-	return static_cast<FDataflowPreviewSceneBase*>(Client->GetPreviewScene())->GetDataflowContent()->GetSimulationRange()[0];
+	return 0.0f; //Temp value while the simulation nodes is being pushed
 }
 
 float SDataflowSimulationViewport::GetViewMaxInput() const
 {
-	return static_cast<FDataflowPreviewSceneBase*>(Client->GetPreviewScene())->GetDataflowContent()->GetSimulationRange()[1];
+	return 5.0f; //Temp value while the simulation nodes is being pushed
 }
 
 
diff --git a/Engine/Plugins/Experimental/Dataflow/Source/DataflowEditor/Public/Dataflow/AssetDefinition_DataflowContext.h b/Engine/Plugins/Experimental/Dataflow/Source/DataflowEditor/Public/Dataflow/AssetDefinition_DataflowContext.h
--- a/Engine/Plugins/Experimental/Dataflow/Source/DataflowEditor/Public/Dataflow/AssetDefinition_DataflowContext.h
+++ b/Engine/Plugins/Experimental/Dataflow/Source/DataflowEditor/Public/Dataflow/AssetDefinition_DataflowContext.h
@@ -12,7 +12,7 @@ namespace DataflowContextDefinitionHelpers
 {
 	// Return true if we should proceed, false if we should re-open the dialog
 	template<class T>
-	DATAFLOWEDITOR_API TObjectPtr<T> CreateNewDataflowContext(const TObjectPtr<UObject>& ContentOwner);
+	DATAFLOWEDITOR_API TObjectPtr<T> CreateNewDataflowContent(const TObjectPtr<UObject>& ContentOwner);
 }
 
 
diff --git a/Engine/Plugins/Experimental/Dataflow/Source/DataflowEditor/Public/Dataflow/DataflowConstructionViewport.h b/Engine/Plugins/Experimental/Dataflow/Source/DataflowEditor/Public/Dataflow/DataflowConstructionViewport.h
--- a/Engine/Plugins/Experimental/Dataflow/Source/DataflowEditor/Public/Dataflow/DataflowConstructionViewport.h
+++ b/Engine/Plugins/Experimental/Dataflow/Source/DataflowEditor/Public/Dataflow/DataflowConstructionViewport.h
@@ -35,9 +35,6 @@ public:
 	virtual void OnFloatingButtonClicked() override;
 
 private:
-	float GetViewMinInput() const;
-	float GetViewMaxInput() const;
-
 	UDataflowEditorMode* GetEdMode() const;
 };
 
diff --git a/Engine/Plugins/Experimental/Dataflow/Source/DataflowEditor/Public/Dataflow/DataflowEditor.h b/Engine/Plugins/Experimental/Dataflow/Source/DataflowEditor/Public/Dataflow/DataflowEditor.h
--- a/Engine/Plugins/Experimental/Dataflow/Source/DataflowEditor/Public/Dataflow/DataflowEditor.h
+++ b/Engine/Plugins/Experimental/Dataflow/Source/DataflowEditor/Public/Dataflow/DataflowEditor.h
@@ -42,24 +42,52 @@ public:
 
 	// UBaseCharacterFXEditor interface
 	virtual TSharedPtr<FBaseAssetToolkit> CreateToolkit() override;
+
+	/** Initialize editor contents given a list of objects */
 	virtual void Initialize(const TArray<TObjectPtr<UObject>>& InObjects) override;
 
-	/** Initialize an editor with a given content and an owner */
-	void InitializeContent(const TObjectPtr<UObject>& ContentOwner);
+	/** Update the terminal contents */
+	void UpdateTerminalContents(const Dataflow::FTimestamp TimeStamp);
+	
+	/** Update the editor content */
+	void UpdateEditorContent();
+
+	/** Dataflow editor content accessors */
+	TObjectPtr<UDataflowBaseContent>& GetEditorContent() { return EditorContent; }
+	const TObjectPtr<UDataflowBaseContent>& GetEditorContent() const { return EditorContent; }
 
-	TObjectPtr<UDataflowBaseContent> GetDataflowContent() { return DataflowContent; }
-	const TObjectPtr<UDataflowBaseContent> GetDataflowContent() const { return DataflowContent; }
+	/** Dataflow terminal contents accessors */
+	TArray<TObjectPtr<UDataflowBaseContent>>& GetTerminalContents() { return TerminalContents; }
+	const TArray<TObjectPtr<UDataflowBaseContent>>& GetTerminalContents() const { return TerminalContents; }
+
+	/** Check if the terminals contents has changed or not */
+	bool HasTerminalsDirty() const { return bHasTerminalsDirty; }
 
 private :
 
 	friend class FDataflowEditorToolkit;
+
+	using ValidTerminalsType = TMap<TSharedPtr<FDataflowNode>,TObjectPtr<UDataflowBaseContent>>;
+
+	/** Remove invalid terminal contents from the container */
+	void RemoveTerminalContents(const TSharedPtr<Dataflow::FGraph>& DataflowGraph, ValidTerminalsType& ValidTerminals);
+	
+	/** Add valid terminal nodes to the container */
+	void AddTerminalContents(const TSharedPtr<Dataflow::FGraph>& DataflowGraph, ValidTerminalsType& ValidTerminals);
 	
 	// Dataflow editor is the owner of the object list to edit/process and the dataflow mode
 	// is the one holding the dynamic mesh components to be rendered in the viewport
 	// It is why the data flow asset/owner/skelmesh have been added here. Could be added
 	// in the subsystem if necessary
 	UPROPERTY()
-	TObjectPtr<UDataflowBaseContent> DataflowContent;
+	TObjectPtr<UDataflowBaseContent> EditorContent;
+
+	/** List of dataflow contents available in the graph and coming from all the terminal nodes */
+	UPROPERTY()
+	TArray<TObjectPtr<UDataflowBaseContent>> TerminalContents;
+
+	/** Boolean to know when the terminals has been changed or not */
+	bool bHasTerminalsDirty = false;
 };
 
 DECLARE_LOG_CATEGORY_EXTERN(LogDataflowEditor, Log, All);
\ No newline at end of file
diff --git a/Engine/Plugins/Experimental/Dataflow/Source/DataflowEditor/Public/Dataflow/DataflowEditorPreviewSceneBase.h b/Engine/Plugins/Experimental/Dataflow/Source/DataflowEditor/Public/Dataflow/DataflowEditorPreviewSceneBase.h
--- a/Engine/Plugins/Experimental/Dataflow/Source/DataflowEditor/Public/Dataflow/DataflowEditorPreviewSceneBase.h
+++ b/Engine/Plugins/Experimental/Dataflow/Source/DataflowEditor/Public/Dataflow/DataflowEditorPreviewSceneBase.h
@@ -23,9 +23,13 @@ public:
 	// FGCObject interface
 	virtual void AddReferencedObjects(FReferenceCollector& Collector) override;
 
-	/** Dataflow content accessors */
-	TObjectPtr<UDataflowBaseContent> GetDataflowContent();
-		const TObjectPtr<UDataflowBaseContent> GetDataflowContent() const;
+	/** Dataflow editor content accessors */
+	TObjectPtr<UDataflowBaseContent>& GetEditorContent();
+	const TObjectPtr<UDataflowBaseContent>& GetEditorContent() const;
+
+	/** Dataflow terminal contents accessors */
+	TArray<TObjectPtr<UDataflowBaseContent>>& GetTerminalContents();
+	const TArray<TObjectPtr<UDataflowBaseContent>>& GetTerminalContents() const;
 
 	/** Root scene actor accessors */
 	TObjectPtr<AActor> GetRootActor() { return RootSceneActor; }
@@ -50,8 +54,9 @@ public:
 protected:
 	
 	/** Root scene actor */
-	TObjectPtr<AActor> RootSceneActor = nullptr;
+	TObjectPtr<AActor> RootSceneActor = nullptr; 
 
+	/** Dataflow editor linked to that preview scene */
 	UDataflowEditor* DataflowEditor = nullptr;
 
 	/** Mode Manager for selection */
diff --git a/Engine/Plugins/Experimental/Dataflow/Source/DataflowEditor/Public/Dataflow/DataflowEditorToolkit.h b/Engine/Plugins/Experimental/Dataflow/Source/DataflowEditor/Public/Dataflow/DataflowEditorToolkit.h
--- a/Engine/Plugins/Experimental/Dataflow/Source/DataflowEditor/Public/Dataflow/DataflowEditorToolkit.h
+++ b/Engine/Plugins/Experimental/Dataflow/Source/DataflowEditor/Public/Dataflow/DataflowEditorToolkit.h
@@ -56,10 +56,16 @@ public:
 	static bool HasDataflowAsset(UObject* ObjectToEdit);
 	static UDataflow* GetDataflowAsset(UObject* ObjectToEdit);
 	static const UDataflow* GetDataflowAsset(const UObject* ObjectToEdit);
+	
+	/** Editor dataflow content accessors */
+	const TObjectPtr<UDataflowBaseContent>& GetEditorContent() const;
+	TObjectPtr<UDataflowBaseContent>& GetEditorContent();
 
-	/** Dataflow Content Access */
-	const TObjectPtr<UDataflowBaseContent> GetDataflowContent() const;
-	TObjectPtr<UDataflowBaseContent> GetDataflowContent();
+	/** Terminal dataflow contents accessors */
+	const TArray<TObjectPtr<UDataflowBaseContent>>& GetTerminalContents() const;
+	TArray<TObjectPtr<UDataflowBaseContent>>& GetTerminalContents();
+	
+	/** Dataflow graph editor accessor */
 	const TSharedPtr<SDataflowGraphEditor> GetDataflowGraphEditor() const { return GraphEditor; }
 
 	// IToolkit interface
@@ -163,7 +169,7 @@ private:
 
 	// Utility factory functions to build the widgets
 	TSharedRef<SDataflowGraphEditor> CreateGraphEditorWidget(UDataflow* ObjectToEdit, TSharedPtr<IStructureDetailsView> PropertiesEditor);
-    TSharedPtr<IDetailsView> CreateAssetDetailsEditorWidget(UObject* ObjectToEdit);
+    TSharedPtr<IDetailsView> CreateAssetDetailsEditorWidget(const TArray<UObject*>& ObjectsToEdit);
     TSharedPtr<IStructureDetailsView> CreateNodeDetailsEditorWidget(UObject* ObjectToEdit);
 
 	/** Create the simulation viewport client */
diff --git a/Engine/Source/Runtime/Experimental/Dataflow/Core/Public/Dataflow/DataflowTerminalNode.h b/Engine/Source/Runtime/Experimental/Dataflow/Core/Public/Dataflow/DataflowTerminalNode.h
--- a/Engine/Source/Runtime/Experimental/Dataflow/Core/Public/Dataflow/DataflowTerminalNode.h
+++ b/Engine/Source/Runtime/Experimental/Dataflow/Core/Public/Dataflow/DataflowTerminalNode.h
@@ -43,7 +43,10 @@ struct FDataflowTerminalNode : public FDataflowNode
 	{ 
 		return InType.ToString().Equals(StaticType().ToString()) 
 			|| Super::IsA(InType); 
-	} 
+	}
+
+	/** Return the terminal asset */
+	virtual TObjectPtr<UObject> GetTerminalAsset() const {return nullptr;}
 
 	//
 	// Error Checking
@@ -61,7 +64,6 @@ private:
 	{
 		Evaluate(Context);
 	};
-
 };
 
 
diff --git a/Engine/Source/Runtime/Experimental/Dataflow/Engine/Private/Dataflow/DataflowContent.cpp b/Engine/Source/Runtime/Experimental/Dataflow/Engine/Private/Dataflow/DataflowContent.cpp
--- a/Engine/Source/Runtime/Experimental/Dataflow/Engine/Private/Dataflow/DataflowContent.cpp
+++ b/Engine/Source/Runtime/Experimental/Dataflow/Engine/Private/Dataflow/DataflowContent.cpp
@@ -20,6 +20,17 @@ void UDataflowContextObject::AddReferencedObjects(UObject* InThis, FReferenceCol
 	Super::AddReferencedObjects(InThis, Collector);
 }
 
+TObjectPtr<UDataflowBaseContent> IDataflowContentOwner::BuildDataflowContent()
+{
+	if(TObjectPtr<UDataflowBaseContent> DataflowContent = CreateDataflowContent())
+	{
+		// Delegate used for notifying owner data invalidation
+		OnContentOwnerChanged.AddUObject(DataflowContent, &UDataflowBaseContent::UpdateContentDatas);
+		return DataflowContent;
+	}
+	return nullptr;
+}
+
 //
 // UDataflowBaseContent
 //
@@ -29,30 +40,61 @@ void UDataflowBaseContent::SetIsDirty(bool InDirty)
 	bIsDirty = InDirty;
 }
 
-
 UDataflowBaseContent::UDataflowBaseContent()
 {
 }
 
+UDataflowBaseContent::~UDataflowBaseContent()
+{
+	if(GetDataflowOwner())
+	{
+		if(IDataflowContentOwner* ContentOwner = Cast<IDataflowContentOwner>(GetDataflowOwner()))
+		{
+			ContentOwner->OnContentOwnerChanged.RemoveAll(this);
+		}
+	}
+}
+
+void UDataflowBaseContent::UpdateContentDatas()
+{
+	if(GetDataflowOwner())
+	{
+		if(const IDataflowContentOwner* ContentOwner = Cast<IDataflowContentOwner>(GetDataflowOwner()))
+		{
+			ContentOwner->UpdateDataflowContent(this);
+		}
+	}
+}
+
 void UDataflowBaseContent::SetDataflowOwner(const TObjectPtr<UObject>& InOwner)
 {
-	DataflowOwner = InOwner;
-	if (DataflowContext)
+	if(!DataflowContext)
 	{
-		DataflowContext->Owner = InOwner;  
-		SetIsDirty(true);
+		DataflowContext = MakeShared<Dataflow::FEngineContext>(nullptr, nullptr, Dataflow::FTimestamp::Invalid);
 	}
+	DataflowContext->Owner = InOwner;  
+	SetIsDirty(true);
 }
 
 TObjectPtr<UObject> UDataflowBaseContent::GetDataflowOwner() const 
 {
-	if (DataflowContext)
+	return DataflowContext ? DataflowContext->Owner : nullptr; 
+}
+
+void UDataflowBaseContent::SetDataflowAsset(const TObjectPtr<UDataflow>& DataflowAsset)
+{
+	if(!DataflowContext)
 	{
-		ensure(DataflowContext->Owner == DataflowOwner);
+		DataflowContext = MakeShared<Dataflow::FEngineContext>(nullptr, nullptr, Dataflow::FTimestamp::Invalid);
 	}
-	return DataflowContext ? DataflowContext->Owner : DataflowOwner; 
+	DataflowContext->Graph = DataflowAsset;  
+	SetIsDirty(true);
 }
 
+TObjectPtr<UDataflow> UDataflowBaseContent::GetDataflowAsset() const 
+{
+	return DataflowContext ? DataflowContext->Graph : nullptr; 
+}
 
 void UDataflowBaseContent::SetLastModifiedTimestamp(Dataflow::FTimestamp InTimestamp, bool bMakeDirty) 
 { 
@@ -74,16 +116,6 @@ void UDataflowBaseContent::SetDataflowContext(const TSharedPtr<Dataflow::FEngine
 	MarkPackageDirty();
 }
 
-
-void UDataflowBaseContent::BuildBaseContent(TObjectPtr<UObject> InDataflowOwner)
-{
-	DataflowOwner = InDataflowOwner;
-	DataflowContext = MakeShared<Dataflow::FEngineContext>(InDataflowOwner, DataflowAsset, FPlatformTime::Cycles64());
-	LastModifiedTimestamp = DataflowContext->GetTimestamp();
-	SetIsDirty(true);
-	MarkPackageDirty();
-}
-
 void UDataflowBaseContent::Serialize(FArchive& Ar)
 {
 	Super::Serialize(Ar);
@@ -91,7 +123,7 @@ void UDataflowBaseContent::Serialize(FArchive& Ar)
 
 	if (!DataflowContext)
 	{
-		DataflowContext = MakeShared<Dataflow::FEngineContext>(DataflowOwner, DataflowAsset, LastModifiedTimestamp);
+		DataflowContext = MakeShared<Dataflow::FEngineContext>(nullptr, nullptr, LastModifiedTimestamp);
 	}
 	DataflowContext->Serialize(Ar);
 }
@@ -99,8 +131,11 @@ void UDataflowBaseContent::Serialize(FArchive& Ar)
 void UDataflowBaseContent::AddReferencedObjects(UObject* InThis, FReferenceCollector& Collector)
 {
 	UDataflowBaseContent* This = CastChecked<UDataflowBaseContent>(InThis);
-	Collector.AddReferencedObject(This->DataflowOwner);
-	Collector.AddReferencedObject(This->DataflowAsset);
+	if(This->DataflowContext)
+	{
+		Collector.AddReferencedObject(This->DataflowContext->Owner);
+		Collector.AddReferencedObject(This->DataflowContext->Graph);
+	}
 	Super::AddReferencedObjects(InThis, Collector);
 }
 
@@ -108,79 +143,19 @@ void UDataflowBaseContent::AddReferencedObjects(UObject* InThis, FReferenceColle
 // UDataflowSkeletalContent
 //
 
-
 UDataflowSkeletalContent::UDataflowSkeletalContent() : Super()
 {
 }
 
-
-void UDataflowSkeletalContent::RegisterWorldContent(FPreviewScene* PreviewScene, AActor* RootActor)
-{
-	Super::RegisterWorldContent(PreviewScene, RootActor);
-
-	SkeletalMeshComponent = NewObject<USkeletalMeshComponent>(RootActor);
-	SkeletalMeshComponent->SetDisablePostProcessBlueprint(true);
-	
-	if(SkeletalMesh)
-	{
-		SkeletalMeshComponent->SetSkeletalMeshAsset(SkeletalMesh);
-
-		Skeleton = SkeletalMesh->GetSkeleton();
-		UpdateAnimationInstance();
-	}
-	SkeletalMeshComponent->UpdateBounds();
-	PreviewScene->AddComponent(SkeletalMeshComponent, SkeletalMeshComponent->GetRelativeTransform());
-}
-
-void UDataflowSkeletalContent::UpdateAnimationInstance()
-{
-	if (AnimationAsset && SkeletalMesh && (AnimationAsset->GetSkeleton() == SkeletalMesh->GetSkeleton()))
-	{
-		AnimationNodeInstance = NewObject<UAnimSingleNodeInstance>(SkeletalMeshComponent);
-		AnimationNodeInstance->SetAnimationAsset(AnimationAsset);
-
-		SkeletalMeshComponent->SetAnimationMode(EAnimationMode::AnimationSingleNode);
-		SkeletalMeshComponent->InitAnim(true);
-		SkeletalMeshComponent->AnimationData.PopulateFrom(AnimationNodeInstance);
-		SkeletalMeshComponent->AnimScriptInstance = AnimationNodeInstance;
-		SkeletalMeshComponent->AnimScriptInstance->InitializeAnimation();
-
-#if WITH_EDITOR
-		SkeletalMeshComponent->ValidateAnimation();
-#endif
-	}
-	else
-	{
-		SkeletalMeshComponent->Stop();
-		SkeletalMeshComponent->AnimationData = FSingleAnimationPlayData();
-		SkeletalMeshComponent->AnimScriptInstance = nullptr;
-		AnimationNodeInstance = nullptr;
-		AnimationAsset = nullptr;
-	}
-}
-	
-void UDataflowSkeletalContent::UnregisterWorldContent(FPreviewScene* PreviewScene)
-{
-	Super::UnregisterWorldContent(PreviewScene);
-	if (SkeletalMeshComponent)
-	{
-		SkeletalMeshComponent->TransformUpdated.RemoveAll(this);
-		PreviewScene->RemoveComponent(SkeletalMeshComponent);
-	}
-}
-
 void UDataflowSkeletalContent::SetSkeletalMesh(const TObjectPtr<USkeletalMesh>& SkeletalMeshAsset)
 {
 	SkeletalMesh = SkeletalMeshAsset;
-	if(SkeletalMesh && SkeletalMesh->GetSkeleton())
-	{
-		Skeleton = SkeletalMesh->GetSkeleton();
-	}
-	if(SkeletalMeshComponent)
+	if(SkeletalMesh)
 	{
- 		SkeletalMeshComponent->SetSkeletalMeshAsset(SkeletalMesh);
-
- 		UpdateAnimationInstance();
+		if(SkeletalMesh && (SkeletalMesh->GetSkeleton() != Skeleton))
+		{
+			SetSkeleton(SkeletalMesh->GetSkeleton());
+		}
 	}
 	SetIsDirty(true);
 }
@@ -188,22 +163,23 @@ void UDataflowSkeletalContent::SetSkeletalMesh(const TObjectPtr<USkeletalMesh>&
 void UDataflowSkeletalContent::SetAnimationAsset(const TObjectPtr<UAnimationAsset>& SkeletalAnimationAsset)
 {
 	AnimationAsset = SkeletalAnimationAsset;
-	if(SkeletalMeshComponent)
+	if(AnimationAsset && (AnimationAsset->GetSkeleton()) != Skeleton)
 	{
- 		UpdateAnimationInstance();
+		SetSkeleton(AnimationAsset->GetSkeleton());
 	}
 	SetIsDirty(true);
 }
 
 void UDataflowSkeletalContent::SetSkeleton(const TObjectPtr<USkeleton>& SkeletonAsset)
 {
-	if(SkeletonAsset)
+	Skeleton = SkeletonAsset;
+	if(SkeletalMesh && (SkeletalMesh->GetSkeleton() != Skeleton))
 	{
-		Skeleton = SkeletonAsset;
-		if(SkeletalMesh && (SkeletalMesh->GetSkeleton() != Skeleton))
-		{
-			SetSkeletalMesh(nullptr);
-		}
+		SetSkeletalMesh(nullptr);
+	}
+	if(AnimationAsset && (AnimationAsset->GetSkeleton() != Skeleton))
+	{
+		SetAnimationAsset(nullptr);
 	}
 	SetIsDirty(true);
 }
@@ -229,23 +205,12 @@ void UDataflowSkeletalContent::PostEditChangeProperty(FPropertyChangedEvent& Pro
 		SetSkeleton(Skeleton);
 	}
 }
+
 #endif //if WITH_EDITOR
 
 void UDataflowSkeletalContent::AddContentObjects(FReferenceCollector& Collector)
 {
 	Super::AddContentObjects(Collector);
-	
-	Collector.AddReferencedObject(SkeletalMeshComponent);
-	Collector.AddReferencedObject(AnimationNodeInstance);
-}
-
-FVector2f UDataflowSkeletalContent::GetSimulationRange() const
-{
-	if(AnimationNodeInstance)
-	{
-		return FVector2f(0.0f, AnimationNodeInstance->GetLength());
-	}
-	return FVector2f(0.0f, 0.0f);
 }
 
 void UDataflowSkeletalContent::AddReferencedObjects(UObject* InThis, FReferenceCollector& Collector)
@@ -254,7 +219,6 @@ void UDataflowSkeletalContent::AddReferencedObjects(UObject* InThis, FReferenceC
 	Collector.AddReferencedObject(This->SkeletalMesh);
 	Collector.AddReferencedObject(This->AnimationAsset);
 	Collector.AddReferencedObject(This->Skeleton);
-	Collector.AddReferencedObject(This->SkeletalMeshComponent);
-	Collector.AddReferencedObject(This->AnimationNodeInstance);
 	Super::AddReferencedObjects(InThis, Collector);
 }
+
diff --git a/Engine/Source/Runtime/Experimental/Dataflow/Engine/Public/Dataflow/DataflowContent.h b/Engine/Source/Runtime/Experimental/Dataflow/Engine/Public/Dataflow/DataflowContent.h
--- a/Engine/Source/Runtime/Experimental/Dataflow/Engine/Public/Dataflow/DataflowContent.h
+++ b/Engine/Source/Runtime/Experimental/Dataflow/Engine/Public/Dataflow/DataflowContent.h
@@ -77,6 +77,45 @@ protected:
 	Dataflow::EDataflowPatternVertexType ConstructionViewMode = Dataflow::EDataflowPatternVertexType::Sim3D;
 };
 
+UINTERFACE(MinimalAPI)
+class UDataflowContentOwner : public UInterface
+{
+	GENERATED_BODY()
+};
+
+/** 
+ * Dataflow interface for any content owner
+ */
+class DATAFLOWENGINE_API IDataflowContentOwner
+{
+public:
+
+	GENERATED_BODY()
+
+	/** Function to build the dataflow content */
+	TObjectPtr<UDataflowBaseContent> BuildDataflowContent();
+
+	/** Notification when owner changed */
+	DECLARE_MULTICAST_DELEGATE(FOnContentOwnerChanged);
+
+	/** Delegate member to be called in the invalidate */
+	FOnContentOwnerChanged OnContentOwnerChanged;
+
+	/** Invalidate all the dataflow contents */
+	void InvalidateDataflowContents() const
+	{
+		OnContentOwnerChanged.Broadcast();
+	}
+	
+	/** Interface to update a dataflow content instance from that owner */
+	virtual void UpdateDataflowContent(const TObjectPtr<UDataflowBaseContent>& DataflowContent) const = 0;
+	
+protected :
+
+	/** Interface to create a dataflow content instance from that owner */
+	virtual TObjectPtr<UDataflowBaseContent> CreateDataflowContent() = 0;
+};
+
 /** 
  * Dataflow content owning dataflow asset that that will be used to evaluate the graph
  */
@@ -87,6 +126,7 @@ class DATAFLOWENGINE_API UDataflowBaseContent : public UDataflowContextObject
 
 public:
 	UDataflowBaseContent();
+	~UDataflowBaseContent();
 
 	/** 
 	*	Dirty - State Invalidation
@@ -108,55 +148,46 @@ public:
 	*/
 	void SetDataflowContext(const TSharedPtr<Dataflow::FEngineContext>& InContext);
 	TSharedPtr<Dataflow::FEngineContext>& GetDataflowContext() { return DataflowContext; }
-	const TSharedPtr<Dataflow::FEngineContext>& GetDataflowContext() const { return DataflowContext; }
-
-	/** Return the simulation time range to be used in the simulation viewport */
-	virtual FVector2f GetSimulationRange() const { return FVector2f(0.0f, 100.0f); } 
-	
-	/** Register components to the simulation world */
-	virtual void RegisterWorldContent(FPreviewScene* PreviewScene, AActor* RootActor) {}
-	
-	/** Unregister components to the simulation world */
-	virtual void UnregisterWorldContent(FPreviewScene* PreviewScene) {}
 
-	/** Build the content context, timestamp*/
-	void BuildBaseContent(TObjectPtr<UObject> InDataflowOwner);
+	/** Rebuild the owner dependent datas  */
+	void UpdateContentDatas();
 
 	/** Collect reference objects for GC */
 	virtual void AddContentObjects(FReferenceCollector& Collector) {}
 	
-	/** Data flow owner accessors */
+	/** Data flow owner accessors (through the context) */
 	void SetDataflowOwner(const TObjectPtr<UObject>& InOwner);
 	TObjectPtr<UObject> GetDataflowOwner() const;
 	
-	/** Data flow asset accessors */
-	void SetDataflowAsset(const TObjectPtr<UDataflow>& InAsset) { DataflowAsset = InAsset;  SetIsDirty(true);}
-	const TObjectPtr<UDataflow>& GetDataflowAsset() const { return DataflowAsset; }
+	/** Data flow asset accessors (through the context) */
+	void SetDataflowAsset(const TObjectPtr<UDataflow>& InAsset);
+	TObjectPtr<UDataflow> GetDataflowAsset() const;
 
 	/** Data flow terminal accessors */
 	void SetDataflowTerminal(const FString& InPath) { DataflowTerminal = InPath;  SetIsDirty(true);}
 	const FString& GetDataflowTerminal() const { return DataflowTerminal; }
 
+	/** Terminal asset accessors */
+	void SetTerminalAsset(const TObjectPtr<UObject>& InAsset) { TerminalAsset = InAsset;  SetIsDirty(true);}
+	const TObjectPtr<UObject>& GetTerminalAsset() const { return TerminalAsset;}
+
+	/** Content Serialization */
 	virtual void Serialize(FArchive& Ar);
 
 	//~ UObject interface
 	static void AddReferencedObjects(UObject* InThis, FReferenceCollector& Collector);
 
 protected:
-
-	/** Data flow asset that we will edit */
-	UPROPERTY(EditAnywhere, Category = "Dataflow")
-	TObjectPtr<UObject> DataflowOwner = nullptr;
-
-	/** Data flow asset that we will edit */
-	UPROPERTY(EditAnywhere, Category = "Dataflow")
-	TObjectPtr<UDataflow> DataflowAsset = nullptr;
-
+	
 	/** Data flow terminal path for evaluation */
-	UPROPERTY(EditAnywhere, Category = "Dataflow")
+	UPROPERTY(EditAnywhere, Category = "Dataflow", Transient, SkipSerialization)
 	FString DataflowTerminal = "";
 
-	/**  Engine context to be used for dataflow evaluation */
+	/** Data flow terminal path for evaluation */
+	UPROPERTY(EditAnywhere, Category = "Dataflow", Transient, SkipSerialization)
+	TObjectPtr<UObject> TerminalAsset = nullptr;
+
+	/**  Engine context (data flow owner/asset) to be used for dataflow evaluation */
     TSharedPtr<Dataflow::FEngineContext> DataflowContext = nullptr;
 
     /** Last data flow evaluated node time stamp */
@@ -171,22 +202,13 @@ protected:
  * Dataflow content owning dataflow and skelmesh assets that that will be used to evaluate the graph
  */
 UCLASS()
-class DATAFLOWENGINE_API UDataflowSkeletalContent : public UDataflowBaseContent
+class DATAFLOWENGINE_API  UDataflowSkeletalContent : public UDataflowBaseContent
 {
 	GENERATED_BODY()
 
 public:
 	UDataflowSkeletalContent();
 	virtual ~UDataflowSkeletalContent() override{}
-
-	/** Return the simulation time range to be used in the simulation viewport */
-	virtual FVector2f GetSimulationRange() const override; 
-
-	/** Register components to the scene world */
-	virtual void RegisterWorldContent(FPreviewScene* PreviewScene, AActor* RootActor) override;
-	
-	/** Unregister components to the scene world */
-	virtual void UnregisterWorldContent(FPreviewScene* PreviewScene) override;
 	
 #if WITH_EDITOR
 	virtual void PostEditChangeProperty(FPropertyChangedEvent& PropertyChangedEvent) override;
@@ -195,13 +217,6 @@ public:
 	/** Collect reference objects for GC */
 	virtual void AddContentObjects(FReferenceCollector& Collector) override;
 
-	/** Update the animation node instance and set it to the skelmesh component*/
-	void UpdateAnimationInstance();
-
-	/** Data flow animation instance accessors */
-	const TObjectPtr<UAnimSingleNodeInstance>& GetAnimationInstance() const {return AnimationNodeInstance;}
-	TObjectPtr<UAnimSingleNodeInstance>& GetAnimationInstance() {return AnimationNodeInstance;}
-
 	/** Data flow skeletal mesh accessors */
 	void SetSkeletalMesh(const TObjectPtr<USkeletalMesh>& InMesh);
 	const TObjectPtr<USkeletalMesh>& GetSkeletalMesh() const { return SkeletalMesh; }
@@ -220,22 +235,14 @@ public:
 protected:
 
 	/** Data flow skeletal mesh*/
-	UPROPERTY(EditAnywhere, Category = "Preview", Transient, SkipSerialization)
+	UPROPERTY(Transient, SkipSerialization)
 	TObjectPtr<USkeletalMesh> SkeletalMesh = nullptr;
-
-	/** Animation asset to be used to preview simulation */
-	UPROPERTY(EditAnywhere, Category = "Preview", Transient, SkipSerialization)
-	TObjectPtr<UAnimationAsset> AnimationAsset = nullptr;
-
+	
 	/** Data flow skeleton */
-	UPROPERTY(EditAnywhere, Category = "Skeleton", Transient, SkipSerialization)
-	TObjectPtr<USkeleton> Skeleton = nullptr;
-
-	/** Skeletal mesh component used in the preview scene */
 	UPROPERTY(Transient, SkipSerialization)
-	TObjectPtr<USkeletalMeshComponent> SkeletalMeshComponent = nullptr;
-
-	/** Anim node instance used with skelmesh component */
+	TObjectPtr<USkeleton> Skeleton = nullptr;
+	
+	/** Animation asset to be used to preview simulation */
 	UPROPERTY(Transient, SkipSerialization)
-	TObjectPtr<UAnimSingleNodeInstance> AnimationNodeInstance = nullptr;
+	TObjectPtr<UAnimationAsset> AnimationAsset = nullptr;
 };
diff --git a/Engine/Source/Runtime/Experimental/GeometryCollectionEngine/Private/GeometryCollection/GeometryCollectionObject.cpp b/Engine/Source/Runtime/Experimental/GeometryCollectionEngine/Private/GeometryCollection/GeometryCollectionObject.cpp
--- a/Engine/Source/Runtime/Experimental/GeometryCollectionEngine/Private/GeometryCollection/GeometryCollectionObject.cpp
+++ b/Engine/Source/Runtime/Experimental/GeometryCollectionEngine/Private/GeometryCollection/GeometryCollectionObject.cpp
@@ -1844,6 +1844,7 @@ void UGeometryCollection::PostEditChangeProperty(struct FPropertyChangedEvent& P
 	{
 		RebuildRenderData();
 	}
+	InvalidateDataflowContents();
 }
 
 bool UGeometryCollection::Modify(bool bAlwaysMarkDirty /*= true*/)
@@ -2026,6 +2027,27 @@ const TArray<UAssetUserData*>* UGeometryCollection::GetAssetUserDataArray() cons
 	return &ToRawPtrTArrayUnsafe(AssetUserData);
 }
 
+TObjectPtr<UDataflowBaseContent> UGeometryCollection::CreateDataflowContent()
+{
+	TObjectPtr<UDataflowBaseContent> BaseContent = NewObject<UDataflowBaseContent>();
+
+	BaseContent->SetDataflowOwner(this);
+	BaseContent->SetTerminalAsset(this);
+	
+	UpdateDataflowContent(BaseContent);
+	
+	return BaseContent;
+}
+
+void UGeometryCollection::UpdateDataflowContent(const TObjectPtr<UDataflowBaseContent>& DataflowContent) const
+{
+	if(const TObjectPtr<UDataflowBaseContent> BaseContent = Cast<UDataflowBaseContent>(DataflowContent))
+	{
+		BaseContent->SetDataflowAsset(DataflowAsset);
+		BaseContent->SetDataflowTerminal(DataflowTerminal);
+	}
+}
+
 #if WITH_EDITOR
 bool UGeometryCollection::CanEditChange(const FProperty* InProperty) const
 {
diff --git a/Engine/Source/Runtime/Experimental/GeometryCollectionEngine/Public/GeometryCollection/GeometryCollectionObject.h b/Engine/Source/Runtime/Experimental/GeometryCollectionEngine/Public/GeometryCollection/GeometryCollectionObject.h
--- a/Engine/Source/Runtime/Experimental/GeometryCollectionEngine/Public/GeometryCollection/GeometryCollectionObject.h
+++ b/Engine/Source/Runtime/Experimental/GeometryCollectionEngine/Public/GeometryCollection/GeometryCollectionObject.h
@@ -6,6 +6,7 @@
 #include "Chaos/ChaosSolverActor.h"
 #include "GeometryCollection/GeometryCollectionDamagePropagationData.h"
 #include "GeometryCollection/GeometryCollectionSimulationTypes.h"
+#include "Dataflow/DataflowContent.h"
 #include "GeometryCollection/ManagedArray.h"
 #include "InstanceUniformShaderParameters.h"
 #include "Interfaces/Interface_AssetUserData.h"
@@ -357,7 +358,7 @@ struct FGeometryCollectionRenderResourceSizeInfo
 *
 */
 UCLASS(BlueprintType, customconstructor, MinimalAPI)
-class UGeometryCollection : public UObject, public IInterface_AssetUserData
+class UGeometryCollection : public UObject, public IInterface_AssetUserData, public IDataflowContentOwner
 {
 	GENERATED_UCLASS_BODY()
 
@@ -827,6 +828,10 @@ public:
 	UPROPERTY(EditAnywhere, BlueprintReadWrite, Category = "Dataflow", DisplayName = "DataFlow Overrides")
 	TMap<FString, FString> Overrides;
 
+	/** IDataflowContentOwner interface */
+	virtual TObjectPtr<UDataflowBaseContent> CreateDataflowContent() override;
+	virtual void UpdateDataflowContent(const TObjectPtr<UDataflowBaseContent>& DataflowContent) const override;
+
 	GEOMETRYCOLLECTIONENGINE_API const TArray<int32>& GetBreadthFirstTransformIndices() const { return BreadthFirstTransformIndices; }
 
 	GEOMETRYCOLLECTIONENGINE_API const TArray<int32>& GetAutoInstanceTransformRemapIndices() const { return AutoInstanceTransformRemapIndices; }
