4- Simulation dataflow nodes/editor

diff --git a/Engine/Plugins/ChaosClothAssetEditor/Source/ChaosClothAssetEditor/Private/ChaosClothAsset/ClothEditorToolkit.cpp b/Engine/Plugins/ChaosClothAssetEditor/Source/ChaosClothAssetEditor/Private/ChaosClothAsset/ClothEditorToolkit.cpp
--- a/Engine/Plugins/ChaosClothAssetEditor/Source/ChaosClothAssetEditor/Private/ChaosClothAsset/ClothEditorToolkit.cpp
+++ b/Engine/Plugins/ChaosClothAssetEditor/Source/ChaosClothAssetEditor/Private/ChaosClothAsset/ClothEditorToolkit.cpp
@@ -907,7 +907,7 @@ void FChaosClothAssetEditorToolkit::EvaluateNode(FDataflowNode* Node, bool bForc
 
 		// Set all asset values by visiting the graph from the specified node
 		TArray<FDataflowTerminalNode*> TerminalNodes;
-		TerminalNodes.Reserve(Graph->GetTerminalNodes().Num());
+		TerminalNodes.Reserve(Graph->GetFilteredNodes(FDataflowTerminalNode::StaticType()).Num());
 		
 		TQueue<FDataflowNode*> Queue;
 		Queue.Enqueue(Node);
diff --git a/Engine/Plugins/Experimental/Dataflow/Dataflow.uplugin b/Engine/Plugins/Experimental/Dataflow/Dataflow.uplugin
--- a/Engine/Plugins/Experimental/Dataflow/Dataflow.uplugin
+++ b/Engine/Plugins/Experimental/Dataflow/Dataflow.uplugin
@@ -54,6 +54,10 @@
 		{
 			"Name": "MeshModelingToolsetExp",
 			"Enabled": true
+		},
+		{
+			"Name": "ChaosCaching",
+			"Enabled": true
 		}
 	]
 }
diff --git a/Engine/Plugins/Experimental/Dataflow/Source/DataflowEditor/DataflowEditor.Build.cs b/Engine/Plugins/Experimental/Dataflow/Source/DataflowEditor/DataflowEditor.Build.cs
--- a/Engine/Plugins/Experimental/Dataflow/Source/DataflowEditor/DataflowEditor.Build.cs
+++ b/Engine/Plugins/Experimental/Dataflow/Source/DataflowEditor/DataflowEditor.Build.cs
@@ -65,7 +65,8 @@ namespace UnrealBuildTool.Rules
 					"XmlParser",
 					"EditorWidgets",
 					"KismetWidgets",      // SScrubControlPanel
-					"AnimGraph"           // UAnimSingleNodeInstance
+					"AnimGraph", 
+					"ChaosCaching" // UAnimSingleNodeInstance
 				}
 			);
 		}
diff --git a/Engine/Plugins/Experimental/Dataflow/Source/DataflowEditor/Private/Dataflow/DataflowEditor.cpp b/Engine/Plugins/Experimental/Dataflow/Source/DataflowEditor/Private/Dataflow/DataflowEditor.cpp
--- a/Engine/Plugins/Experimental/Dataflow/Source/DataflowEditor/Private/Dataflow/DataflowEditor.cpp
+++ b/Engine/Plugins/Experimental/Dataflow/Source/DataflowEditor/Private/Dataflow/DataflowEditor.cpp
@@ -87,7 +87,7 @@ void UDataflowEditor::RemoveTerminalContents(const TSharedPtr<Dataflow::FGraph>&
 {
 	for(int32 ContentIndex = TerminalContents.Num()-1; ContentIndex >= 0; --ContentIndex)
 	{
-		if(TSharedPtr<FDataflowNode> TerminalNode = DataflowGraph->FindTerminalNode(
+		if(TSharedPtr<FDataflowNode> TerminalNode = DataflowGraph->FindFilteredNode(FDataflowTerminalNode::StaticType(),
 			FName(TerminalContents[ContentIndex]->GetDataflowTerminal())))
 		{
 			if(Cast<IDataflowContentOwner>(TerminalNode->AsType<FDataflowTerminalNode>()->GetTerminalAsset()))
@@ -121,7 +121,7 @@ void UDataflowEditor::AddTerminalContents(const TSharedPtr<Dataflow::FGraph>& Da
 		TerminalContent->SetTerminalAsset(TerminalAsset);
 	};
 			
-	for(const TSharedPtr<FDataflowNode>& DataflowNode : DataflowGraph->GetTerminalNodes())
+	for(const TSharedPtr<FDataflowNode>& DataflowNode : DataflowGraph->GetFilteredNodes(FDataflowTerminalNode::StaticType()))
 	{
 		if(const FDataflowTerminalNode* TerminalNode = DataflowNode->AsType<FDataflowTerminalNode>())
 		{
diff --git a/Engine/Plugins/Experimental/Dataflow/Source/DataflowEditor/Private/Dataflow/DataflowEditorCommands.cpp b/Engine/Plugins/Experimental/Dataflow/Source/DataflowEditor/Private/Dataflow/DataflowEditorCommands.cpp
--- a/Engine/Plugins/Experimental/Dataflow/Source/DataflowEditor/Private/Dataflow/DataflowEditorCommands.cpp
+++ b/Engine/Plugins/Experimental/Dataflow/Source/DataflowEditor/Private/Dataflow/DataflowEditorCommands.cpp
@@ -33,6 +33,7 @@
 
 const FString FDataflowEditorCommandsImpl::BeginWeightMapPaintToolIdentifier = TEXT("BeginWeightMapPaintTool");
 const FString FDataflowEditorCommandsImpl::AddWeightMapNodeIdentifier = TEXT("AddWeightMapNode");
+const FString FDataflowEditorCommandsImpl::UpdateSimulationCacheIdentifier = TEXT("UpdateSimulationCache");
 
 // @todo(brice) Remove Example Tools
 //const FString FDataflowEditorCommandsImpl::BeginAttributeEditorToolIdentifier = TEXT("BeginAttributeEditorTool");
@@ -62,6 +63,7 @@ void FDataflowEditorCommandsImpl::RegisterCommands()
 
 	UI_COMMAND(BeginWeightMapPaintTool, "Add Weight Map", "Paint weight maps on the mesh", EUserInterfaceActionType::None, FInputChord());
 	UI_COMMAND(AddWeightMapNode, "Add Weight Map", "Paint weight maps on the mesh", EUserInterfaceActionType::Button, FInputChord());
+	UI_COMMAND(UpdateSimulationCache, "UpdateSimulationCache", "Update Simulation Cache", EUserInterfaceActionType::ToggleButton, FInputChord());
 
 
 	if (Dataflow::FNodeFactory* Factory = Dataflow::FNodeFactory::GetInstance())
@@ -253,7 +255,6 @@ void FDataflowEditorCommands::EvaluateTerminalNode(Dataflow::FContext& Context,
 	}
 }
 
-
 bool FDataflowEditorCommands::OnNodeVerifyTitleCommit(const FText& NewText, UEdGraphNode* GraphNode, FText& OutErrorMessage)
 {
 	if (GraphNode)
diff --git a/Engine/Plugins/Experimental/Dataflow/Source/DataflowEditor/Private/Dataflow/DataflowEditorStyle.h b/Engine/Plugins/Experimental/Dataflow/Source/DataflowEditor/Private/Dataflow/DataflowEditorStyle.h
--- a/Engine/Plugins/Experimental/Dataflow/Source/DataflowEditor/Private/Dataflow/DataflowEditorStyle.h
+++ b/Engine/Plugins/Experimental/Dataflow/Source/DataflowEditor/Private/Dataflow/DataflowEditorStyle.h
@@ -23,6 +23,8 @@ public:
 
 		const FVector2D ToolbarIconSize = Icon20x20;
 
+		SetCoreContentRoot(FPaths::EngineContentDir() / TEXT("Editor/Slate"));
+
 		SetContentRoot(IPluginManager::Get().FindPlugin("Dataflow")->GetBaseDir() / TEXT("Resources"));
 		Set("ClassIcon.Dataflow", new FSlateVectorImageBrush(RootToContentDir(TEXT("DataflowAsset_16.svg")), Icon16x16));
 		Set("ClassThumbnail.Dataflow", new FSlateVectorImageBrush(RootToContentDir(TEXT("DataflowAsset_64.svg")), Icon64x64));
@@ -43,10 +45,13 @@ public:
 		//Set(*AttributeEditorPropertyName, new FSlateImageBrush(RootToContentDir(TEXT("Slate/Dataflow_SelectObject_20x.png")), Icon20x20));
 		//const FString MeshSelectionPropertyName = "DataflowEditor." + FDataflowEditorCommandsImpl::BeginMeshSelectionToolIdentifier;
 		//Set(*MeshSelectionPropertyName, new FSlateImageBrush(RootToContentDir(TEXT("Slate/Dataflow_SelectVertex_20x.png")), Icon20x20));
-
+		
 		FString PropertyNameString = "DataflowEditor." + FDataflowEditorCommandsImpl::AddWeightMapNodeIdentifier;
 		Set(*PropertyNameString, new IMAGE_BRUSH_SVG("Icons/PaintMaps", ToolbarIconSize));
 
+		PropertyNameString = "DataflowEditor." + FDataflowEditorCommandsImpl::UpdateSimulationCacheIdentifier;
+		Set(*PropertyNameString, new IMAGE_BRUSH_SVG("Icons/LV_Record", ToolbarIconSize));
+
 		DefaultMaterial = Cast<UMaterial>(StaticLoadObject(UMaterial::StaticClass(), NULL, TEXT("/Engine/BasicShapes/BasicShapeMaterial")));
 		VertexMaterial = Cast<UMaterial>(StaticLoadObject(UMaterial::StaticClass(), NULL, TEXT("/Engine/EditorMaterials/Dataflow/DataflowVertexMaterial")));
 		FSlateStyleRegistry::RegisterSlateStyle(*this);
diff --git a/Engine/Plugins/Experimental/Dataflow/Source/DataflowEditor/Private/Dataflow/DataflowEditorToolkit.cpp b/Engine/Plugins/Experimental/Dataflow/Source/DataflowEditor/Private/Dataflow/DataflowEditorToolkit.cpp
--- a/Engine/Plugins/Experimental/Dataflow/Source/DataflowEditor/Private/Dataflow/DataflowEditorToolkit.cpp
+++ b/Engine/Plugins/Experimental/Dataflow/Source/DataflowEditor/Private/Dataflow/DataflowEditorToolkit.cpp
@@ -3,7 +3,9 @@
 #include "Dataflow/DataflowEditorToolkit.h"
 
 #include "AdvancedPreviewScene.h"
+#include "AdvancedPreviewSceneModule.h"
 #include "AssetEditorModeManager.h"
+#include "DetailCategoryBuilder.h"
 #include "Animation/Skeleton.h"
 #include "Dataflow/DataflowCore.h"
 #include "Dataflow/DataflowEditor.h"
@@ -40,6 +42,7 @@
 #include "Framework/Commands/GenericCommands.h"
 #include "GameFramework/Actor.h"
 #include "GraphEditorActions.h"
+#include "IDetailCustomization.h"
 #include "ISkeletonTree.h"
 #include "IStructureDetailsView.h"
 #include "Kismet/GameplayStatics.h"
@@ -61,6 +64,7 @@ FAutoConsoleVariableRef CVARDataflowEnableSkeletonView(TEXT("p.Dataflow.Editor.E
 
 const FName FDataflowEditorToolkit::GraphCanvasTabId(TEXT("DataflowEditor_GraphCanvas"));
 const FName FDataflowEditorToolkit::NodeDetailsTabId(TEXT("DataflowEditor_NodeDetails"));
+const FName FDataflowEditorToolkit::PreviewSceneTabId(TEXT("DataflowEditor_PreviewScene"));
 const FName FDataflowEditorToolkit::SkeletonViewTabId(TEXT("DataflowEditor_SkeletonView"));
 const FName FDataflowEditorToolkit::SelectionViewTabId_1(TEXT("DataflowEditor_SelectionView_1"));
 const FName FDataflowEditorToolkit::SelectionViewTabId_2(TEXT("DataflowEditor_SelectionView_2"));
@@ -438,6 +442,14 @@ void FDataflowEditorToolkit::CreateWidgets()
 			AssetDetailsEditor = CreateAssetDetailsEditorWidget({EditorContent->GetDataflowOwner(),EditorContent->GetDataflowAsset()});
 			GraphEditor = CreateGraphEditorWidget(DataflowAsset, NodeDetailsEditor);
 			CreateSimulationViewportClient();
+
+			FAdvancedPreviewSceneModule& AdvancedPreviewSceneModule = FModuleManager::LoadModuleChecked<FAdvancedPreviewSceneModule>("AdvancedPreviewScene");
+			
+			AdvancedPreviewSettingsWidget = AdvancedPreviewSceneModule.CreateAdvancedPreviewSceneSettingsWidget(SimulationScene.ToSharedRef(), 
+				SimulationScene->GetPreviewSceneDescription(),
+				TArray<FAdvancedPreviewSceneModule::FDetailCustomizationInfo>(), 
+				TArray<FAdvancedPreviewSceneModule::FPropertyTypeCustomizationInfo>(),
+				TArray<FAdvancedPreviewSceneModule::FDetailDelegates>());
 		}
 	}
 }
@@ -495,7 +507,8 @@ void FDataflowEditorToolkit::CreateSimulationViewportClient()
 	SimulationViewportDelegate = [this](FAssetEditorViewportConstructionArgs InArgs)
 	{
 		TSharedRef<SDataflowSimulationViewport> Viewport = SAssignNew(DataflowSimulationViewport, SDataflowSimulationViewport, InArgs)
-			.ViewportClient(StaticCastSharedPtr<FDataflowSimulationViewportClient>(SimulationViewportClient));
+			.ViewportClient(StaticCastSharedPtr<FDataflowSimulationViewportClient>(SimulationViewportClient))
+			.CommandList(GetToolkitCommands().ToSharedPtr());;
 
 		if (UDataflowEditorMode* const DataflowMode = Cast<UDataflowEditorMode>(EditorModeManager->GetActiveScriptableMode(UDataflowEditorMode::EM_DataflowEditorModeId)))
 		{
@@ -929,6 +942,17 @@ TSharedRef<SDockTab> FDataflowEditorToolkit::SpawnTab_SimulationViewport(const F
 	return DockableTab;
 }
 
+TSharedRef<SDockTab> FDataflowEditorToolkit::SpawnTab_PreviewScene(const FSpawnTabArgs& Args)
+{
+	check(Args.GetTabId() == PreviewSceneTabId);
+
+	return SNew(SDockTab)
+		.Label(LOCTEXT("DataflowEditor_PreviewScene_TabTitle", "PreviewScene"))
+		[
+			AdvancedPreviewSettingsWidget->AsShared()
+		];
+}
+
 TSharedRef<SDockTab> FDataflowEditorToolkit::SpawnTab_GraphCanvas(const FSpawnTabArgs& Args)
 {
 	check(Args.GetTabId() == GraphCanvasTabId);
@@ -986,7 +1010,7 @@ TSharedRef<SDockTab> FDataflowEditorToolkit::SpawnTab_SelectionView(const FSpawn
 	check(DataflowEditor);
 	check(DataflowEditor->GetEditorContent());
 
-		if (Args.GetTabId() == SelectionViewTabId_1)
+	if (Args.GetTabId() == SelectionViewTabId_1)
 	{
 		DataflowSelectionView_1 = MakeShared<FDataflowSelectionView>(FDataflowSelectionView(DataflowEditor->GetEditorContent()));
 		if (DataflowSelectionView_1.IsValid())
@@ -1149,6 +1173,11 @@ void FDataflowEditorToolkit::RegisterTabSpawners(const TSharedRef<FTabManager>&
 		.SetDisplayName(LOCTEXT("AssetDetailsTab", "Asset Details"))
 		.SetGroup(EditorMenuCategory.ToSharedRef())
 		.SetIcon(FSlateIcon(FAppStyle::GetAppStyleSetName(), "LevelEditor.Tabs.Details"));
+
+	InTabManager->RegisterTabSpawner(PreviewSceneTabId, FOnSpawnTab::CreateSP(this, &FDataflowEditorToolkit::SpawnTab_PreviewScene))
+		.SetDisplayName(LOCTEXT("PreviewSceneTab", "PreviewScene"))
+		.SetGroup(EditorMenuCategory.ToSharedRef())
+		.SetIcon(FSlateIcon(FAppStyle::GetAppStyleSetName(), "LevelEditor.Tabs.PreviewScene"));
 	
 	InTabManager->RegisterTabSpawner(GraphCanvasTabId, FOnSpawnTab::CreateSP(this, &FDataflowEditorToolkit::SpawnTab_GraphCanvas))
 		.SetDisplayName(LOCTEXT("DataflowTab", "Dataflow Graph"))
diff --git a/Engine/Plugins/Experimental/Dataflow/Source/DataflowEditor/Private/Dataflow/DataflowSimulationGenerator.cpp b/Engine/Plugins/Experimental/Dataflow/Source/DataflowEditor/Private/Dataflow/DataflowSimulationGenerator.cpp
new file mode 100644
--- /dev/null
+++ b/Engine/Plugins/Experimental/Dataflow/Source/DataflowEditor/Private/Dataflow/DataflowSimulationGenerator.cpp
@@ -0,0 +1,392 @@
+﻿// Copyright Epic Games, Inc. All Rights Reserved.
+
+#include "Dataflow/DataflowSimulationGenerator.h"
+
+#include "Engine/World.h"
+#include "Chaos/CacheManagerActor.h"
+#include "Dataflow/DataflowObject.h"
+#include "Dataflow/DataflowSolverNodes.h"
+#include "Misc/AsyncTaskNotification.h" 
+
+DEFINE_LOG_CATEGORY(LogDataflowSimulationGenerator);
+
+#define LOCTEXT_NAMESPACE "DataflowSimulationGenerator"
+
+namespace Dataflow
+{
+	float GetCacheDuration(const TObjectPtr<UDataflow>& SimulationGraph, const TSharedPtr<Dataflow::FSimulationContext>& SimulationContext, const int32 SamplingRate, float& MinTime, float& MaxTime, int32& NumFrames)
+	{
+		float CacheDuration = 0.0f;
+		if(SimulationContext.IsValid())
+		{
+			SimulationContext->SetSimulationFlag(Dataflow::ESimulationFlags::SetupSimulation);
+		}
+		if(const TSharedPtr<Dataflow::FGraph> DataflowGraph = SimulationGraph->GetDataflow())
+		{
+			MinTime = TNumericLimits<float>::Max();
+			MaxTime = TNumericLimits<float>::Lowest();
+			for(const TSharedPtr<FDataflowNode>& CacheNode : DataflowGraph->GetFilteredNodes(FDataflowCacheNode::StaticType()))
+			{
+				const FVector2f NodeRange = StaticCastSharedPtr<FDataflowCacheNode>(CacheNode)->GetTimeRange(*SimulationContext.Get());
+				MinTime = FMath::Min(MinTime, NodeRange[0]);
+				MaxTime = FMath::Max(MaxTime, NodeRange[1]);
+			}
+			NumFrames = (MaxTime > MinTime) ? FMath::Floor((MaxTime - MinTime) * SamplingRate) : 0;
+			CacheDuration = ( MaxTime > MinTime) ? (MaxTime - MinTime) : 0.0f;
+		}
+		if(SimulationContext.IsValid())
+		{
+			SimulationContext->UnsetSimulationFlag(Dataflow::ESimulationFlags::SetupSimulation);
+		}
+		return CacheDuration;
+	}
+
+	bool ShouldResetWorld(const TObjectPtr<UDataflow>& SimulationGraph, const TSharedPtr<Dataflow::FSimulationContext>& SimulationContext, Dataflow::FTimestamp& LastTimeStamp)
+	{
+		if(const TSharedPtr<Dataflow::FGraph> DataflowGraph = SimulationGraph->GetDataflow())
+		{
+			for(const TSharedPtr<FDataflowNode>& CacheNode : DataflowGraph->GetFilteredNodes(FDataflowCacheNode::StaticType()))
+			{
+				if(CacheNode->GetTimestamp() >= LastTimeStamp)
+				{
+					LastTimeStamp = SimulationContext->GetTimestamp();
+					return true;
+				}
+			}
+		}
+		return false;
+	}
+
+	void UpdateAnimationNodes(const TObjectPtr<UDataflow>& SimulationGraph, const TSharedPtr<Dataflow::FSimulationContext>& SimulationContext, const float AnimationTime)
+	{
+		if(const TSharedPtr<Dataflow::FGraph> DataflowGraph = SimulationGraph->GetDataflow())
+		{
+			for(const TSharedPtr<FDataflowNode>& AnimationNode : DataflowGraph->GetFilteredNodes(FDataflowAnimationNode::StaticType()))
+			{
+				StaticCastSharedPtr<FDataflowAnimationNode>(AnimationNode)->SetAnimationTime(*SimulationContext.Get(), AnimationTime);
+			}
+		}
+	}
+	
+	void EvaluateSimulationGraph(const TObjectPtr<UDataflow>& SimulationGraph, const TSharedPtr<Dataflow::FSimulationContext>& SimulationContext,
+		const float DeltaTime, const float SimulationTime, const ESimulationFlags SimulationFlags)
+	{
+		if(SimulationContext.IsValid())
+		{
+			SimulationContext->SetSimulationFlag(SimulationFlags);
+			if(SimulationContext->HasSimulationFlag(ESimulationFlags::UpdateSimulation))
+			{
+				SimulationContext->SetTimingInfos(DeltaTime, SimulationTime);
+			}
+		}
+
+		if(SimulationGraph)
+		{
+			if(const TSharedPtr<Dataflow::FGraph> DataflowGraph = SimulationGraph->GetDataflow())
+			{
+				if(SimulationContext->HasSimulationFlag(ESimulationFlags::UpdateSimulation))
+				{
+					// Invalidation of all the simulation nodes that are always dirty
+					for(const TSharedPtr<FDataflowNode>& SimulationNode : DataflowGraph->GetFilteredNodes(FDataflowInvalidNode::StaticType()))
+					{
+						SimulationNode->Invalidate();
+					}
+				}
+
+				// Pull the graph evaluation from the cache nodes
+				for(const TSharedPtr<FDataflowNode>& CacheNode : DataflowGraph->GetFilteredNodes(FDataflowCacheNode::StaticType()))
+				{
+					SimulationContext->Evaluate(CacheNode.Get(), nullptr);
+				}
+			}
+		}
+		if(SimulationContext.IsValid())
+		{
+			SimulationContext->UnsetSimulationFlag(SimulationFlags);
+		}
+	}
+
+	void CleanSimulationContext(TSharedPtr<Dataflow::FSimulationContext>& SimulationContext)
+	{
+		if(const TObjectPtr<AChaosCacheManager> CacheManager = Cast<AChaosCacheManager>(SimulationContext->GetRootActor()))
+		{
+			if(SimulationContext->GetCachingMode() == ECachingMode::RecordCache)
+			{
+				CacheManager->EndEvaluate();
+			}
+
+			//  Clear the observed components
+			CacheManager->ClearObservedComponents();
+			
+			TInlineComponentArray<UPrimitiveComponent*> PrimComponents;
+			CacheManager->GetComponents(PrimComponents);
+
+			// Unregister all the components
+			for(UPrimitiveComponent* PrimComponent : PrimComponents)
+			{
+				PrimComponent->SelectionOverrideDelegate.Unbind();
+				PrimComponent->UnregisterComponent();
+				PrimComponent->DestroyComponent();
+			}
+		}
+
+		// Reset the simulation context
+		SimulationContext.Reset();
+	}
+
+	void BuildSimulationContext(const TObjectPtr<UObject>& SimulationOwner, const TObjectPtr<UDataflow>& SimulationGraph, const TObjectPtr<AChaosCacheManager>& CacheManager,
+		const ECachingMode CachingMode, TSharedPtr<Dataflow::FSimulationContext>& SimulationContext)
+	{
+		if(CacheManager)
+		{
+			// Create a simulation context that will hold all the simulation datas
+			SimulationContext = MakeShared<Dataflow::FSimulationContext>(
+						SimulationOwner, SimulationGraph,
+					Dataflow::FTimestamp::Invalid, CacheManager, CachingMode);
+
+			// Evaluate the setup stage of the simulation graph
+			EvaluateSimulationGraph(SimulationGraph, SimulationContext, 
+					0.0f, 0.0f , ESimulationFlags::SetupSimulation);
+
+			// Init the cache manager
+			CacheManager->SetObservedComponentProperties(CacheManager->CacheMode);
+
+			if(SimulationContext->GetCachingMode() == ECachingMode::RecordCache)
+			{
+				CacheManager->BeginEvaluate();
+			}
+		}
+	}
+	
+	void FDataflowSimulationTask::DoWork()
+	{
+		if(SimulationGraph)
+		{
+			const int32 NumFrames = (MaxTime-MinTime) / DeltaTime;
+			for (int32 FrameIndex = 0; FrameIndex < NumFrames; ++FrameIndex)
+			{
+				if (!TaskManager->bCancelled.load())
+				{
+					// Compute the simulation time that will be sent to the graph
+					const float SimulationTime = MinTime + (FrameIndex+1) * DeltaTime;
+
+					// Evaluate the simulation graph update stage
+					EvaluateSimulationGraph(SimulationGraph, SimulationContext,
+						DeltaTime, SimulationTime, ESimulationFlags::UpdateSimulation);
+
+					// Finish the frame
+					TaskManager->SimulationResource->FinishFrame();
+				}
+				else
+				{
+					break;
+				}
+			}
+		}
+	}
+	
+	bool FDataflowTaskManager::AllocateSimulationResource(const TObjectPtr<UDataflow> SimulationGraph, const int32 SamplingRate, const TObjectPtr<UChaosCacheCollection>& CacheCollection)
+	{
+		SimulationWorld = UWorld::CreateWorld(EWorldType::Editor, false);
+		TObjectPtr<AChaosCacheManager> CacheManager = SimulationWorld->SpawnActor<AChaosCacheManager>(AChaosCacheManager::StaticClass());
+
+		CacheManager->StartMode = EStartMode::Timed;
+		CacheManager->CacheMode = ECacheMode::Record;
+		CacheManager->CacheCollection = CacheCollection;
+		
+		SimulationResource = MakeShared<FDataflowSimulationResource>();
+		SimulationResource->NumSimulatedFrames = &NumSimulatedFrames;
+		SimulationResource->bCancelled = &bCancelled;
+		
+		if(SimulationTask.IsValid())
+		{
+			// Clone the simulation graph to avoid evaluating the graph on another thread while modifying it on the game thread
+			SimulationTask->GetTask().SimulationGraph = DuplicateObject<UDataflow>(SimulationGraph, SimulationWorld, NAME_None);
+
+			// Build the simulation context
+			BuildSimulationContext(SimulationTask->GetTask().SimulationGraph, SimulationTask->GetTask().SimulationGraph,
+				CacheManager, ECachingMode::RecordCache, SimulationTask->GetTask().SimulationContext);
+
+			// Compute the simulation duration in order to run the simulation
+			const float Duration = GetCacheDuration(SimulationTask->GetTask().SimulationGraph, SimulationTask->GetTask().SimulationContext, SamplingRate,
+				SimulationTask->GetTask().MinTime, SimulationTask->GetTask().MaxTime, NumFrames);
+
+			// Update the delta time 
+			SimulationTask->GetTask().DeltaTime = Duration / NumFrames;
+		}
+
+		return true;
+	}
+	
+	void FDataflowTaskManager::FreeSimulationResource()
+	{
+		if (SimulationTask.IsValid())
+		{
+			SimulationTask->EnsureCompletion();
+
+			CleanSimulationContext(SimulationTask->GetTask().SimulationContext);
+		}
+		
+		SimulationResource.Reset();
+		SimulationWorld->DestroyWorld(false);
+	}
+
+	void FDataflowTaskManager::CancelSimulationGeneration()
+	{
+		bCancelled.store(true);
+		SimulationTask->TryAbandonTask();
+	}
+	
+	FDataflowSimulationGenerator::FDataflowSimulationGenerator()
+	{}
+
+	FDataflowSimulationGenerator::~FDataflowSimulationGenerator()
+	{
+		if (TaskManager != nullptr)
+		{
+			TaskManager->FreeSimulationResource();
+		}
+	}
+
+	void FDataflowSimulationGenerator::Tick(float DeltaTime)
+	{
+		if (PendingAction == EDataflowGeneratorActions::StartGenerate)
+		{
+			StartGenerateSimulation();
+		}
+		else if (PendingAction == EDataflowGeneratorActions::TickGenerate)
+		{
+			TickGenerateSimulation();
+		}
+	}
+
+	TStatId FDataflowSimulationGenerator::GetStatId() const
+	{
+		RETURN_QUICK_DECLARE_CYCLE_STAT(FDataflowSimulationGenerator, STATGROUP_Tickables);
+	}
+
+	void FDataflowSimulationGenerator::StartGenerateSimulation()
+	{
+		check(PendingAction == EDataflowGeneratorActions::StartGenerate);
+		
+		if (TaskManager != nullptr)
+		{
+			UE_LOG(LogDataflowSimulationGenerator, Error, TEXT("Previous generation is still running."));
+			PendingAction = EDataflowGeneratorActions::NoAction;
+			return;
+		}
+		if(!SimulationGraph)
+		{
+			UE_LOG(LogDataflowSimulationGenerator, Error, TEXT("No simulation graph available in the generator"));
+			PendingAction = EDataflowGeneratorActions::NoAction;
+			return;
+		}
+		TaskManager = TSharedPtr<FDataflowTaskManager>(new FDataflowTaskManager);
+
+		TaskManager->SimulationTask = MakeUnique<FAsyncTask<FDataflowSimulationTask>>();
+		TaskManager->SimulationTask->GetTask().TaskManager = TaskManager;
+		
+		TaskManager->AllocateSimulationResource(SimulationGraph, SamplingRate, CacheCollection);
+		
+		TaskManager->SimulationTask->StartBackgroundTask();
+	
+		FAsyncTaskNotificationConfig NotificationConfig;
+		NotificationConfig.TitleText = LOCTEXT("SimulateDataflow", "Simulating Dataflow Content");
+		NotificationConfig.ProgressText = FText::FromString(TEXT("0%"));
+		NotificationConfig.bCanCancel = true;
+		NotificationConfig.bKeepOpenOnSuccess = true;
+		NotificationConfig.bKeepOpenOnFailure = true;
+		TaskManager->AsyncNotification = MakeUnique<FAsyncTaskNotification>(NotificationConfig);
+		TaskManager->StartTime = FDateTime::UtcNow();
+		TaskManager->LastUpdateTime = TaskManager->StartTime;
+	
+		PendingAction = EDataflowGeneratorActions::TickGenerate;
+	}
+	
+	void FDataflowSimulationGenerator::TickGenerateSimulation()
+	{
+		check(PendingAction == EDataflowGeneratorActions::TickGenerate && TaskManager != nullptr);
+			
+		bool bFinished = false;
+		const bool bCancelled = TaskManager->AsyncNotification->GetPromptAction() == EAsyncTaskNotificationPromptAction::Cancel;
+		if (TaskManager->SimulationTask->IsDone())
+		{
+			bFinished = true;
+		}
+		else if (bCancelled)
+		{
+			TaskManager->CancelSimulationGeneration();
+			bFinished = true;
+		}
+			
+		if (!bFinished)
+		{
+			const FDateTime CurrentTime = FDateTime::UtcNow();
+			const double SinceLastUpdate = (CurrentTime - TaskManager->LastUpdateTime).GetTotalSeconds();
+			if (SinceLastUpdate < 0.2)
+			{
+				return;
+			}
+			
+			const int32 NumSimulatedFrames = TaskManager->NumSimulatedFrames.load();
+			const int32 NumTotalFrames = TaskManager->NumFrames;
+			const FText ProgressMessage = FText::FromString(FString::Printf(TEXT("Finished %d/%d, %.1f%%"), NumSimulatedFrames, NumTotalFrames, 100.0 * NumSimulatedFrames / NumTotalFrames));
+			TaskManager->AsyncNotification->SetProgressText(ProgressMessage);
+			TaskManager->LastUpdateTime = CurrentTime;
+		}
+		else
+		{
+			FreeTaskResource(bCancelled);
+			PendingAction = EDataflowGeneratorActions::NoAction;
+		}
+	}
+
+	void FDataflowSimulationGenerator::SetSimulationGraph(const TObjectPtr<UDataflow>& DataflowGraph) 
+	{
+		SimulationGraph = DataflowGraph;
+	}
+
+	void FDataflowSimulationGenerator::SetSamplingRate(const int32 FrameRate) 
+	{
+		SamplingRate = FrameRate;
+	}
+	
+	void FDataflowSimulationGenerator::SetCacheCollection(const TObjectPtr<UChaosCacheCollection>& ChaosCache) 
+	{
+		CacheCollection = ChaosCache;
+	}
+	
+	void FDataflowSimulationGenerator::RequestGeneratorAction(EDataflowGeneratorActions ActionType)
+	{
+		if (PendingAction != EDataflowGeneratorActions::NoAction)
+		{
+			return;
+		}
+		PendingAction = ActionType;
+	}
+	
+	void FDataflowSimulationGenerator::FreeTaskResource(bool bCancelled)
+	{
+		TaskManager->AsyncNotification->SetProgressText(LOCTEXT("Finishing", "Finishing, please wait"));
+		TaskManager->FreeSimulationResource();
+		const FDateTime CurrentTime = FDateTime::UtcNow();
+		UE_LOG(LogDataflowSimulationGenerator, Log, TEXT("Simulation finished in %f seconds"), (CurrentTime - TaskManager->StartTime).GetTotalSeconds());
+	
+		{
+			// TODO : save the result to the chaos cache collection
+		}
+		if (bCancelled)
+		{
+			TaskManager->AsyncNotification->SetProgressText(LOCTEXT("Cancelled", "Cancelled"));
+			TaskManager->AsyncNotification->SetComplete(false);
+		}
+		else
+		{
+			TaskManager->AsyncNotification->SetProgressText(LOCTEXT("Finished", "Finished"));
+			TaskManager->AsyncNotification->SetComplete(true);
+		}
+		TaskManager.Reset();
+	}
+};
+
+#undef LOCTEXT_NAMESPACE
diff --git a/Engine/Plugins/Experimental/Dataflow/Source/DataflowEditor/Private/Dataflow/DataflowSimulationPanel.cpp b/Engine/Plugins/Experimental/Dataflow/Source/DataflowEditor/Private/Dataflow/DataflowSimulationPanel.cpp
--- a/Engine/Plugins/Experimental/Dataflow/Source/DataflowEditor/Private/Dataflow/DataflowSimulationPanel.cpp
+++ b/Engine/Plugins/Experimental/Dataflow/Source/DataflowEditor/Private/Dataflow/DataflowSimulationPanel.cpp
@@ -7,6 +7,7 @@
 #include "Animation/AnimSingleNodeInstance.h"
 #include "Animation/AnimSequence.h"
 #include "Animation/BlendSpace.h"
+#include "Chaos/CacheManagerActor.h"
 #include "Dataflow/DataflowContent.h"
 #include "Widgets/Input/SButton.h"
 #include "Widgets/Images/SImage.h"
@@ -47,6 +48,7 @@ void SDataflowSimulationPanel::Construct( const SDataflowSimulationPanel::FArgum
 			.DisplayDrag(this, &SDataflowSimulationPanel::GetDisplayDrag)
 			.OnValueChanged(this, &SDataflowSimulationPanel::OnValueChanged)
 			.OnBeginSliderMovement(this, &SDataflowSimulationPanel::OnBeginSliderMovement)
+			.OnClickedRecord(this, &SDataflowSimulationPanel::OnClick_Record)
 			.OnClickedForwardPlay(this, &SDataflowSimulationPanel::OnClick_Forward)
 			.OnClickedForwardStep(this, &SDataflowSimulationPanel::OnClick_Forward_Step)
 			.OnClickedForwardEnd(this, &SDataflowSimulationPanel::OnClick_Forward_End)
@@ -63,11 +65,8 @@ void SDataflowSimulationPanel::Construct( const SDataflowSimulationPanel::FArgum
 			.TransportControlWidgetsToCreate(TransportControlWidgets)
 		]
 	];
-	UpdatePreviewAnimationInstance();
-	ApplyPlaybackSettings();
 }
 
-
 TSharedRef<SWidget> SDataflowSimulationPanel::OnCreatePreviewPlaybackModeWidget()
 {
 	PreviewPlaybackModeButton = SNew(SButton)
@@ -124,10 +123,10 @@ TSharedRef<SWidget> SDataflowSimulationPanel::OnCreatePreviewPlaybackModeWidget(
 
 FReply SDataflowSimulationPanel::OnClick_Forward_Step()
 {
-	if (UAnimSingleNodeInstance* const PreviewInstance = GetPreviewAnimationInstance())
+	if (const TSharedPtr<FDataflowSimulationScene> PreviewScene = SimulationScene.Pin())
 	{
-		PreviewInstance->SetPlaying(false);
-		PreviewInstance->StepForward();
+		const float StepSize = GetSequenceLength() / GetNumberOfKeys();
+		PreviewScene->SimulationTime += StepSize;
 	}
 
 	return FReply::Handled();
@@ -135,10 +134,9 @@ FReply SDataflowSimulationPanel::OnClick_Forward_Step()
 
 FReply SDataflowSimulationPanel::OnClick_Forward_End()
 {
-	if (UAnimSingleNodeInstance* const PreviewInstance = GetPreviewAnimationInstance())
+	if (const TSharedPtr<FDataflowSimulationScene> PreviewScene = SimulationScene.Pin())
 	{
-		PreviewInstance->SetPlaying(false);
-		PreviewInstance->SetPosition(PreviewInstance->GetLength(), false);
+		PreviewScene->SimulationTime = GetSequenceLength();
 	}
 
 	return FReply::Handled();
@@ -146,10 +144,10 @@ FReply SDataflowSimulationPanel::OnClick_Forward_End()
 
 FReply SDataflowSimulationPanel::OnClick_Backward_Step()
 {
-	if (UAnimSingleNodeInstance* const PreviewInstance = GetPreviewAnimationInstance())
+	if (const TSharedPtr<FDataflowSimulationScene> PreviewScene = SimulationScene.Pin())
 	{
-		PreviewInstance->SetPlaying(false);
-		PreviewInstance->StepBackward();
+		const float StepSize = GetSequenceLength() / GetNumberOfKeys();
+		PreviewScene->SimulationTime -= StepSize;
 	}
 
 	return FReply::Handled();
@@ -157,44 +155,32 @@ FReply SDataflowSimulationPanel::OnClick_Backward_Step()
 
 FReply SDataflowSimulationPanel::OnClick_Backward_End()
 {
-	if (UAnimSingleNodeInstance* const PreviewInstance = GetPreviewAnimationInstance())
+	if (const TSharedPtr<FDataflowSimulationScene> PreviewScene = SimulationScene.Pin())
 	{
-		PreviewInstance->SetPlaying(false);
-		PreviewInstance->SetPosition(0.f, false);
+		PreviewScene->SimulationTime = 0.0f;
 	}
 
 	return FReply::Handled();
 }
 
-FReply SDataflowSimulationPanel::OnClick_Forward()
+FReply SDataflowSimulationPanel::OnClick_Record()
 {
-	if (UAnimSingleNodeInstance* const PreviewInstance = GetPreviewAnimationInstance())
+	if (const TSharedPtr<FDataflowSimulationScene> PreviewScene = SimulationScene.Pin())
 	{
-		const bool bIsReverse = PreviewInstance->IsReverse();
-		const bool bIsPlaying = PreviewInstance->IsPlaying();
-
-		// if current bIsReverse and bIsPlaying, we'd like to just turn off reverse
-		if (bIsReverse && bIsPlaying)
-		{
-			PreviewInstance->SetReverse(false);
-		}
-		// already playing, simply pause
-		else if (bIsPlaying) 
-		{
-			PreviewInstance->SetPlaying(false);
-		}
-		// if not playing, play forward
-		else 
-		{
-			//if we're at the end of the animation, jump back to the beginning before playing
-			if ( GetScrubValue() >= GetSequenceLength() )
-			{
-				PreviewInstance->SetPosition(0.0f, false);
-			}
+		PreviewScene->UpdateSimulationCache();
+	}
+	return FReply::Handled();
+}
 
-			PreviewInstance->SetReverse(false);
-			PreviewInstance->SetPlaying(true);
-		}
+FReply SDataflowSimulationPanel::OnClick_Forward()
+{
+	if(PlaybackMode == EPlaybackMode::PlayingForward)
+	{
+		PlaybackMode = EPlaybackMode::Stopped;
+	}
+	else
+	{
+		PlaybackMode = EPlaybackMode::PlayingForward;
 	}
 
 	return FReply::Handled();
@@ -202,31 +188,13 @@ FReply SDataflowSimulationPanel::OnClick_Forward()
 
 FReply SDataflowSimulationPanel::OnClick_Backward()
 {
-	if (UAnimSingleNodeInstance* const PreviewInstance = GetPreviewAnimationInstance())
+	if(PlaybackMode == EPlaybackMode::PlayingReverse)
 	{
-		const bool bIsReverse = PreviewInstance->IsReverse();
-		const bool bIsPlaying = PreviewInstance->IsPlaying();
-
-		// if currently playing forward, just simply turn on reverse
-		if (!bIsReverse && bIsPlaying)
-		{
-			PreviewInstance->SetReverse(true);
-		}
-		else if (bIsPlaying)
-		{
-			PreviewInstance->SetPlaying(false);
-		}
-		else
-		{
-			// if we're at the beginning of the animation, jump back to the end before playing
-			if ( GetScrubValue() <= 0.0f )
-			{
-				PreviewInstance->SetPosition(GetSequenceLength(), false);
-			}
-
-			PreviewInstance->SetPlaying(true);
-			PreviewInstance->SetReverse(true);
-		}
+		PlaybackMode = EPlaybackMode::Stopped;
+	}
+	else
+	{
+		PlaybackMode = EPlaybackMode::PlayingReverse;
 	}
 
 	return FReply::Handled();
@@ -237,22 +205,6 @@ FReply SDataflowSimulationPanel::OnClick_PreviewPlaybackMode()
 	if (PreviewPlaybackMode == EDataflowPlaybackMode::Default)
 	{
 		PreviewPlaybackMode = EDataflowPlaybackMode::Looping;
-		
-		if (UAnimSingleNodeInstance* const PreviewInstance = GetPreviewAnimationInstance())
-		{
-			// If we paused due to hitting the end point, start playing again when entering loop mode
-			const float CurrentTime = PreviewInstance->GetCurrentTime();
-			const float PlayRate = PreviewInstance->GetPlayRate();
-			const float AssetPlayLength = PreviewInstance->CurrentAsset->GetPlayLength();
-			if (PlayRate < 0.0 && CurrentTime <= 0.0)
-			{
-				PreviewInstance->SetPlaying(true);
-			}
-			else if (PlayRate > 0.0 && CurrentTime >= AssetPlayLength)
-			{
-				PreviewInstance->SetPlaying(true);
-			}
-		}
 	}
 	else if (PreviewPlaybackMode == EDataflowPlaybackMode::Looping)
 	{
@@ -261,120 +213,64 @@ FReply SDataflowSimulationPanel::OnClick_PreviewPlaybackMode()
 	else
 	{
 		PreviewPlaybackMode = EDataflowPlaybackMode::Default;
-
-		if (UAnimSingleNodeInstance* const PreviewInstance = GetPreviewAnimationInstance())
-		{
-			// If we're switching to linear playback, set it to forward mode
-			PreviewInstance->SetReverse(false);
-		}
 	}
 
-	ApplyPlaybackSettings();
-
 	return FReply::Handled();
 }
 
-void SDataflowSimulationPanel::ApplyPlaybackSettings()
+void SDataflowSimulationPanel::OnTickPlayback(double InCurrentTime, float InDeltaTime)
 {
-	UAnimSingleNodeInstance* const PreviewInstance = GetPreviewAnimationInstance();
-
-	switch (PreviewPlaybackMode)
+	if (const TSharedPtr<FDataflowSimulationScene> PreviewScene = SimulationScene.Pin())
 	{
-	case EDataflowPlaybackMode::Default:
-		if (PreviewInstance)
-		{
-			PreviewInstance->SetLooping(false);
-		}
-		break;
-	case EDataflowPlaybackMode::Looping:
-		if (PreviewInstance)
-		{
-			PreviewInstance->SetLooping(true);
-		}
-		break;
-	case EDataflowPlaybackMode::PingPong:
-		if (PreviewInstance)
+		const float SequenceLength = GetSequenceLength();
+		
+		const float SimulationTime = (PlaybackMode == EPlaybackMode::PlayingForward) ?
+			PreviewScene->SimulationTime + InDeltaTime : PreviewScene->SimulationTime - InDeltaTime;
+		
+		if(PreviewPlaybackMode == EDataflowPlaybackMode::Looping)
 		{
-			PreviewInstance->SetLooping(false);
+			PreviewScene->SimulationTime = SimulationTime - SequenceLength * FMath::Floor(SimulationTime / SequenceLength);
 		}
-		break;
-	}
-}
-
-void SDataflowSimulationPanel::OnTickPlayback(double InCurrentTime, float InDeltaTime)
-{
-	if (UAnimSingleNodeInstance* const PreviewInstance = GetPreviewAnimationInstance())
-	{
-		const float CurrentTime = PreviewInstance->GetCurrentTime();
-		const float PlayRate = PreviewInstance->GetPlayRate();
-		const float AssetPlayLength = PreviewInstance->CurrentAsset->GetPlayLength();
-
-		if (PreviewPlaybackMode == EDataflowPlaybackMode::PingPong)
+		else
 		{
-			if (PlayRate < 0.0 && CurrentTime <= 0.0)
+			if(PreviewPlaybackMode == EDataflowPlaybackMode::PingPong)
 			{
-				PreviewInstance->SetReverse(!PreviewInstance->IsReverse());
-			}
-			else if (PlayRate > 0.0 && CurrentTime >= AssetPlayLength)
-			{
-				PreviewInstance->SetReverse(!PreviewInstance->IsReverse());
+				if((PlaybackMode == EPlaybackMode::PlayingForward) && (SimulationTime >= SequenceLength))
+				{
+					PlaybackMode = EPlaybackMode::PlayingReverse;
+				}
+
+				if((PlaybackMode == EPlaybackMode::PlayingReverse) && (SimulationTime <= 0.0f))
+				{
+					PlaybackMode = EPlaybackMode::PlayingForward;
+				}
 			}
+			PreviewScene->SimulationTime = FMath::Clamp(SimulationTime, 0.0f,  SequenceLength);
 		}
 	}
 }
 
 EPlaybackMode::Type SDataflowSimulationPanel::GetPlaybackMode() const
 {
-	if (const UAnimSingleNodeInstance* const PreviewInstance = GetPreviewAnimationInstance())
-	{
-		if (PreviewInstance->IsPlaying())
-		{
-			return PreviewInstance->IsReverse() ? EPlaybackMode::PlayingReverse : EPlaybackMode::PlayingForward;
-		}
-		return EPlaybackMode::Stopped;
-	}
-	
-	return EPlaybackMode::Stopped;
+	return PlaybackMode;
 }
 
 void SDataflowSimulationPanel::OnValueChanged(float NewValue)
 {
-	if (UAnimSingleNodeInstance* const PreviewInstance = GetPreviewAnimationInstance())
+	if (const TSharedPtr<FDataflowSimulationScene> PreviewScene = SimulationScene.Pin())
 	{
-		PreviewInstance->SetPosition(NewValue);
+		PreviewScene->SimulationTime = NewValue;
 	}
 }
 
 void SDataflowSimulationPanel::OnBeginSliderMovement()
-{
-	if (UAnimSingleNodeInstance* const PreviewInstance = GetPreviewAnimationInstance())
-	{
-		PreviewInstance->SetPlaying(false);
-	}
-}
+{}
 
 uint32 SDataflowSimulationPanel::GetNumberOfKeys() const
 {
-	if (const UAnimSingleNodeInstance* const PreviewInstance = GetPreviewAnimationInstance())
+	if (const TSharedPtr<const FDataflowSimulationScene> PreviewScene = SimulationScene.Pin())
 	{
-		const float Length = const_cast<UAnimSingleNodeInstance*>(PreviewInstance)->GetLength();
-
-		// if anim sequence, use correct num frames
-		int32 NumKeys = (int32)(Length / 0.0333f);
-
-		if (PreviewInstance->CurrentAsset)
-		{
-			if (PreviewInstance->CurrentAsset->IsA(UAnimSequenceBase::StaticClass()))
-			{
-				NumKeys = CastChecked<UAnimSequenceBase>(PreviewInstance->CurrentAsset)->GetNumberOfSampledKeys();
-			}
-			else if (PreviewInstance->CurrentAsset->IsA(UBlendSpace::StaticClass()))
-			{
-				// Blendspaces dont display frame notches, so just return 0 here
-				NumKeys = 0;
-			}
-		}
-		return NumKeys;
+		return PreviewScene->GetNumFrames();
 	}
 
 	return 1;
@@ -382,72 +278,30 @@ uint32 SDataflowSimulationPanel::GetNumberOfKeys() const
 
 float SDataflowSimulationPanel::GetSequenceLength() const
 {
-	if (const UAnimSingleNodeInstance* const PreviewInstance = GetPreviewAnimationInstance())	// non-const because UAnimSingleNodeInstance::GetLength() is non-const
+	if (const TSharedPtr<const FDataflowSimulationScene> PreviewScene = SimulationScene.Pin())
 	{
-		return const_cast<UAnimSingleNodeInstance*>(PreviewInstance)->GetLength();
+		return PreviewScene->GetTimeRange()[1]-PreviewScene->GetTimeRange()[0];
 	}
-
-	return 0.f;
+	return 0.0f;
 }
 
 float SDataflowSimulationPanel::GetScrubValue() const
-{
-	if (const UAnimSingleNodeInstance* const PreviewInstance = GetPreviewAnimationInstance())
-	{
-		return PreviewInstance->GetCurrentTime(); 
-	}
-
-	return 0.f;
-}
-
-bool SDataflowSimulationPanel::GetDisplayDrag() const
-{
-	const UAnimSingleNodeInstance* const PreviewInstance = GetPreviewAnimationInstance();
-	if (PreviewInstance && PreviewInstance->CurrentAsset)
-	{
-		return true;
-	}
-
-	return false;
-}
-
-UAnimSingleNodeInstance* SDataflowSimulationPanel::GetPreviewAnimationInstance()
 {
 	if (const TSharedPtr<const FDataflowSimulationScene> PreviewScene = SimulationScene.Pin())
 	{
-		if(const TObjectPtr<UDataflowSkeletalContent> DataflowContent = Cast<UDataflowSkeletalContent>(PreviewScene->GetEditorContent()))
-		{
-			// temporary nullptr while the simulation nodes are pushed 
-			return nullptr;
-		}
+		return PreviewScene->SimulationTime;
 	}
 
-	return nullptr;
+	return 0.0f;
 }
 
-const UAnimSingleNodeInstance* SDataflowSimulationPanel::GetPreviewAnimationInstance() const
+bool SDataflowSimulationPanel::GetDisplayDrag() const
 {
 	if (const TSharedPtr<const FDataflowSimulationScene> PreviewScene = SimulationScene.Pin())
 	{
-		if(const TObjectPtr<UDataflowSkeletalContent> DataflowContent = Cast<UDataflowSkeletalContent>(PreviewScene->GetEditorContent()))
-		{
-			// temporary nullptr while the simulation nodes are pushed 
-			return nullptr;
-		}
-	}
-
-	return nullptr;
-}
-
-void SDataflowSimulationPanel::UpdatePreviewAnimationInstance()
-{
-	if (const TSharedPtr<FDataflowSimulationScene> PreviewScene = SimulationScene.Pin())
-	{
-		if(TObjectPtr<UDataflowSkeletalContent> DataflowContent = Cast<UDataflowSkeletalContent>(PreviewScene->GetEditorContent()))
-		{
-			// do nothing while the simulation nodes are pushed 
-		}
+			return true;
 	}
+	return false;
 }
 
 #undef LOCTEXT_NAMESPACE
diff --git a/Engine/Plugins/Experimental/Dataflow/Source/DataflowEditor/Private/Dataflow/DataflowSimulationScene.cpp b/Engine/Plugins/Experimental/Dataflow/Source/DataflowEditor/Private/Dataflow/DataflowSimulationScene.cpp
--- a/Engine/Plugins/Experimental/Dataflow/Source/DataflowEditor/Private/Dataflow/DataflowSimulationScene.cpp
+++ b/Engine/Plugins/Experimental/Dataflow/Source/DataflowEditor/Private/Dataflow/DataflowSimulationScene.cpp
@@ -1,49 +1,198 @@
 // Copyright Epic Games, Inc. All Rights Reserved.
 
 #include "Dataflow/DataflowSimulationScene.h"
-
+#include "Chaos/CacheManagerActor.h"
+#include "Misc/TransactionObjectEvent.h"
 #include "Dataflow/DataflowEditor.h"
-#include "Selection.h"
+#include "EngineUtils.h"
 
 #define LOCTEXT_NAMESPACE "FDataflowSimulationScene"
 
 //
-// FDataflowSimulationScene
+// Simulation Scene
 //
 
 FDataflowSimulationScene::FDataflowSimulationScene(FPreviewScene::ConstructionValues ConstructionValues, UDataflowEditor* InEditor)
 	: FDataflowPreviewSceneBase(ConstructionValues, InEditor)
 {
-	if(RootSceneActor)
+	SimulationSceneDescription = NewObject<UDataflowSimulationSceneDescription>();
+	SimulationSceneDescription->SetSimulationScene(this);
+
+	SimulationGenerator = MakeShared<Dataflow::FDataflowSimulationGenerator>();
+
+	const TObjectPtr<AChaosCacheManager> CacheManager = GetWorld()->SpawnActor<AChaosCacheManager>(AChaosCacheManager::StaticClass());
+
+	CacheManager->StartMode = EStartMode::Timed;
+	CacheManager->CacheMode = ECacheMode::None;
+	
+	RootSceneActor = CacheManager;
+
+	CreateSimulationWorld();
+}
+
+FDataflowSimulationScene::~FDataflowSimulationScene()
+{
+	ResetSimulationWorld();
+}
+
+void FDataflowSimulationScene::UnbindSceneSelection()
+{
+	if(SimulationContext && SimulationContext->GetRootActor() && SimulationContext->GetSimulationWorld())
 	{
 		TInlineComponentArray<UPrimitiveComponent*> PrimComponents;
-		RootSceneActor->GetComponents(PrimComponents);
+		SimulationContext->GetRootActor()->GetComponents(PrimComponents);
 
 		for(UPrimitiveComponent* PrimComponent : PrimComponents)
 		{
-			PrimComponent->SelectionOverrideDelegate =
-				UPrimitiveComponent::FSelectionOverride::CreateRaw(this, &FDataflowPreviewSceneBase::IsComponentSelected);
+			PrimComponent->SelectionOverrideDelegate.Unbind();
 		}
 	}
 }
 
-FDataflowSimulationScene::~FDataflowSimulationScene()
+void FDataflowSimulationScene::ResetSimulationWorld()
 {
-	if(RootSceneActor)
+	if(SimulationContext && SimulationContext->GetRootActor())
 	{
-		TInlineComponentArray<UPrimitiveComponent*> PrimComponents;
-		RootSceneActor->GetComponents(PrimComponents);
+		// Unbind the scene selection
+		UnbindSceneSelection();
+		
+		// Clear the simulation context
+		CleanSimulationContext(SimulationContext);
+	}
+}
 
+void FDataflowSimulationScene::BindSceneSelection()
+{
+	if(SimulationContext && SimulationContext->GetRootActor() && SimulationContext->GetSimulationWorld())
+	{
+		TInlineComponentArray<UPrimitiveComponent*> PrimComponents;
+		SimulationContext->GetRootActor()->GetComponents(PrimComponents);
+		
 		for(UPrimitiveComponent* PrimComponent : PrimComponents)
 		{
-			PrimComponent->SelectionOverrideDelegate.Unbind();
+			PrimComponent->SelectionOverrideDelegate =
+				UPrimitiveComponent::FSelectionOverride::CreateRaw(this, &FDataflowPreviewSceneBase::IsComponentSelected);
+		}
+	}
+}
+
+void FDataflowSimulationScene::CreateSimulationWorld()
+{
+	if(const TObjectPtr<AChaosCacheManager> CacheManager = Cast<AChaosCacheManager>(RootSceneActor))
+	{
+		if(const TObjectPtr<UDataflowBaseContent>& EditorContent = GetEditorContent())
+		{
+			if(SimulationGenerator && SimulationSceneDescription)
+			{
+				SimulationGenerator->SetSimulationGraph(EditorContent->GetDataflowAsset());
+				SimulationGenerator->SetSamplingRate(SimulationSceneDescription->SamplingRate);
+				SimulationGenerator->SetCacheCollection(SimulationSceneDescription->CacheAsset);
+		
+				CacheManager->CacheCollection = SimulationSceneDescription->CacheAsset;
+			}
+
+			// Build the simulation context
+			BuildSimulationContext(EditorContent->GetDataflowOwner(), EditorContent->GetDataflowAsset(),
+				CacheManager, Dataflow::ECachingMode::PlayCache, SimulationContext);
+		
+			// Update the simulation time ranges and number of frames
+			GetCacheDuration(EditorContent->GetDataflowAsset(), SimulationContext, SimulationSceneDescription->SamplingRate,
+				TimeRange[0], TimeRange[1], NumFrames);
+
+			// update the selection binding since we are constantly editing the graph
+			BindSceneSelection();
 		}
 	}
 }
 
+void FDataflowSimulationScene::UpdateSimulationCache()
+{
+	if(SimulationGenerator.IsValid())
+	{
+		SimulationGenerator->RequestGeneratorAction(Dataflow::EDataflowGeneratorActions::StartGenerate);
+	}
+}
+
 void FDataflowSimulationScene::TickDataflowScene(const float DeltaSeconds)
 {
 	GetWorld()->Tick(ELevelTick::LEVELTICK_All, DeltaSeconds);
+
+	if(const TObjectPtr<UDataflowBaseContent>& EditorContent = GetEditorContent())
+	{
+		if( Dataflow::ShouldResetWorld(EditorContent->GetDataflowAsset(), SimulationContext,LastTimeStamp))
+		{
+			// unregister components, cache manager, selection...
+			ResetSimulationWorld();
+
+			// register components, cache manager, selection...
+			CreateSimulationWorld();
+		}
+
+		// Load the cache at some point in time
+		Cast<AChaosCacheManager>(RootSceneActor)->SetStartTime(SimulationTime);
+
+		// Update all the animation at the simulation time
+		Dataflow::UpdateAnimationNodes(EditorContent->GetDataflowAsset(), SimulationContext, SimulationTime);
+	}
+}
+
+void FDataflowSimulationScene::AddReferencedObjects(FReferenceCollector& Collector)
+{
+	FDataflowPreviewSceneBase::AddReferencedObjects(Collector);
+
+	Collector.AddReferencedObject(SimulationSceneDescription);
+}
+
+void FDataflowSimulationScene::SceneDescriptionPropertyChanged(const FName& PropertyName)
+{
+	if (PropertyName == GET_MEMBER_NAME_CHECKED(UDataflowSimulationSceneDescription, SamplingRate))
+	{
+		if(SimulationGenerator)
+		{
+			SimulationGenerator->SetSamplingRate(SimulationSceneDescription->SamplingRate);
+		}
+	}
+	else if(PropertyName == GET_MEMBER_NAME_CHECKED(UDataflowSimulationSceneDescription, CacheAsset))
+	{
+		if(SimulationGenerator)
+		{
+			SimulationGenerator->SetCacheCollection(SimulationSceneDescription->CacheAsset);
+
+			Cast<AChaosCacheManager>(RootSceneActor)->CacheCollection = SimulationSceneDescription->CacheAsset;
+		}
+	}
+}
+
+void UDataflowSimulationSceneDescription::SetSimulationScene(FDataflowSimulationScene* InSimulationScene)
+{
+	SimulationScene = InSimulationScene;
+}
+
+void UDataflowSimulationSceneDescription::PostEditChangeProperty(FPropertyChangedEvent& PropertyChangedEvent)
+{
+	Super::PostEditChangeProperty(PropertyChangedEvent);
+
+	if (SimulationScene)
+	{
+		SimulationScene->SceneDescriptionPropertyChanged(PropertyChangedEvent.GetMemberPropertyName());
+	}
+
+	DataflowSimulationSceneDescriptionChanged.Broadcast();
+}
+
+void UDataflowSimulationSceneDescription::PostTransacted(const FTransactionObjectEvent& TransactionEvent)
+{
+	Super::PostTransacted(TransactionEvent);
+
+	// On Undo/Redo, PostEditChangeProperty just gets an empty FPropertyChangedEvent. However this function gets enough info to figure out which property changed
+	if (TransactionEvent.GetEventType() == ETransactionObjectEventType::UndoRedo && TransactionEvent.HasPropertyChanges())
+	{
+		const TArray<FName>& PropertyNames = TransactionEvent.GetChangedProperties();
+		for (const FName& PropertyName : PropertyNames)
+		{
+			SimulationScene->SceneDescriptionPropertyChanged(PropertyName);
+		}
+	}
 }
 
 #undef LOCTEXT_NAMESPACE
diff --git a/Engine/Plugins/Experimental/Dataflow/Source/DataflowEditor/Private/Dataflow/DataflowSimulationViewport.cpp b/Engine/Plugins/Experimental/Dataflow/Source/DataflowEditor/Private/Dataflow/DataflowSimulationViewport.cpp
--- a/Engine/Plugins/Experimental/Dataflow/Source/DataflowEditor/Private/Dataflow/DataflowSimulationViewport.cpp
+++ b/Engine/Plugins/Experimental/Dataflow/Source/DataflowEditor/Private/Dataflow/DataflowSimulationViewport.cpp
@@ -1,23 +1,30 @@
 // Copyright Epic Games, Inc. All Rights Reserved.
 
 #include "Dataflow/DataflowSimulationViewport.h"
-
 #include "Dataflow/DataflowActor.h"
 #include "Dataflow/DataflowEditorMode.h"
 #include "Dataflow/DataflowSimulationViewportClient.h"
 #include "Dataflow/DataflowEditorToolkit.h"
 #include "EditorModeManager.h"
+#include "Animation/AnimSingleNodeInstance.h"
 #include "Dataflow/DataflowContent.h"
+#include "Dataflow/DataflowEditorCommands.h"
 #include "Dataflow/DataflowSimulationScene.h"
 #include "Dataflow/DataflowSimulationPanel.h"
+#include "Dataflow/DataflowSimulationViewportToolbar.h"
 
 #define LOCTEXT_NAMESPACE "SDataflowSimulationViewport"
 
-
 SDataflowSimulationViewport::SDataflowSimulationViewport()
 {
 }
 
+const TSharedPtr<FDataflowSimulationScene>& SDataflowSimulationViewport::GetSimulationScene() const
+{
+	const TSharedPtr<FDataflowSimulationViewportClient> DataflowClient = StaticCastSharedPtr<FDataflowSimulationViewportClient>(Client);
+	return DataflowClient->GetDataflowEditorToolkit().Pin()->GetSimulationScene();
+}
+
 void SDataflowSimulationViewport::Construct(const FArguments& InArgs, const FAssetEditorViewportConstructionArgs& InViewportConstructionArgs)
 {
 	SAssetEditorViewport::FArguments ParentArgs;
@@ -27,8 +34,7 @@ void SDataflowSimulationViewport::Construct(const FArguments& InArgs, const FAss
 
 	if(static_cast<FDataflowSimulationScene*>(Client->GetPreviewScene())->CanRunSimulation())
 	{
-		TSharedPtr<FDataflowSimulationViewportClient> DataflowClient = StaticCastSharedPtr<FDataflowSimulationViewportClient>(Client);
-		TWeakPtr<FDataflowSimulationScene> SimulationScene = DataflowClient->GetDataflowEditorToolkit().Pin()->GetSimulationScene();
+		TWeakPtr<FDataflowSimulationScene> SimulationScene = GetSimulationScene();
             
 		ViewportOverlay->AddSlot()
 		[
@@ -52,6 +58,10 @@ void SDataflowSimulationViewport::Construct(const FArguments& InArgs, const FAss
 		];
 	}
 }
+TSharedPtr<SWidget> SDataflowSimulationViewport::MakeViewportToolbar()
+{
+	return SNew(SDataflowSimulationViewportToolBar, SharedThis(this)).CommandList(CommandList);
+}
 
 void SDataflowSimulationViewport::OnFocusViewportToSelection()
 {
@@ -77,6 +87,17 @@ UDataflowEditorMode* SDataflowSimulationViewport::GetEdMode() const
 void SDataflowSimulationViewport::BindCommands()
 {
 	SAssetEditorViewport::BindCommands();
+	
+	const FDataflowEditorCommandsImpl& CommandInfos = FDataflowEditorCommands::Get();
+
+	CommandList->MapAction(
+		CommandInfos.UpdateSimulationCache,
+		FExecuteAction::CreateLambda([this]()
+		{
+			static_cast<FDataflowSimulationScene*>(Client->GetPreviewScene())->UpdateSimulationCache();
+		}),
+		FCanExecuteAction::CreateLambda([this]() { return true; }),
+		FIsActionChecked::CreateLambda([this]() { return false; }));
 }
 
 bool SDataflowSimulationViewport::IsVisible() const
@@ -100,14 +121,14 @@ void SDataflowSimulationViewport::OnFloatingButtonClicked()
 {
 }
 
-float SDataflowSimulationViewport:: GetViewMinInput() const
+float SDataflowSimulationViewport::GetViewMinInput() const
 {
-	return 0.0f; //Temp value while the simulation nodes is being pushed
+	return GetSimulationScene()->GetTimeRange()[0];
 }
 
 float SDataflowSimulationViewport::GetViewMaxInput() const
 {
-	return 5.0f; //Temp value while the simulation nodes is being pushed
+	return GetSimulationScene()->GetTimeRange()[1];
 }
 
 
diff --git a/Engine/Plugins/Experimental/Dataflow/Source/DataflowEditor/Private/Dataflow/DataflowSimulationViewportClient.cpp b/Engine/Plugins/Experimental/Dataflow/Source/DataflowEditor/Private/Dataflow/DataflowSimulationViewportClient.cpp
--- a/Engine/Plugins/Experimental/Dataflow/Source/DataflowEditor/Private/Dataflow/DataflowSimulationViewportClient.cpp
+++ b/Engine/Plugins/Experimental/Dataflow/Source/DataflowEditor/Private/Dataflow/DataflowSimulationViewportClient.cpp
@@ -17,8 +17,6 @@
 #include "SGraphPanel.h"
 #include "SNodePanel.h"
 
-
-
 FDataflowSimulationViewportClient::FDataflowSimulationViewportClient(FEditorModeTools* InModeTools,
                                                              FPreviewScene* InPreviewScene,  const bool bCouldTickScene,
                                                              const TWeakPtr<SEditorViewport> InEditorViewportWidget)
@@ -61,6 +59,34 @@ void FDataflowSimulationViewportClient::Tick(float DeltaSeconds)
 void FDataflowSimulationViewportClient::ProcessClick(FSceneView& View, HHitProxy* HitProxy, FKey Key, EInputEvent Event, uint32 HitX, uint32 HitY)
 {
 	Super::ProcessClick(View, HitProxy, Key, Event, HitX, HitY);
+
+	USelection* SelectedComponents = ModeTools->GetSelectedComponents();
+	
+	TArray<UPrimitiveComponent*> PreviouslySelectedComponents;
+	SelectedComponents->GetSelectedObjects<UPrimitiveComponent>(PreviouslySelectedComponents);
+
+	SelectedComponents->Modify();
+	SelectedComponents->BeginBatchSelectOperation();
+
+	SelectedComponents->DeselectAll();
+
+	if (HitProxy && HitProxy->IsA(HActor::StaticGetType()))
+	{
+		const HActor* ActorProxy = static_cast<HActor*>(HitProxy);
+		if (ActorProxy && ActorProxy->PrimComponent && ActorProxy->Actor)
+		{
+			UPrimitiveComponent* Component = const_cast<UPrimitiveComponent*>(ActorProxy->PrimComponent.Get());
+			SelectedComponents->Select(Component);
+			Component->PushSelectionToProxy();
+		}
+	}
+
+	SelectedComponents->EndBatchSelectOperation();
+
+	for (UPrimitiveComponent* const Component : PreviouslySelectedComponents)
+	{
+		Component->PushSelectionToProxy();
+	}
 }
 
 
diff --git a/Engine/Plugins/Experimental/Dataflow/Source/DataflowEditor/Private/Dataflow/DataflowSimulationViewportToolbar.cpp b/Engine/Plugins/Experimental/Dataflow/Source/DataflowEditor/Private/Dataflow/DataflowSimulationViewportToolbar.cpp
new file mode 100644
--- /dev/null
+++ b/Engine/Plugins/Experimental/Dataflow/Source/DataflowEditor/Private/Dataflow/DataflowSimulationViewportToolbar.cpp
@@ -0,0 +1,58 @@
+﻿// Copyright Epic Games, Inc. All Rights Reserved.
+#include "Dataflow/DataflowSimulationViewportToolbar.h"
+#include "Dataflow/DataflowSimulationViewport.h"
+#include "Dataflow/DataflowEditorCommands.h"
+#include "Styling/AppStyle.h"
+#include "Framework/MultiBox/MultiBoxBuilder.h"
+#include "SEditorViewportToolBarMenu.h"
+#include "Widgets/Layout/SBorder.h"
+#include "Widgets/SBoxPanel.h"
+
+#define LOCTEXT_NAMESPACE "DataflowSimulationViewportToolBar"
+
+void SDataflowSimulationViewportToolBar::Construct(const FArguments& InArgs, TSharedPtr<SDataflowSimulationViewport> InDataflowViewport)
+{
+	EditorViewport = InDataflowViewport;
+	CommandList = InArgs._CommandList;
+	Extenders = InArgs._Extenders;
+	
+	SCommonEditorViewportToolbarBase::Construct(SCommonEditorViewportToolbarBase::FArguments(), InDataflowViewport);
+}
+
+void SDataflowSimulationViewportToolBar::ExtendLeftAlignedToolbarSlots(TSharedPtr<SHorizontalBox> MainBoxPtr, TSharedPtr<SViewportToolBar> ParentToolBarPtr) const
+{
+	const FMargin ToolbarSlotPadding(2.0f, 2.0f);
+	MainBoxPtr->AddSlot()
+		.Padding(ToolbarSlotPadding)
+		[
+			MakeToolBar(Extenders)
+		];
+}
+
+TSharedRef<SWidget> SDataflowSimulationViewportToolBar::MakeToolBar(const TSharedPtr<FExtender> InExtenders) const
+{
+	FSlimHorizontalToolBarBuilder ToolbarBuilder(CommandList, FMultiBoxCustomization::None, InExtenders);
+	
+	const FName ToolBarStyle = "EditorViewportToolBar";
+	ToolbarBuilder.SetStyle(&FAppStyle::Get(), ToolBarStyle);
+	ToolbarBuilder.SetLabelVisibility(EVisibility::Collapsed);
+
+	ToolbarBuilder.BeginSection("Sim Controls");
+	ToolbarBuilder.BeginBlockGroup();
+	{
+		// the simulation caching should be triggered from the simulation panel as well
+		// we keep it here as well to have a template to add future simulation controls
+		ToolbarBuilder.AddToolBarButton(FDataflowEditorCommands::Get().UpdateSimulationCache,
+			NAME_None,
+			TAttribute<FText>(),
+			TAttribute<FText>(),
+			FSlateIcon(FAppStyle::Get().GetStyleSetName(), "Animation.Record"),
+			FName(*FDataflowEditorCommands::Get().UpdateSimulationCacheIdentifier));
+	}
+	ToolbarBuilder.EndBlockGroup();
+	ToolbarBuilder.EndSection();
+
+	return ToolbarBuilder.MakeWidget();
+}
+
+#undef LOCTEXT_NAMESPACE
diff --git a/Engine/Plugins/Experimental/Dataflow/Source/DataflowEditor/Public/Dataflow/DataflowEditorCommands.h b/Engine/Plugins/Experimental/Dataflow/Source/DataflowEditor/Public/Dataflow/DataflowEditorCommands.h
--- a/Engine/Plugins/Experimental/Dataflow/Source/DataflowEditor/Public/Dataflow/DataflowEditorCommands.h
+++ b/Engine/Plugins/Experimental/Dataflow/Source/DataflowEditor/Public/Dataflow/DataflowEditorCommands.h
@@ -42,26 +42,28 @@ public:
 	* @param bUnbind if true, commands are removed, otherwise added
 	*/
 	 static void UpdateToolCommandBinding(UInteractiveTool* Tool, TSharedPtr<FUICommandList> UICommandList, bool bUnbind = false);
-
-
-
-	TSharedPtr< FUICommandInfo > EvaluateNode;
-	TSharedPtr< FUICommandInfo > CreateComment;
-	TSharedPtr< FUICommandInfo > ToggleEnabledState;
+	
+	TSharedPtr<FUICommandInfo> EvaluateNode;
+	TSharedPtr<FUICommandInfo> CreateComment;
+	TSharedPtr<FUICommandInfo> ToggleEnabledState;
 	TSharedPtr<FUICommandInfo> ToggleObjectSelection;
 	TSharedPtr<FUICommandInfo> ToggleFaceSelection;
 	TSharedPtr<FUICommandInfo> ToggleVertexSelection;
-	TSharedPtr< FUICommandInfo > AddOptionPin;
-	TSharedPtr< FUICommandInfo > RemoveOptionPin;
-	TSharedPtr< FUICommandInfo > ZoomToFitGraph;
+	TSharedPtr<FUICommandInfo> AddOptionPin;
+	TSharedPtr<FUICommandInfo> RemoveOptionPin;
+	TSharedPtr<FUICommandInfo> ZoomToFitGraph;
 
 	TMap< FName, TSharedPtr<FUICommandInfo> > CreateNodesMap;
 
 	const static FString BeginWeightMapPaintToolIdentifier;
 	TSharedPtr<FUICommandInfo> BeginWeightMapPaintTool;
+	
 	const static FString AddWeightMapNodeIdentifier;
 	TSharedPtr<FUICommandInfo> AddWeightMapNode;
 
+	const static FString UpdateSimulationCacheIdentifier;
+	TSharedPtr<FUICommandInfo> UpdateSimulationCache;
+
 	// @todo(brice) Remove Example Tools
 	//const static FString BeginAttributeEditorToolIdentifier;
 	//TSharedPtr<FUICommandInfo> BeginAttributeEditorTool;
diff --git a/Engine/Plugins/Experimental/Dataflow/Source/DataflowEditor/Public/Dataflow/DataflowEditorMode.h b/Engine/Plugins/Experimental/Dataflow/Source/DataflowEditor/Public/Dataflow/DataflowEditorMode.h
--- a/Engine/Plugins/Experimental/Dataflow/Source/DataflowEditor/Public/Dataflow/DataflowEditorMode.h
+++ b/Engine/Plugins/Experimental/Dataflow/Source/DataflowEditor/Public/Dataflow/DataflowEditorMode.h
@@ -69,8 +69,10 @@ public:
 	FDataflowConstructionScene* GetDataflowConstructionScene() { return ConstructionScene; }
 	const FDataflowConstructionScene* GetDataflowConstructionScene() const { return ConstructionScene; }
 
+	/** Set the simulation scene to store the simulation components */
 	void SetDataflowSimlationScene(FDataflowSimulationScene* InSimulationScene);
-
+	FDataflowSimulationScene* GetDataflowSimulationScene() { return SimulationScene; }
+	const FDataflowSimulationScene* GetDataflowSimulationScene() const { return SimulationScene; }
 
 	/** Set the data flow graph editor to create nodes once the tools have ended*/
 	void SetDataflowGraphEditor(TSharedPtr<SDataflowGraphEditor> DataflowGraphEditor);
@@ -158,8 +160,10 @@ private:
 	void InitializeContextObject();
 	void DeleteContextObject();
 
-	/** Dataflow preview scene from the toolkit */
+	/** Dataflow construction scene from the toolkit */
 	FDataflowConstructionScene* ConstructionScene = nullptr;
+
+	/** Dataflow simulation scene from the toolkit */
 	FDataflowSimulationScene* SimulationScene = nullptr;
 
 	/** Correspondence between node types and commands to launch tools */
diff --git a/Engine/Plugins/Experimental/Dataflow/Source/DataflowEditor/Public/Dataflow/DataflowEditorToolkit.h b/Engine/Plugins/Experimental/Dataflow/Source/DataflowEditor/Public/Dataflow/DataflowEditorToolkit.h
--- a/Engine/Plugins/Experimental/Dataflow/Source/DataflowEditor/Public/Dataflow/DataflowEditorToolkit.h
+++ b/Engine/Plugins/Experimental/Dataflow/Source/DataflowEditor/Public/Dataflow/DataflowEditorToolkit.h
@@ -7,6 +7,7 @@
 #include "Dataflow/DataflowObjectInterface.h"
 #include "Misc/NotifyHook.h"
 #include "GraphEditor.h"
+#include "SAdvancedPreviewDetailsTab.h"
 #include "TickableEditorObject.h"
 #include "Dataflow/DataflowSelectionView.h"
 #include "Dataflow/DataflowCollectionSpreadSheet.h"
@@ -81,7 +82,6 @@ public:
 	/** Dataflow preview scenes accessor */
 	const TSharedPtr<FDataflowSimulationScene>& GetSimulationScene() const {return SimulationScene;}
 
-
 	// FSerializableObject interface
 	virtual void AddReferencedObjects(FReferenceCollector& Collector) override;
 	virtual FString GetReferencerName() const override
@@ -115,6 +115,7 @@ private:
 	TSharedRef<SDockTab> SpawnTab_CollectionSpreadSheet(const FSpawnTabArgs& Args);
 	TSharedRef<SDockTab> SpawnTab_AssetDetails(const FSpawnTabArgs& Args);
 	TSharedRef<SDockTab> SpawnTab_SimulationViewport(const FSpawnTabArgs& Args);
+	TSharedRef<SDockTab> SpawnTab_PreviewScene(const FSpawnTabArgs& Args);
 	
 	// FTickableEditorObject interface
 	virtual void Tick(float DeltaTime) override;
@@ -150,6 +151,7 @@ private:
 	static const FName CollectionSpreadSheetTabId_3;
 	static const FName CollectionSpreadSheetTabId_4;
 	static const FName SimulationViewportTabId;
+	static const FName PreviewSceneTabId;
 
 	// List of all the widgets shared ptr that will be built in the editor
 	TSharedPtr<SDataflowConstructionViewport> DataflowConstructionViewport;
@@ -166,6 +168,7 @@ private:
 	TSharedPtr<FDataflowCollectionSpreadSheet> DataflowCollectionSpreadSheet_2;
 	TSharedPtr<FDataflowCollectionSpreadSheet> DataflowCollectionSpreadSheet_3;
 	TSharedPtr<FDataflowCollectionSpreadSheet> DataflowCollectionSpreadSheet_4;
+	TSharedPtr<SWidget> AdvancedPreviewSettingsWidget;
 
 	// Utility factory functions to build the widgets
 	TSharedRef<SDataflowGraphEditor> CreateGraphEditorWidget(UDataflow* ObjectToEdit, TSharedPtr<IStructureDetailsView> PropertiesEditor);
diff --git a/Engine/Plugins/Experimental/Dataflow/Source/DataflowEditor/Public/Dataflow/DataflowSimulationGenerator.h b/Engine/Plugins/Experimental/Dataflow/Source/DataflowEditor/Public/Dataflow/DataflowSimulationGenerator.h
new file mode 100644
--- /dev/null
+++ b/Engine/Plugins/Experimental/Dataflow/Source/DataflowEditor/Public/Dataflow/DataflowSimulationGenerator.h
@@ -0,0 +1,207 @@
+﻿// Copyright Epic Games, Inc. All Rights Reserved.
+
+#pragma once
+
+#include "CoreTypes.h"
+#include "UObject/ObjectPtr.h"
+#include "TickableEditorObject.h"
+#include "Logging/LogMacros.h"
+#include "Misc/AsyncTaskNotification.h"
+#include "Dataflow/DataflowContent.h"
+
+DECLARE_LOG_CATEGORY_EXTERN(LogDataflowSimulationGenerator, Log, All);
+
+class UChaosCacheCollection;
+class AChaosCacheManager;
+
+namespace Dataflow
+{
+	/** Simulation Task to be run on the async thread */
+	class FDataflowSimulationTask : public FNonAbandonableTask
+	{
+		public:
+		FDataflowSimulationTask()
+		{}
+
+		/** Run the simulation */
+		void DoWork();
+
+		/** Can abandon check */
+		bool CanAbandon()
+		{
+			return true;
+		}
+
+		FORCEINLINE TStatId GetStatId() const
+		{
+			RETURN_QUICK_DECLARE_CYCLE_STAT(TTaskRunner, STATGROUP_ThreadPoolAsyncTasks);
+		}
+	
+		/** Task Manager */
+		TSharedPtr<struct FDataflowTaskManager> TaskManager = nullptr;
+		
+		/** Simulation delta time */
+		float DeltaTime = 0.0f;
+
+		/** Simulation min time */
+		float MinTime = TNumericLimits<float>::Max();
+
+		/** Simulation max time */
+		float MaxTime = TNumericLimits<float>::Lowest();
+
+		/** Simulation graph to be executed */
+		TObjectPtr<UDataflow> SimulationGraph;
+
+		/** Simulation  context used during graph evaluation */
+		TSharedPtr<Dataflow::FSimulationContext> SimulationContext;
+	};
+
+	/** Async simulation resource that will be used while simulating */
+	struct FDataflowSimulationResource
+	{
+		/** Number of simulated frames */
+		std::atomic<int32>* NumSimulatedFrames = nullptr;
+
+		/** Async cancel boolean */
+		std::atomic<bool>* bCancelled = nullptr;
+
+		/** Check if the simulation is canceled or not */
+		bool IsCancelled() const
+		{
+			return !bCancelled || bCancelled->load();
+		}
+		
+		/** Finish simulating the current frame */
+		void FinishFrame()
+		{
+			if (NumSimulatedFrames)
+			{
+				++(*NumSimulatedFrames);
+			}
+		}
+	};
+
+	/** Simulation Task manager */
+	struct FDataflowTaskManager
+	{
+		/** Simulation resource */
+		TSharedPtr<FDataflowSimulationResource> SimulationResource;
+
+		/** Simulation task */
+		TUniquePtr<FAsyncTask<FDataflowSimulationTask>> SimulationTask;
+
+		/** Async notification*/
+		TUniquePtr<FAsyncTaskNotification> AsyncNotification;
+
+		/** Number of frames to simulate */
+		int32 NumFrames;
+
+		/** Start time of the simulation */
+		FDateTime StartTime;
+
+		/** Last updated time */
+		FDateTime LastUpdateTime;
+
+		/** Allocate the simulation resource from the properties */
+        bool AllocateSimulationResource(const TObjectPtr<UDataflow> SimulationGraph, const int32 SamplingRate, const TObjectPtr<UChaosCacheCollection>& CacheCollection);
+
+        /** Free the simulation resource */
+        void FreeSimulationResource();
+
+        /** Cancel the simulation generation */
+        void CancelSimulationGeneration();
+
+		/** Number of simulated frames */
+		std::atomic<int32> NumSimulatedFrames = 0;
+		
+		/** Boolean to check if the simulation has been cancelled */
+		std::atomic<bool> bCancelled = false;
+		
+		/** Temporary world created to run the simulation */
+		UWorld* SimulationWorld = nullptr;
+	};
+
+	/** Enum for all the generator actions */
+	enum class EDataflowGeneratorActions
+	{
+		NoAction,
+		StartGenerate,
+		TickGenerate
+	};
+
+	/** Dataflow simulation generator */
+	class  FDataflowSimulationGenerator : public FTickableEditorObject
+	{
+	public:
+		FDataflowSimulationGenerator();
+		virtual ~FDataflowSimulationGenerator();
+
+		//~ Begin FTickableEditorObject Interface
+		virtual void Tick(float DeltaTime) override;
+		virtual TStatId GetStatId() const override;
+		virtual ETickableTickType GetTickableTickType() const override { return ETickableTickType::Always; }
+		//~ End FTickableEditorObject Interface
+
+		/** Set the simulation graph to be executed on the async thread */
+		void SetSimulationGraph(const TObjectPtr<UDataflow>& DataflowGraph);
+
+		/** Set the sampling rate for cache recording */
+		void SetSamplingRate(const int32 FrameRate);
+
+		/** Set the cache collection for cache recording */
+		void SetCacheCollection(const TObjectPtr<UChaosCacheCollection>&);
+
+		/** Enqueue a generator action to be processed on the async thread */
+		void RequestGeneratorAction(EDataflowGeneratorActions Action);
+		
+	private:
+
+		/** Start generating the simulation */
+		void StartGenerateSimulation();
+		
+		/** Tick the generated simulation */
+		void TickGenerateSimulation();
+
+		/** Free the task resource used while generating the simulation */
+		void FreeTaskResource(bool bCancelled);
+
+		/** Graph to be used to generate the simulation output */
+		TObjectPtr<UDataflow> SimulationGraph = nullptr;
+
+		/** Cache collection to store the caches */
+		TObjectPtr<UChaosCacheCollection> CacheCollection = nullptr;
+
+		/** Sampling rate used to record the cache */
+		int32 SamplingRate = 30;
+
+		/** Pending action to be send to the async thread */
+		EDataflowGeneratorActions PendingAction = EDataflowGeneratorActions::NoAction;
+
+		/** Task manager to run the async tasks */
+		TSharedPtr<FDataflowTaskManager> TaskManager = nullptr;
+	};
+
+	/** Evaluate the simulation graph given a simulation context and a simulation flag */
+	void EvaluateSimulationGraph(const TObjectPtr<UDataflow>& SimulationGraph,
+		const TSharedPtr<Dataflow::FSimulationContext>& SimulationContext, const float DeltaTime,
+		const float SimulationTime, const ESimulationFlags SimulationFlags);
+
+	/** Extract the time range, number of frames and duration of the cache sequence */
+	float GetCacheDuration(const TObjectPtr<UDataflow>& SimulationGraph, const TSharedPtr<Dataflow::FSimulationContext>& SimulationContext,
+		const int32 SamplingRate, float& MinTime, float& MaxTime, int32& NumFrames);
+
+	/** Check if the simulation cache nodes have changed to trigger a reset */
+	bool ShouldResetWorld(const TObjectPtr<UDataflow>& SimulationGraph, const TSharedPtr<Dataflow::FSimulationContext>& SimulationContext, Dataflow::FTimestamp& LastTimeStamp);
+
+	/** Update the animation nodes given a context and an animation time */
+	void UpdateAnimationNodes(const TObjectPtr<UDataflow>& SimulationGraph, const TSharedPtr<Dataflow::FSimulationContext>& SimulationContext, const float AnimationTime);
+
+	/** clean the context (unregister componentds, clesn cache, reset context)*/
+	void CleanSimulationContext(TSharedPtr<Dataflow::FSimulationContext>& SimulationContext);
+
+	/** Build simulaiton context */
+	void BuildSimulationContext(const TObjectPtr<UObject>& SimulationOwner, const TObjectPtr<UDataflow>& SimulationGraph, const TObjectPtr<AChaosCacheManager>& CacheManager,
+		const ECachingMode CachingMode, TSharedPtr<Dataflow::FSimulationContext>& SimulationContext);
+};
+
+
diff --git a/Engine/Plugins/Experimental/Dataflow/Source/DataflowEditor/Public/Dataflow/DataflowSimulationPanel.h b/Engine/Plugins/Experimental/Dataflow/Source/DataflowEditor/Public/Dataflow/DataflowSimulationPanel.h
--- a/Engine/Plugins/Experimental/Dataflow/Source/DataflowEditor/Public/Dataflow/DataflowSimulationPanel.h
+++ b/Engine/Plugins/Experimental/Dataflow/Source/DataflowEditor/Public/Dataflow/DataflowSimulationPanel.h
@@ -43,9 +43,7 @@ private:
 	FReply OnClick_Forward();
 	FReply OnClick_Backward();
 	FReply OnClick_PreviewPlaybackMode();
-
-	/** Change the playback settings */
-	void ApplyPlaybackSettings();
+	FReply OnClick_Record();
 
 	/** Main simulation delegates */
 	void OnTickPlayback(double InCurrentTime, float InDeltaTime);
@@ -58,13 +56,6 @@ private:
 	/** Get ethe current scrub value */
 	float GetScrubValue() const;
 
-	/** Animation instance accessors */
-	UAnimSingleNodeInstance* GetPreviewAnimationInstance();
-	const UAnimSingleNodeInstance* GetPreviewAnimationInstance() const;
-
-	/** Update the widget based on the modified animation instance*/
-	void UpdatePreviewAnimationInstance();
-
 	/** Get the number of keys */
 	uint32 GetNumberOfKeys() const;
 
@@ -85,4 +76,7 @@ private:
 
 	/** Preview playback mode (looping...)*/
 	EDataflowPlaybackMode PreviewPlaybackMode = EDataflowPlaybackMode::Looping;
+
+	/** playback mode */
+	EPlaybackMode::Type PlaybackMode = EPlaybackMode::Stopped;
 };
diff --git a/Engine/Plugins/Experimental/Dataflow/Source/DataflowEditor/Public/Dataflow/DataflowSimulationScene.h b/Engine/Plugins/Experimental/Dataflow/Source/DataflowEditor/Public/Dataflow/DataflowSimulationScene.h
--- a/Engine/Plugins/Experimental/Dataflow/Source/DataflowEditor/Public/Dataflow/DataflowSimulationScene.h
+++ b/Engine/Plugins/Experimental/Dataflow/Source/DataflowEditor/Public/Dataflow/DataflowSimulationScene.h
@@ -3,9 +3,49 @@
 #pragma once
 
 #include "Dataflow/DataflowEditorPreviewSceneBase.h"
+#include "Dataflow/DataflowSimulationGenerator.h"
+#include "Chaos/CacheCollection.h"
+
+#include "DataflowSimulationScene.generated.h"
 
 class UDataflowEditor;
 
+DECLARE_EVENT(UDataflowSimulationSceneDescription, FDataflowSimulationSceneDescriptionChanged)
+
+UCLASS()
+class DATAFLOWEDITOR_API UDataflowSimulationSceneDescription : public UObject
+{
+public:
+	GENERATED_BODY()
+
+	FDataflowSimulationSceneDescriptionChanged DataflowSimulationSceneDescriptionChanged;
+
+	UDataflowSimulationSceneDescription()
+	{
+		SetFlags(RF_Transactional);
+	}
+
+	/** Set the simulation scene */
+	void SetSimulationScene(class FDataflowSimulationScene* SimulationScene);
+
+	/** Caching asset to be used to record the simulation  */
+	UPROPERTY(EditAnywhere, Category="Caching")
+	TObjectPtr<UChaosCacheCollection> CacheAsset = nullptr;
+
+	/** Number of sampling frames per second for caching*/
+	UPROPERTY(EditAnywhere, Category="Caching")
+	int32 SamplingRate = 30;
+
+private:
+
+	//~ UObject Interface
+	virtual void PostEditChangeProperty(struct FPropertyChangedEvent& PropertyChangedEvent) override;
+	virtual void PostTransacted(const FTransactionObjectEvent& TransactionEvent) override;
+
+	/** Simulation scene linked to that descriptor */
+	class FDataflowSimulationScene* SimulationScene;
+};
+
 
 /**
  * Dataflow simulation scene holding all the dataflow content components
@@ -15,6 +55,10 @@ class DATAFLOWEDITOR_API FDataflowSimulationScene : public FDataflowPreviewScene
 public:
 
 	FDataflowSimulationScene(FPreviewScene::ConstructionValues ConstructionValues, UDataflowEditor* Editor);
+	
+	// FGCObject interface
+	virtual void AddReferencedObjects(FReferenceCollector& Collector) override;
+	
 	virtual ~FDataflowSimulationScene();
 
 	/** Tick data flow scene */
@@ -23,9 +67,63 @@ public:
 	/** Check if the preview scene can run simulation */
 	virtual bool CanRunSimulation() const { return true; }
 
+	/** Get the scene description used in the preview scene widget */
+	UDataflowSimulationSceneDescription* GetPreviewSceneDescription() const { return SimulationSceneDescription; }
+
+	/** Create all the simulation world components and instances */
+	void CreateSimulationWorld();
+
+	/** Reset all the simulation world components and instances */
+	void ResetSimulationWorld();
+
+	/** Reset the simulation scene */
 	void ResetSimulationScene() {}
+
+	/** Check if there is something to render */
 	bool HasRenderableGeometry() { return true; }
 
+	/** Update Scene in response to the SceneDescription changing */
+	void SceneDescriptionPropertyChanged(const FName& PropertyName);
+
+	/** Update the simulation cache */
+	void UpdateSimulationCache();
+
+	/** Get the simulation time range */
+	const FVector2f& GetTimeRange() const {return TimeRange;}
+
+	/** Get the number of frames */
+	const int32& GetNumFrames() const {return NumFrames;}
+
+	/** Simulation time used to drive the cache loading */
+	float SimulationTime;
+	
+private:
+
+	/** Bind the scene selection to the components */
+	void BindSceneSelection();
+
+	/** Unbind the scene selection from the components */
+	void UnbindSceneSelection();
+	
+	/** Simulation scene description */
+	TObjectPtr<UDataflowSimulationSceneDescription> SimulationSceneDescription;
+
+	/** Simulation generator to record the simulation result */
+	TSharedPtr<Dataflow::FDataflowSimulationGenerator> SimulationGenerator;
+
+	/** Simulation context used for the simulation graph */
+	TSharedPtr<Dataflow::FSimulationContext> SimulationContext;
+
+	/** Cache time range in seconds */
+	FVector2f TimeRange;
+
+	/** Number of cache frames */
+	int32 NumFrames;
+
+	/** Last context time stamp for which we regenerated the world */
+	Dataflow::FTimestamp LastTimeStamp = Dataflow::FTimestamp::Invalid;
 };
 
 
+
+
diff --git a/Engine/Plugins/Experimental/Dataflow/Source/DataflowEditor/Public/Dataflow/DataflowSimulationViewport.h b/Engine/Plugins/Experimental/Dataflow/Source/DataflowEditor/Public/Dataflow/DataflowSimulationViewport.h
--- a/Engine/Plugins/Experimental/Dataflow/Source/DataflowEditor/Public/Dataflow/DataflowSimulationViewport.h
+++ b/Engine/Plugins/Experimental/Dataflow/Source/DataflowEditor/Public/Dataflow/DataflowSimulationViewport.h
@@ -16,8 +16,9 @@ class SDataflowSimulationViewport : public SBaseCharacterFXEditorViewport, publi
 {
 public:
 	SLATE_BEGIN_ARGS(SDataflowSimulationViewport) {}
-	SLATE_ARGUMENT(TSharedPtr<FEditorViewportClient>, ViewportClient)
-		SLATE_END_ARGS()
+		SLATE_ARGUMENT(TSharedPtr<FEditorViewportClient>, ViewportClient)
+		SLATE_ARGUMENT(TSharedPtr<FUICommandList>, CommandList)
+	SLATE_END_ARGS()
 
 	SDataflowSimulationViewport();
 
@@ -26,6 +27,7 @@ public:
 	// SEditorViewport
 	virtual void BindCommands() override;
 	virtual bool IsVisible() const override;
+	virtual TSharedPtr<SWidget> MakeViewportToolbar() override;
 	virtual void OnFocusViewportToSelection() override;
 
 	// ICommonEditorViewportToolbarInfoProvider
@@ -33,6 +35,9 @@ public:
 	virtual TSharedPtr<FExtender> GetExtenders() const override;
 	virtual void OnFloatingButtonClicked() override;
 
+	/** Get the simulation scene */
+	const TSharedPtr<class FDataflowSimulationScene>& GetSimulationScene() const;
+
 private:
 	float GetViewMinInput() const;
 	float GetViewMaxInput() const;
diff --git a/Engine/Plugins/Experimental/Dataflow/Source/DataflowEditor/Public/Dataflow/DataflowSimulationViewportToolbar.h b/Engine/Plugins/Experimental/Dataflow/Source/DataflowEditor/Public/Dataflow/DataflowSimulationViewportToolbar.h
new file mode 100644
--- /dev/null
+++ b/Engine/Plugins/Experimental/Dataflow/Source/DataflowEditor/Public/Dataflow/DataflowSimulationViewportToolbar.h
@@ -0,0 +1,33 @@
+﻿// Copyright Epic Games, Inc. All Rights Reserved.
+#pragma once
+
+#include "CoreMinimal.h"
+#include "SCommonEditorViewportToolbarBase.h"
+
+// ----------------------------------------------------------------------------------
+class SDataflowSimulationViewport;
+class FExtender;
+class FUICommandList;
+
+/** Base toolbar for the dataflow. Should be extended to add more features */
+class SDataflowSimulationViewportToolBar : public SCommonEditorViewportToolbarBase
+{
+public:
+	SLATE_BEGIN_ARGS(SDataflowSimulationViewportToolBar) {}
+		SLATE_ARGUMENT(TSharedPtr<FUICommandList>, CommandList)
+		SLATE_ARGUMENT(TSharedPtr<FExtender>, Extenders)
+	SLATE_END_ARGS()
+
+	/** Constructs this widget with the given parameters */
+	void Construct(const FArguments& InArgs, TSharedPtr<SDataflowSimulationViewport> InDataflowViewport);
+	virtual void ExtendLeftAlignedToolbarSlots(TSharedPtr<SHorizontalBox> MainBoxPtr, TSharedPtr<SViewportToolBar> ParentToolBarPtr) const override;
+	
+private:
+	/** Reference to the parent viewport */
+	TWeakPtr<SDataflowSimulationViewport> EditorViewport;
+
+	TSharedRef<SWidget> MakeToolBar(const TSharedPtr<FExtender> InExtenders) const;
+
+	TSharedPtr<FUICommandList> CommandList;
+	TSharedPtr<FExtender> Extenders;
+};
\ No newline at end of file
diff --git a/Engine/Plugins/Experimental/Dataflow/Source/DataflowNodes/Private/Dataflow/DataflowNodesPlugin.cpp b/Engine/Plugins/Experimental/Dataflow/Source/DataflowNodes/Private/Dataflow/DataflowNodesPlugin.cpp
--- a/Engine/Plugins/Experimental/Dataflow/Source/DataflowNodes/Private/Dataflow/DataflowNodesPlugin.cpp
+++ b/Engine/Plugins/Experimental/Dataflow/Source/DataflowNodes/Private/Dataflow/DataflowNodesPlugin.cpp
@@ -10,6 +10,7 @@
 #include "Dataflow/DataflowSelectionNodes.h"
 #include "Dataflow/DataflowContextOverridesNodes.h"
 #include "Dataflow/DataflowCollectionAddScalarVertexPropertyNode.h"
+#include "Dataflow/DataflowSolverNodes.h"
 
 #define LOCTEXT_NAMESPACE "DataflowNodes"
 
@@ -22,6 +23,12 @@ void IDataflowNodesPlugin::StartupModule()
 	Dataflow::RegisterContextOverridesNodes();
 	Dataflow::DataflowCollectionAttributeKeyNodes();
 	DATAFLOW_NODE_REGISTER_CREATION_FACTORY(FDataflowCollectionAddScalarVertexPropertyNode);
+
+	Dataflow::RegisterNodeFilter(FDataflowTerminalNode::StaticType());
+	Dataflow::RegisterNodeFilter(FDataflowSimulationNode::StaticType());
+	Dataflow::RegisterNodeFilter(FDataflowInvalidNode::StaticType());
+	Dataflow::RegisterNodeFilter(FDataflowCacheNode::StaticType());
+	Dataflow::RegisterNodeFilter(FDataflowAnimationNode::StaticType());
 }
 
 void IDataflowNodesPlugin::ShutdownModule()
diff --git a/Engine/Plugins/Experimental/Dataflow/Source/DataflowNodes/Private/Dataflow/DataflowSolverNodes.cpp b/Engine/Plugins/Experimental/Dataflow/Source/DataflowNodes/Private/Dataflow/DataflowSolverNodes.cpp
new file mode 100644
--- /dev/null
+++ b/Engine/Plugins/Experimental/Dataflow/Source/DataflowNodes/Private/Dataflow/DataflowSolverNodes.cpp
@@ -0,0 +1,3 @@
+﻿#include "Dataflow/DataflowSolverNodes.h"
+
+#include UE_INLINE_GENERATED_CPP_BY_NAME(DataflowSolverNodes)
\ No newline at end of file
diff --git a/Engine/Plugins/Experimental/Dataflow/Source/DataflowNodes/Public/Dataflow/DataflowSolverNodes.h b/Engine/Plugins/Experimental/Dataflow/Source/DataflowNodes/Public/Dataflow/DataflowSolverNodes.h
new file mode 100644
--- /dev/null
+++ b/Engine/Plugins/Experimental/Dataflow/Source/DataflowNodes/Public/Dataflow/DataflowSolverNodes.h
@@ -0,0 +1,159 @@
+﻿// Copyright Epic Games, Inc. All Rights Reserved.
+
+#pragma once
+
+#include "CoreMinimal.h"
+#include "Dataflow/DataflowNode.h"
+#include "Dataflow/DataflowObjectInterface.h"
+
+#include "DataflowSolverNodes.generated.h"
+
+/**
+* FDataflowSimulationNode
+*		Base class for simulation nodes within the Dataflow graph. 
+* 
+*		Simulation nodes are used to simulate data from the calling client. 
+*/
+USTRUCT()
+struct FDataflowSimulationNode : public FDataflowNode
+{
+	GENERATED_USTRUCT_BODY()
+
+	FDataflowSimulationNode()
+		: Super() { }
+
+	FDataflowSimulationNode(const Dataflow::FNodeParameters& Param, FGuid InGuid = FGuid::NewGuid())
+		: Super(Param,InGuid) {
+	}
+
+	/** FDataflowNode interface */
+	virtual ~FDataflowSimulationNode() {}
+	
+	static FName StaticType() { return FName("FDataflowSimulationNode"); }
+
+	virtual bool IsA(FName InType) const override 
+	{ 
+		return InType.ToString().Equals(StaticType().ToString()) 
+			|| Super::IsA(InType); 
+	}
+
+	/** Evaluate simulation dispatch */
+	virtual void Evaluate(Dataflow::FContext& Context, const FDataflowOutput* Output) const override
+	{
+		if(Context.IsA(Dataflow::FSimulationContext::StaticType()))
+		{
+			Dataflow::FSimulationContext& SimulationContext = StaticCast<Dataflow::FSimulationContext&>(Context);
+			EvaluateSimulation(SimulationContext, Output);
+			
+		}
+	}
+
+	protected :
+
+	/** Evaluate simulation node given a simulation context */
+	virtual void EvaluateSimulation(Dataflow::FSimulationContext& SimulationContext, const FDataflowOutput* Output) const {};
+};
+
+/**
+* FDataflowInvalidNode
+*		Base class for invalid nodes within the Dataflow graph. 
+* 
+*		Invalid nodes will be always invalidated while simulating
+*/
+USTRUCT()
+struct FDataflowInvalidNode : public FDataflowSimulationNode
+{
+	GENERATED_USTRUCT_BODY()
+
+	FDataflowInvalidNode()
+		: Super() { }
+
+	FDataflowInvalidNode(const Dataflow::FNodeParameters& Param, FGuid InGuid = FGuid::NewGuid())
+		: Super(Param,InGuid) {
+	}
+
+	/** FDataflowNode interface */
+	virtual ~FDataflowInvalidNode() {}
+	
+	static FName StaticType() { return FName("FDataflowInvalidNode"); }
+
+	virtual bool IsA(FName InType) const override 
+	{ 
+		return InType.ToString().Equals(StaticType().ToString()) 
+			|| Super::IsA(InType); 
+	}
+};
+
+/**
+* FDataflowCacheNode
+*		Base class for cache nodes within the Dataflow graph. 
+* 
+*		Cache nodes are used to provide an interface to cache data out from the calling client. 
+*/
+USTRUCT()
+struct FDataflowCacheNode : public FDataflowSimulationNode
+{
+	GENERATED_USTRUCT_BODY()
+
+	FDataflowCacheNode()
+		: Super() { }
+
+	FDataflowCacheNode(const Dataflow::FNodeParameters& Param, FGuid InGuid = FGuid::NewGuid())
+		: Super(Param,InGuid) {
+	}
+
+	/** FDataflowNode interface */
+	virtual ~FDataflowCacheNode() {}
+	
+	static FName StaticType() { return FName("FDataflowCacheNode"); }
+
+	virtual bool IsA(FName InType) const override 
+	{ 
+		return InType.ToString().Equals(StaticType().ToString()) 
+			|| Super::IsA(InType); 
+	}
+
+	/** Return the cache asset that we are building */
+	virtual TObjectPtr<UObject> GetCacheAsset(Dataflow::FSimulationContext& SimulationContext) const {return nullptr;} 
+	
+	/** Return the time range used to cache the datas */
+	virtual FVector2f GetTimeRange(Dataflow::FSimulationContext& SimulationContext) const {return FVector2f(0,5.0f);}
+};
+
+/**
+* FDataflowAnimationNode
+*		Base class for animtion nodes within the Dataflow graph. 
+* 
+*		Animation nodes are used to update some animtion datas  from the calling client. 
+*/
+USTRUCT()
+struct FDataflowAnimationNode : public FDataflowSimulationNode
+{
+	GENERATED_USTRUCT_BODY()
+
+	FDataflowAnimationNode()
+		: Super() { }
+
+	FDataflowAnimationNode(const Dataflow::FNodeParameters& Param, FGuid InGuid = FGuid::NewGuid())
+		: Super(Param,InGuid) {
+	}
+
+	/** FDataflowNode interface */
+	virtual ~FDataflowAnimationNode() {}
+	
+	static FName StaticType() { return FName("FDataflowAnimationNode"); }
+
+	virtual bool IsA(FName InType) const override 
+	{ 
+		return InType.ToString().Equals(StaticType().ToString()) 
+			|| Super::IsA(InType); 
+	}
+	
+	/** Set the animation time */
+	virtual void SetAnimationTime(Dataflow::FSimulationContext& SimulationContext, const float AnimationTime) {}
+};
+
+
+
+
+
diff --git a/Engine/Source/Runtime/Experimental/Dataflow/Core/Private/Dataflow/DataflowGraph.cpp b/Engine/Source/Runtime/Experimental/Dataflow/Core/Private/Dataflow/DataflowGraph.cpp
--- a/Engine/Source/Runtime/Experimental/Dataflow/Core/Private/Dataflow/DataflowGraph.cpp
+++ b/Engine/Source/Runtime/Experimental/Dataflow/Core/Private/Dataflow/DataflowGraph.cpp
@@ -14,6 +14,7 @@ DEFINE_LOG_CATEGORY_STATIC(DATAFLOW_LOG, Error, All);
 
 namespace Dataflow
 {
+	TSet<FName> FGraph::RegisteredFilters = {};
 
 	FGraph::FGraph(FGuid InGuid)
 		: Guid(InGuid)
@@ -51,9 +52,15 @@ namespace Dataflow
 			}
 		}
 		Nodes.Remove(Node);
-		if (Node->IsA(FDataflowTerminalNode::StaticType()))
+		for(const FName& RegisteredType : RegisteredFilters)
 		{
-			TerminalNodes.Remove(Node);
+			if (Node->IsA(RegisteredType))
+			{
+				if(TArray< TSharedPtr<FDataflowNode> >* FoundNodes = FilteredNodes.Find(RegisteredType))
+				{
+					FoundNodes->Remove(Node);
+				}
+			}
 		}
 	}
 
@@ -428,5 +435,10 @@ namespace Dataflow
 			ConnectionToFix->FixAndPropagateType();
 		}
 	}
+
+	void RegisterNodeFilter(const FName& NodeFilter)
+	{
+		FGraph::RegisteredFilters.Add(NodeFilter);
+	}
 }
 
diff --git a/Engine/Source/Runtime/Experimental/Dataflow/Core/Public/Dataflow/DataflowGraph.h b/Engine/Source/Runtime/Experimental/Dataflow/Core/Public/Dataflow/DataflowGraph.h
--- a/Engine/Source/Runtime/Experimental/Dataflow/Core/Public/Dataflow/DataflowGraph.h
+++ b/Engine/Source/Runtime/Experimental/Dataflow/Core/Public/Dataflow/DataflowGraph.h
@@ -50,25 +50,43 @@ namespace Dataflow
 
 		FGuid  Guid;
 		TArray< TSharedPtr<FDataflowNode> > Nodes;
-		TArray< TSharedPtr<FDataflowNode> > TerminalNodes;
+		TMap<FName, TArray<TSharedPtr<FDataflowNode>>> FilteredNodes;
 		TArray< FLink > Connections;
 		TSet< FName > DisabledNodes;
+		
+		/** Node filter type that could be used for fast access*/
+		static TSet<FName> RegisteredFilters;
+
+		/** Friend register function */
+		friend DATAFLOWCORE_API void RegisterNodeFilter(const FName& NodeFilter);
+		
 	public:
 		DATAFLOWCORE_API FGraph(FGuid InGuid = FGuid::NewGuid());
 		virtual ~FGraph() {}
-
-		const TArray< TSharedPtr<FDataflowNode> >& GetTerminalNodes() const { return TerminalNodes; }
+		
+		const TArray< TSharedPtr<FDataflowNode> >& GetFilteredNodes(const FName& NodeFilter) const
+		{
+			static TArray< TSharedPtr<FDataflowNode> > EmptyNodes;
+			if(const TArray< TSharedPtr<FDataflowNode> >* FoundNodes = FilteredNodes.Find(NodeFilter))
+			{
+				return *FoundNodes;
+			}
+			return EmptyNodes;
+		}
 		const TArray< TSharedPtr<FDataflowNode> >& GetNodes() const { return Nodes; }
 		TArray< TSharedPtr<FDataflowNode> >& GetNodes() { return Nodes; }
 		int NumNodes() { return Nodes.Num(); }
-
+		
 		template<class T> TSharedPtr<T> AddNode(T* InNode)
 		{
 			TSharedPtr<T> NewNode(InNode);
 			Nodes.AddUnique(NewNode);
-			if (NewNode->IsA(FDataflowTerminalNode::StaticType()))
+			for(const FName& RegisteredType : RegisteredFilters)
 			{
-				TerminalNodes.AddUnique(NewNode);
+				if (NewNode->IsA(RegisteredType))
+				{
+					FilteredNodes.FindOrAdd(RegisteredType).Add(NewNode);
+				}
 			}
 			return NewNode;
 		}
@@ -77,9 +95,12 @@ namespace Dataflow
 		{
 			TSharedPtr<T> NewNode(InNode.Release());
 			Nodes.AddUnique(NewNode);
-			if (NewNode->IsA(FDataflowTerminalNode::StaticType()))
+			for(const FName& RegisteredType : RegisteredFilters)
 			{
-				TerminalNodes.AddUnique(NewNode);
+				if (NewNode->IsA(RegisteredType))
+				{
+					FilteredNodes.FindOrAdd(RegisteredType).Add(NewNode);
+				}
 			}
 			return NewNode;
 		}
@@ -132,9 +153,9 @@ namespace Dataflow
 			return TSharedPtr<const FDataflowNode>(nullptr);
 		}
 
-		TSharedPtr<FDataflowNode> FindTerminalNode(FName InName)
+		TSharedPtr<FDataflowNode> FindFilteredNode(const FName& NodeFilter, FName InName) const
 		{
-			for (TSharedPtr<FDataflowNode> Node : TerminalNodes)
+			for (TSharedPtr<FDataflowNode> Node : GetFilteredNodes(NodeFilter))
 			{
 				if (Node->GetName().IsEqual(InName))
 				{
@@ -161,8 +182,11 @@ namespace Dataflow
 		DATAFLOWCORE_API static void SerializeForSaving(FArchive& Ar, FGraph* InGraph, TArray<TSharedPtr<FDataflowNode>>& InNodes, TArray<FLink>& InConnections);
 		DATAFLOWCORE_API static void SerializeForLoading(FArchive& Ar, FGraph* InGraph, UObject* OwningObject);
 	};
+
+	DATAFLOWCORE_API void RegisterNodeFilter(const FName& NodeFilter);
 }
 
+
 FORCEINLINE FArchive& operator<<(FArchive& Ar, Dataflow::FLink& Value)
 {
 	Ar << Value.InputNode << Value.OutputNode << Value.Input << Value.Output;
diff --git a/Engine/Source/Runtime/Experimental/Dataflow/Engine/Private/Dataflow/DataflowBlueprintLibrary.cpp b/Engine/Source/Runtime/Experimental/Dataflow/Engine/Private/Dataflow/DataflowBlueprintLibrary.cpp
--- a/Engine/Source/Runtime/Experimental/Dataflow/Engine/Private/Dataflow/DataflowBlueprintLibrary.cpp
+++ b/Engine/Source/Runtime/Experimental/Dataflow/Engine/Private/Dataflow/DataflowBlueprintLibrary.cpp
@@ -10,8 +10,7 @@ void UDataflowBlueprintLibrary::EvaluateTerminalNodeByName(UDataflow* Dataflow,
 {
 	if (Dataflow && Dataflow->Dataflow)
 	{
-		TSharedPtr<FDataflowNode> Node = Dataflow->Dataflow->FindTerminalNode(TerminalNodeName);
-		if (Node)
+		if (const TSharedPtr<FDataflowNode> Node = Dataflow->Dataflow->FindFilteredNode(FDataflowTerminalNode::StaticType(), TerminalNodeName))
 		{
 			if (const FDataflowTerminalNode* TerminalNode = Node->AsType<const FDataflowTerminalNode>())
 			{
diff --git a/Engine/Source/Runtime/Experimental/Dataflow/Engine/Public/Dataflow/DataflowObjectInterface.h b/Engine/Source/Runtime/Experimental/Dataflow/Engine/Public/Dataflow/DataflowObjectInterface.h
--- a/Engine/Source/Runtime/Experimental/Dataflow/Engine/Public/Dataflow/DataflowObjectInterface.h
+++ b/Engine/Source/Runtime/Experimental/Dataflow/Engine/Public/Dataflow/DataflowObjectInterface.h
@@ -4,6 +4,7 @@
 
 #include "Containers/UnrealString.h"
 #include "Dataflow/DataflowGraph.h"
+#include "GameFramework/Actor.h"
 #include "Dataflow/DataflowNodeParameters.h"
 #include "GeometryCollection/ManagedArrayCollection.h"
 
@@ -28,7 +29,7 @@ namespace Dataflow
 		TObjectPtr<UObject> Owner = nullptr;
 		TObjectPtr<UDataflow> Graph = nullptr;
 
-		~TEngineContext(){}
+		virtual ~TEngineContext() {}
 
 		int32 GetKeys(TSet<FContextCacheKey>& InKeys) const { return Base::GetKeys(InKeys); }
 
@@ -38,7 +39,101 @@ namespace Dataflow
 
 	};
 
+	/** Simulation flag that could be used across simulation nodes.
+	 * The setup and clean modes are game thread only whereas the
+	 * update mode could run on game thread and physics thread
+	 */
+	enum class ESimulationFlags : uint8
+	{
+		SetupSimulation = 1 << 0,
+		CleanSimulation = 1 << 1,
+		UpdateSimulation = 1 << 2
+	};
+
+	ENUM_CLASS_FLAGS(ESimulationFlags);
+
+	/** Caching mode that could be used across simulation nodes */
+	enum class ECachingMode : uint8
+	{
+		PlayCache = 1 << 0,
+		RecordCache = 1 << 1
+	};
+
+	/** Simulation context that will be used by all the simulation/cache nodes*/
+	template<class Base = FContextSingle>
+	class TSimulationContext : public TEngineContext<Base>
+	{
+	public:
+		DATAFLOW_CONTEXT_INTERNAL(TEngineContext<Base>, TSimulationContext);
+
+		TSimulationContext(const TObjectPtr<UObject>& InOwner,const TObjectPtr<UDataflow>& InGraph, FTimestamp InTimestamp,
+			const TObjectPtr<AActor>& InRootActor, const ECachingMode InCachingMode)
+				: Super(InOwner, InGraph, InTimestamp), SimulationWorld(InRootActor->GetWorld()), RootActor(InRootActor),
+					SimulationFlags(ESimulationFlags::SetupSimulation), CachingMode(InCachingMode)
+		{
+			
+		}
+		
+		virtual ~TSimulationContext() override {};
+		
+		/** Set the timing infos */
+		void SetTimingInfos(const float DeltaSeconds, const float TimeSeconds) {DeltaTime = DeltaSeconds; SimulationTime = TimeSeconds;}
+
+		/** Get the delta time in seconds */
+		float GetDeltaTime() const {return DeltaTime;}
+		
+		/** Get the simulation time in seconds */
+		float GetSimulationTime() const {return SimulationTime;}
+		
+		/** Get the simulation world */
+		const TObjectPtr<UWorld>& GetSimulationWorld() const {return SimulationWorld;}
+		
+		/** Get the root actor */
+		const TObjectPtr<AActor>& GetRootActor() const {return RootActor;}
+
+		/** Check if the simulation mode is present in the simulation context */
+		bool HasSimulationFlag(const ESimulationFlags InSimulationFlag) const
+		{
+			return (SimulationFlags & InSimulationFlag) == InSimulationFlag;
+		}
+
+		/** Set the simulation mode */
+		void SetSimulationFlag(const ESimulationFlags InSimulationFlag) {SimulationFlags |= InSimulationFlag;}
+
+		/** Unset the simulation mode */
+		void UnsetSimulationFlag(const ESimulationFlags InSimulationFlag) {SimulationFlags &= ~InSimulationFlag;}
+
+		/** Get the caching mode */
+		ECachingMode GetCachingMode() const {return CachingMode;}
+
+		/** Set the caching mode */
+		void SetCachingMode(const ECachingMode InCachingMode) {CachingMode = InCachingMode;}
+		
+	private :
+		
+		/** Evolution world used to register components */
+		TObjectPtr<UWorld> SimulationWorld;
+
+		/** Root actor of the world */
+		TObjectPtr<AActor> RootActor;
+
+		/** Simulation time */
+		float SimulationTime = 0.0f;
+
+		/** Delta time */
+		float DeltaTime = 0.0f;
+
+		/** Simulation flags that could be used across simulation nodes */
+		ESimulationFlags SimulationFlags;
+
+		/** Caching mode that could be used across cache nodes */
+		ECachingMode CachingMode;
+	};
+
 	typedef TEngineContext<FContextSingle> FEngineContext;
 	typedef TEngineContext<FContextThreaded> FEngineContextThreaded;
 
+	typedef TSimulationContext<FContextSingle> FSimulationContext;
+	typedef TSimulationContext<FContextThreaded> FSimulationContextThreaded;
+
 }
