Dataflow Editor:
- when a tool is running hide wireframe rendering in the scene and disable toggling wireframe mode on graph nodes
- when clicking on a node, don't update ShouldRenderNode unless the node selection actually changes

diff --git a/Engine/Plugins/Experimental/Dataflow/Source/DataflowEditor/Private/Dataflow/DataflowEditorMode.cpp b/Engine/Plugins/Experimental/Dataflow/Source/DataflowEditor/Private/Dataflow/DataflowEditorMode.cpp
--- a/Engine/Plugins/Experimental/Dataflow/Source/DataflowEditor/Private/Dataflow/DataflowEditorMode.cpp
+++ b/Engine/Plugins/Experimental/Dataflow/Source/DataflowEditor/Private/Dataflow/DataflowEditorMode.cpp
@@ -232,9 +232,30 @@ void UDataflowEditorMode::CreateToolkit()
 	Toolkit = MakeShared<FDataflowEditorModeToolkit>();
 }
 
+
+void UDataflowEditorMode::SetWireframeRenderToggleEnabled(bool bEnable)
+{
+	if (const TObjectPtr<UDataflowBaseContent> EditorContent = ConstructionScene->GetEditorContent())
+	{
+		if (const TObjectPtr<UDataflow> DataflowGraph = EditorContent->GetDataflowAsset())
+		{
+			for (UEdGraphNode* const EdGraphNode : DataflowGraph->Nodes)
+			{
+				if (UDataflowEdNode* const DataflowEdNode = Cast<UDataflowEdNode>(EdGraphNode))
+				{
+					DataflowEdNode->SetCanEnableWireframeRenderNode(bEnable);
+				}
+			}
+		}
+	}
+}
+
 void UDataflowEditorMode::OnToolStarted(UInteractiveToolManager* Manager, UInteractiveTool* Tool)
 {
 	FDataflowEditorCommandsImpl::UpdateToolCommandBinding(Tool, ToolCommandList, false);
+
+	// Temporarily disable wireframe render toggle switch on all nodes
+	SetWireframeRenderToggleEnabled(false);
 }
 
 void UDataflowEditorMode::OnToolEnded(UInteractiveToolManager* Manager, UInteractiveTool* Tool)
@@ -261,6 +282,9 @@ void UDataflowEditorMode::OnToolEnded(UInteractiveToolManager* Manager, UInterac
 	{
 		GraphEditor->SetEnabled(true);
 	}
+
+	// Re-enable wireframe render toggle switch on all nodes
+	SetWireframeRenderToggleEnabled(true);
 }
 
 void UDataflowEditorMode::BindCommands()
diff --git a/Engine/Plugins/Experimental/Dataflow/Source/DataflowEditor/Private/Dataflow/DataflowEditorToolkit.cpp b/Engine/Plugins/Experimental/Dataflow/Source/DataflowEditor/Private/Dataflow/DataflowEditorToolkit.cpp
--- a/Engine/Plugins/Experimental/Dataflow/Source/DataflowEditor/Private/Dataflow/DataflowEditorToolkit.cpp
+++ b/Engine/Plugins/Experimental/Dataflow/Source/DataflowEditor/Private/Dataflow/DataflowEditorToolkit.cpp
@@ -655,13 +655,7 @@ void FDataflowEditorToolkit::OnNodeSelectionChanged(const TSet<UObject*>& InNewS
 			return Objs;
 		};
 
-		for (UObject* const PreviouslySelectedNode : SelectedDataflowNodes)
-		{
-			if (UDataflowEdNode* PreviouslySelectedEdNode = Cast<UDataflowEdNode>(PreviouslySelectedNode))
-			{
-				PreviouslySelectedEdNode->SetShouldRenderNode(false);
-			}
-		}
+		const TSet<TObjectPtr<UObject>> PreviouslySelectedNodes = SelectedDataflowNodes;
 
 		// Only keep UDataflowEdNode from NewSelection
 		TSet< TObjectPtr<UObject> > NodeSelection = FindDataflowNodesInSet(AsObjectPointers(InNewSelection));
@@ -710,17 +704,25 @@ void FDataflowEditorToolkit::OnNodeSelectionChanged(const TSet<UObject*>& InNewS
 
 			SelectedDataflowNodes = NodeSelection;
 		}
-
-		for (UObject* const SelectedNode : NodeSelection)
+	
+		if (bPrimarySelectionChanged)
 		{
-			if (UDataflowEdNode* SelectedEdNode = Cast<UDataflowEdNode>(SelectedNode))
+			for (UObject* const PreviouslySelectedNode : PreviouslySelectedNodes)
 			{
-				SelectedEdNode->SetShouldRenderNode(true);
+				if (UDataflowEdNode* PreviouslySelectedEdNode = Cast<UDataflowEdNode>(PreviouslySelectedNode))
+				{
+					PreviouslySelectedEdNode->SetShouldRenderNode(false);
+				}
+			}
+
+			for (UObject* const SelectedNode : NodeSelection)
+			{
+				if (UDataflowEdNode* SelectedEdNode = Cast<UDataflowEdNode>(SelectedNode))
+				{
+					SelectedEdNode->SetShouldRenderNode(true);
+				}
 			}
-		}
 
-		if (bPrimarySelectionChanged)
-		{
 			EditorContent->SetPrimarySelectedNode(nullptr);
 
 			UDataflowEditorMode* const DataflowMode = CastChecked<UDataflowEditorMode>(EditorModeManager->GetActiveScriptableMode(UDataflowEditorMode::EM_DataflowEditorModeId));
diff --git a/Engine/Plugins/Experimental/Dataflow/Source/DataflowEditor/Private/Dataflow/DataflowSNode.cpp b/Engine/Plugins/Experimental/Dataflow/Source/DataflowEditor/Private/Dataflow/DataflowSNode.cpp
--- a/Engine/Plugins/Experimental/Dataflow/Source/DataflowEditor/Private/Dataflow/DataflowSNode.cpp
+++ b/Engine/Plugins/Experimental/Dataflow/Source/DataflowEditor/Private/Dataflow/DataflowSNode.cpp
@@ -69,6 +69,14 @@ void SDataflowEdNode::Construct(const FArguments& InArgs, UDataflowEdNode* InNod
 					else
 						DataflowGraphNode->SetShouldWireframeRenderNode(false);
 				}
+			})
+		.IsEnabled_Lambda([this]()->bool
+			{
+				if (DataflowGraphNode)
+				{
+					return DataflowGraphNode->CanEnableWireframeRenderNode();
+				}
+				return false;
 			});
 
 
diff --git a/Engine/Plugins/Experimental/Dataflow/Source/DataflowEditor/Private/DataflowEditorTools/DataflowEditorWeightMapPaintTool.cpp b/Engine/Plugins/Experimental/Dataflow/Source/DataflowEditor/Private/DataflowEditorTools/DataflowEditorWeightMapPaintTool.cpp
--- a/Engine/Plugins/Experimental/Dataflow/Source/DataflowEditor/Private/DataflowEditorTools/DataflowEditorWeightMapPaintTool.cpp
+++ b/Engine/Plugins/Experimental/Dataflow/Source/DataflowEditor/Private/DataflowEditorTools/DataflowEditorWeightMapPaintTool.cpp
@@ -170,32 +170,18 @@ void UDataflowEditorWeightMapPaintTool::Setup()
 {
 	UMeshSculptToolBase::Setup();
 
-	auto IsolateComponent = [&](UDataflowEditorCollectionComponent* SelectedComponent)
-	{
-		if (FDataflowConstructionScene* Scene = Mode->GetDataflowConstructionScene())
-		{
-			Scene->SetVisibility(false);
-			Scene->SetVisibility(true, SelectedComponent);
-		}
-	};
-	if (TSharedPtr<FDataflowConstructionViewportClient> VC = Mode->GetConstructionViewportClient().Pin())
-	{
-		if (USelection* SelectedComponents = VC->GetSelectedComponents())
-		{
-			if (UDataflowEditorCollectionComponent* DataflowComponent = SelectedComponents->GetBottom<UDataflowEditorCollectionComponent>())
-			{
-				IsolateComponent(DataflowComponent);
-			}
-		}
-	}
-
-
 	// Get the selected weight map node
 	WeightMapNodeToUpdate = DataflowEditorContextObject->GetPrimarySelectedNodeOfType<FDataflowCollectionAddScalarVertexPropertyNode>();
 	checkf(WeightMapNodeToUpdate, TEXT("No Weight Map Node is currently selected, or more than one node is selected"));
 
 	SetToolDisplayName(LOCTEXT("ToolName", "Paint Weight Maps"));
 
+	// Hide all meshes in the DataflowConstructionScene, as we will be painting onto our own Preview mesh
+	if (FDataflowConstructionScene* Scene = Mode->GetDataflowConstructionScene())
+	{
+		Scene->SetVisibility(false);
+	}
+
 	// create dynamic mesh component to use for live preview
 	FActorSpawnParameters SpawnInfo;
 	PreviewMeshActor = TargetWorld->SpawnActor<AInternalToolFrameworkActor>(FVector::ZeroVector, FRotator::ZeroRotator, SpawnInfo);
diff --git a/Engine/Plugins/Experimental/Dataflow/Source/DataflowEditor/Public/Dataflow/DataflowEditorMode.h b/Engine/Plugins/Experimental/Dataflow/Source/DataflowEditor/Public/Dataflow/DataflowEditorMode.h
--- a/Engine/Plugins/Experimental/Dataflow/Source/DataflowEditor/Public/Dataflow/DataflowEditorMode.h
+++ b/Engine/Plugins/Experimental/Dataflow/Source/DataflowEditor/Public/Dataflow/DataflowEditorMode.h
@@ -157,6 +157,9 @@ private:
 	*/
 	UEdGraphNode* CreateAndConnectNewNode(const FName& NewNodeTypeName, UEdGraphNode& UpstreamNode, const FName& ConnectionTypeName, const FName& NewNodeConnectionName);
 
+	/** On all nodes in the graph, set whether the wireframe rendering toggle button is enabled */
+	void SetWireframeRenderToggleEnabled(bool bEnable);
+
 	void InitializeContextObject();
 	void DeleteContextObject();
 
diff --git a/Engine/Source/Runtime/Experimental/Dataflow/Engine/Private/Dataflow/DataflowEdNode.cpp b/Engine/Source/Runtime/Experimental/Dataflow/Engine/Private/Dataflow/DataflowEdNode.cpp
--- a/Engine/Source/Runtime/Experimental/Dataflow/Engine/Private/Dataflow/DataflowEdNode.cpp
+++ b/Engine/Source/Runtime/Experimental/Dataflow/Engine/Private/Dataflow/DataflowEdNode.cpp
@@ -75,6 +75,16 @@ void UDataflowEdNode::SetShouldWireframeRenderNode(bool bInRender)
 	}
 }
 
+void UDataflowEdNode::SetCanEnableWireframeRenderNode(bool bInCanEnable)
+{
+	bCanEnableRenderWireframe = bInCanEnable;
+}
+
+bool UDataflowEdNode::CanEnableWireframeRenderNode() const
+{
+	return bCanEnableRenderWireframe;
+}
+
 TSharedPtr<FDataflowNode> UDataflowEdNode::GetDataflowNode()
 {
 	if(TSharedPtr<Dataflow::FGraph> Dataflow = GetDataflowGraph())
diff --git a/Engine/Source/Runtime/Experimental/Dataflow/Engine/Public/Dataflow/DataflowEdNode.h b/Engine/Source/Runtime/Experimental/Dataflow/Engine/Public/Dataflow/DataflowEdNode.h
--- a/Engine/Source/Runtime/Experimental/Dataflow/Engine/Public/Dataflow/DataflowEdNode.h
+++ b/Engine/Source/Runtime/Experimental/Dataflow/Engine/Public/Dataflow/DataflowEdNode.h
@@ -96,6 +96,9 @@ public:
 	DATAFLOWENGINE_API void SetShouldWireframeRenderNode(bool bInRender);
 	bool ShouldWireframeRenderNode() const { return bRenderWireframeInAssetEditor; }
 
+	DATAFLOWENGINE_API void SetCanEnableWireframeRenderNode(bool bInCanEnable);
+	DATAFLOWENGINE_API bool CanEnableWireframeRenderNode() const;
+
 	DATAFLOWENGINE_API TArray<Dataflow::FRenderingParameter> GetRenderParameters() const;
 	DATAFLOWENGINE_API virtual bool Render(GeometryCollection::Facades::FRenderingFacade& RenderData, const TSharedPtr<Dataflow::FContext> Context) const;
 
@@ -107,6 +110,9 @@ private:
 	UPROPERTY()
 	bool bRenderWireframeInAssetEditor = false;
 
+	UPROPERTY()
+	bool bCanEnableRenderWireframe = true;
+
 	void OnPinConnection(const UEdGraphPin& Pin, const FEdGraphPinType& Type);
 	void OnPinDisconnection(const UEdGraphPin& Pin);
 	void SetAnyTypePinsToType(const FEdGraphPinType& Type);
