Dataflow Editor - Fixed crash when changing comment properties.

#rnx

diff --git a/Engine/Plugins/Experimental/Dataflow/Source/DataflowEditor/Private/Dataflow/DataflowEditorCommands.cpp b/Engine/Plugins/Experimental/Dataflow/Source/DataflowEditor/Private/Dataflow/DataflowEditorCommands.cpp
--- a/Engine/Plugins/Experimental/Dataflow/Source/DataflowEditor/Private/Dataflow/DataflowEditorCommands.cpp
+++ b/Engine/Plugins/Experimental/Dataflow/Source/DataflowEditor/Private/Dataflow/DataflowEditorCommands.cpp
@@ -355,17 +355,18 @@ void FDataflowEditorCommands::OnPropertyValueChanged(UDataflow* OutDataflow, TSh
 
 			for (UObject* const SelectedNode : SelectedNodes)
 			{
-				UDataflowEdNode* const Node = CastChecked<UDataflowEdNode>(SelectedNode);  // SelectedNodes is assumed to be the graph panel selection filtered to UDataflowEdNode
-
-				if (TSharedPtr<FDataflowNode> DataflowNode = Node->GetDataflowNode())
+				if (UDataflowEdNode* const Node = Cast<UDataflowEdNode>(SelectedNode))
 				{
-					DataflowNode->Invalidate();
-					OutLastNodeTimestamp = Dataflow::FTimestamp::Invalid;
-
-					// Reflect the active state on the drawing of the node
-					if (DataflowNode->bActive != Node->IsNodeEnabled())
+					if (TSharedPtr<FDataflowNode> DataflowNode = Node->GetDataflowNode())
 					{
-						Node->SetEnabledState(DataflowNode->bActive ? ENodeEnabledState::Enabled : ENodeEnabledState::Disabled);
+						DataflowNode->Invalidate();
+						OutLastNodeTimestamp = Dataflow::FTimestamp::Invalid;
+
+						// Reflect the active state on the drawing of the node
+						if (DataflowNode->bActive != Node->IsNodeEnabled())
+						{
+							Node->SetEnabledState(DataflowNode->bActive ? ENodeEnabledState::Enabled : ENodeEnabledState::Disabled);
+						}
 					}
 				}
 			}
