Dataflow : SkeletalView selection from selected construciton view components.

diff --git a/Engine/Plugins/Experimental/Dataflow/Source/DataflowEditor/Private/Dataflow/DataflowConstructionScene.cpp b/Engine/Plugins/Experimental/Dataflow/Source/DataflowEditor/Private/Dataflow/DataflowConstructionScene.cpp
--- a/Engine/Plugins/Experimental/Dataflow/Source/DataflowEditor/Private/Dataflow/DataflowConstructionScene.cpp
+++ b/Engine/Plugins/Experimental/Dataflow/Source/DataflowEditor/Private/Dataflow/DataflowConstructionScene.cpp
@@ -203,7 +203,8 @@ void FDataflowConstructionScene::UpdateDynamicMeshComponents()
 									EditorContent->SetPrimaryRenderCollection(RenderCollection);
 								}
 								const FString MeshName = Facade.GetGeometryName()[MeshIndex];
-								UDynamicMeshComponent* const NewDynamicMeshComponent = AddDynamicMeshComponent(WireframeDynamicMeshKey, MeshName, MoveTemp(DynamicMesh), {});
+								const FString UniqueObjectName = MakeUniqueObjectName(RootSceneActor, UDataflowEditorCollectionComponent::StaticClass(), FName(MeshName)).ToString();
+								UDynamicMeshComponent* const NewDynamicMeshComponent = AddDynamicMeshComponent(WireframeDynamicMeshKey, UniqueObjectName, MoveTemp(DynamicMesh), {});
 								NewDynamicMeshComponent->SetVisibility(false);
 								MeshComponentsForWireframeRendering.Add(NewDynamicMeshComponent);
 							}
@@ -235,9 +236,10 @@ void FDataflowConstructionScene::ResetDynamicMeshComponents()
 
 TObjectPtr<UDynamicMeshComponent>& FDataflowConstructionScene::AddDynamicMeshComponent(FDataflowRenderKey InKey, const FString& MeshName, UE::Geometry::FDynamicMesh3&& DynamicMesh, const TArray<UMaterialInterface*>& MaterialSet)
 {
-	const FName UniqueObjectName = MakeUniqueObjectName(RootSceneActor, UDataflowEditorCollectionComponent::StaticClass(), FName(MeshName));
-	
-	TObjectPtr<UDataflowEditorCollectionComponent> DynamicMeshComponent = NewObject<UDataflowEditorCollectionComponent>(RootSceneActor, UniqueObjectName);
+	// Dont use the MakeUniqueObjectName for the component, we need to keep the name alinged with the collection so selection will work in 
+	// other editors. 
+	// const FName UniqueObjectName = MakeUniqueObjectName(RootSceneActor, UDataflowEditorCollectionComponent::StaticClass(), FName(MeshName));
+	TObjectPtr<UDataflowEditorCollectionComponent> DynamicMeshComponent = NewObject<UDataflowEditorCollectionComponent>(RootSceneActor, FName(MeshName));
 
 	DynamicMeshComponent->MeshIndex = InKey.Value;
 	DynamicMeshComponent->Node = InKey.Key;
