Dataflow : Add a VertexWeightsRenderer

diff --git a/Engine/Plugins/Experimental/Dataflow/Source/DataflowEnginePlugin/Private/Dataflow/DataflowEngineRendering.cpp b/Engine/Plugins/Experimental/Dataflow/Source/DataflowEnginePlugin/Private/Dataflow/DataflowEngineRendering.cpp
--- a/Engine/Plugins/Experimental/Dataflow/Source/DataflowEnginePlugin/Private/Dataflow/DataflowEngineRendering.cpp
+++ b/Engine/Plugins/Experimental/Dataflow/Source/DataflowEnginePlugin/Private/Dataflow/DataflowEngineRendering.cpp
@@ -2,6 +2,7 @@
 
 #include "Dataflow/DataflowEngineRendering.h"
 
+#include "Dataflow/DataflowConnectionTypes.h"
 #include "Dataflow/DataflowEnginePlugin.h"
 #include "Dataflow/DataflowRenderingFactory.h"
 #include "DynamicMesh/DynamicMesh3.h"
@@ -12,9 +13,10 @@
 #include "GeometryCollection/ManagedArrayCollection.h"
 #include "GeometryCollection/GeometryCollectionAlgo.h"
 #include "UDynamicMesh.h"
+
 namespace Dataflow
 {
-	void RenderBasicGeometryCollection(GeometryCollection::Facades::FRenderingFacade& RenderCollection, const Dataflow::FGraphRenderingState& State)
+	void RenderBasicGeometryCollection(GeometryCollection::Facades::FRenderingFacade& RenderCollection, const Dataflow::FGraphRenderingState& State, TArray<FLinearColor>* VertexColorOverride = nullptr)
 	{
 		FManagedArrayCollection Default;
 		FName PrimaryOutput = State.GetRenderOutputs()[0]; // "Collection"
@@ -86,18 +88,28 @@ namespace Dataflow
 
 		// Copy VertexColors from the Collection if exists otherwise set it to IDataflowEnginePlugin::SurfaceColor
 		TArray<FLinearColor> VertexColors; VertexColors.AddUninitialized(Vertex.Num());
-		if (const TManagedArray<FLinearColor>* VertexColorManagedArray = Collection.FindAttribute<FLinearColor>("Color", FGeometryCollection::VerticesGroup))
+		if (VertexColorOverride && VertexColorOverride->Num() == Vertex.Num())
 		{
 			for (int32 VertexIdx = 0; VertexIdx < VertexColors.Num(); ++VertexIdx)
 			{
-				VertexColors[VertexIdx] = (*VertexColorManagedArray)[VertexIdx];
+				VertexColors[VertexIdx] = (*VertexColorOverride)[VertexIdx];
 			}
 		}
 		else
 		{
-			for (int32 VertexIdx = 0; VertexIdx < VertexColors.Num(); ++VertexIdx)
+			if (const TManagedArray<FLinearColor>* VertexColorManagedArray = Collection.FindAttribute<FLinearColor>("Color", FGeometryCollection::VerticesGroup))
 			{
-				VertexColors[VertexIdx] = FLinearColor(IDataflowEnginePlugin::SurfaceColor);
+				for (int32 VertexIdx = 0; VertexIdx < VertexColors.Num(); ++VertexIdx)
+				{
+					VertexColors[VertexIdx] = (*VertexColorManagedArray)[VertexIdx];
+				}
+			}
+			else
+			{
+				for (int32 VertexIdx = 0; VertexIdx < VertexColors.Num(); ++VertexIdx)
+				{
+					VertexColors[VertexIdx] = FLinearColor(IDataflowEnginePlugin::SurfaceColor);
+				}
 			}
 		}
 
@@ -108,7 +120,7 @@ namespace Dataflow
 
 	}
 
-	void RenderMeshIndexedGeometryCollection(GeometryCollection::Facades::FRenderingFacade& RenderCollection, const Dataflow::FGraphRenderingState& State)
+	void RenderMeshIndexedGeometryCollection(GeometryCollection::Facades::FRenderingFacade& RenderCollection, const Dataflow::FGraphRenderingState& State, TArray<FLinearColor>* VertexColorOverride = nullptr )
 	{
 		auto ToD = [](FVector3f V) { return FVector3d(V.X, V.Y, V.Z); };
 		auto ToF = [](FVector3d V) { return FVector3f(V.X, V.Y, V.Z); };
@@ -129,7 +141,7 @@ namespace Dataflow
 		const TManagedArray<int32>& VertexCount = Collection.GetAttribute<int32>("VertexCount", FGeometryCollection::GeometryGroup);
 		const TManagedArray<int32>& FacesStart = Collection.GetAttribute<int32>("FaceStart", FGeometryCollection::GeometryGroup);
 		const TManagedArray<int32>& FacesCount = Collection.GetAttribute<int32>("FaceCount", FGeometryCollection::GeometryGroup);
-		
+		int32 TotalVertices = Collection.NumElements(FGeometryCollection::VerticesGroup);
 
 		TArray<FTransform> M;
 		GeometryCollectionAlgo::GlobalMatrices(Transforms, Parents, M);
@@ -152,7 +164,7 @@ namespace Dataflow
 
 				FIntVector Tri = FIntVector(Face[0], Face[1], Face[2]);
 				FTransform Ms[3] = { M[BoneIndex[Tri[0]]], M[BoneIndex[Tri[1]]], M[BoneIndex[Tri[2]]] };
-				FIntVector MovedTri = FIntVector(Face[0]- VertexStart[Gdx], Face[1] - VertexStart[Gdx], Face[2] - VertexStart[Gdx]);
+				FIntVector MovedTri = FIntVector(Face[0] - VertexStart[Gdx], Face[1] - VertexStart[Gdx], Face[2] - VertexStart[Gdx]);
 
 				Tris[Tdx++] = MovedTri;
 				if (!Visited[MovedTri[0]]) Vertices[Tri[0] - VertexStart[Gdx]] = ToF(Ms[0].TransformPosition(ToD(Vertex[Tri[0]])));
@@ -169,7 +181,7 @@ namespace Dataflow
 			{
 				if (!Visited[i])
 				{
-					Vertices[i] = ToF(M[BoneIndex[i+ VertexStart[Gdx]]].TransformPosition(ToD(Vertex[i+ VertexStart[Gdx]])));
+					Vertices[i] = ToF(M[BoneIndex[i + VertexStart[Gdx]]].TransformPosition(ToD(Vertex[i + VertexStart[Gdx]])));
 				}
 			}
 
@@ -195,19 +207,30 @@ namespace Dataflow
 
 			// Copy VertexColors from the Collection if exists otherwise set it to IDataflowEnginePlugin::SurfaceColor
 			TArray<FLinearColor> VertexColors; VertexColors.AddUninitialized(Vertices.Num());
-			if (const TManagedArray<FLinearColor>* VertexColorManagedArray = Collection.FindAttribute<FLinearColor>("Color", FGeometryCollection::VerticesGroup))
+			if (VertexColorOverride && VertexColorOverride->Num() == TotalVertices)
 			{
 				int32 LastVertIndex = VertexStart[Gdx] + VertexCount[Gdx];
 				for (int32 VertexIdx = VertexStart[Gdx], SrcVertexIdx = 0; VertexIdx < LastVertIndex; ++VertexIdx, ++SrcVertexIdx)
 				{
-					VertexColors[SrcVertexIdx] = (*VertexColorManagedArray)[VertexIdx];
+					VertexColors[SrcVertexIdx] = (*VertexColorOverride)[VertexIdx];
 				}
 			}
 			else
 			{
-				for (int32 VertexIdx = 0; VertexIdx < VertexColors.Num(); ++VertexIdx)
+				if (const TManagedArray<FLinearColor>* VertexColorManagedArray = Collection.FindAttribute<FLinearColor>("Color", FGeometryCollection::VerticesGroup))
 				{
-					VertexColors[VertexIdx] = FLinearColor(IDataflowEnginePlugin::SurfaceColor);
+					int32 LastVertIndex = VertexStart[Gdx] + VertexCount[Gdx];
+					for (int32 VertexIdx = VertexStart[Gdx], SrcVertexIdx = 0; VertexIdx < LastVertIndex; ++VertexIdx, ++SrcVertexIdx)
+					{
+						VertexColors[SrcVertexIdx] = (*VertexColorManagedArray)[VertexIdx];
+					}
+				}
+				else
+				{
+					for (int32 VertexIdx = 0; VertexIdx < VertexColors.Num(); ++VertexIdx)
+					{
+						VertexColors[VertexIdx] = FLinearColor(IDataflowEnginePlugin::SurfaceColor);
+					}
 				}
 			}
 
@@ -217,7 +240,7 @@ namespace Dataflow
 				FString GeometryName = State.GetGuid().ToString(); GeometryName.AppendChar('.').AppendInt(Gdx);
 				if (BoneIndex[VertexStart[Gdx]] != INDEX_NONE)
 				{
-					 GeometryName = BoneNames[BoneIndex[VertexStart[Gdx]]];
+					GeometryName = BoneNames[BoneIndex[VertexStart[Gdx]]];
 				}
 				int32 GeometryIndex = RenderCollection.StartGeometryGroup(GeometryName);
 				RenderCollection.AddSurface(MoveTemp(Vertices), MoveTemp(Tris), MoveTemp(VertexNormals), MoveTemp(VertexColors));
@@ -233,8 +256,9 @@ namespace Dataflow
 		/**
 		* DataflowNode (FGeometryCollection) Rendering
 		*
+		*		@param Name : "SurfaceRender"
 		*		@param Type : FGeometryCollection::StaticType()
-
+		* 
 		*		@param Outputs : {FManagedArrayCollection : "Collection"}
 		*/
 		FRenderingFactory::GetInstance()->RegisterOutput({ "SurfaceRender", FGeometryCollection::StaticType()},
@@ -275,11 +299,110 @@ namespace Dataflow
 			});
 
 
+		/**
+		* DataflowNode (FGeometryCollection) Rendering a vertex weight map
+		*
+		*		@param Name : "SurfaceWeightsRender"
+		*		@param Type : FGeometryCollection::StaticType()
+		*		@param Outputs : {FManagedArrayCollection : "Collection"}
+		*		@param Outputs : {FAttributeKey : "AttributeKey"}
+		*/
+		FRenderingFactory::GetInstance()->RegisterOutput({ "SurfaceWeightsRender", FGeometryCollection::StaticType() },
+			[](GeometryCollection::Facades::FRenderingFacade& RenderCollection, const Dataflow::FGraphRenderingState& State)
+			{
+				if (State.GetRenderOutputs().Num()>=2)
+				{
+					FManagedArrayCollection Default;
+					FName PrimaryOutput = State.GetRenderOutputs()[0]; // "Collection"
+					const FManagedArrayCollection & Collection = State.GetValue<FManagedArrayCollection>(PrimaryOutput, Default);
+
+					const bool bFoundIndices = Collection.FindAttributeTyped<FIntVector>("Indices", FGeometryCollection::FacesGroup) != nullptr;
+					const bool bFoundVertices = Collection.FindAttributeTyped<FVector3f>("Vertex", FGeometryCollection::VerticesGroup) != nullptr;
+					const bool bFoundTransforms = Collection.FindAttributeTyped<FTransform3f>(FTransformCollection::TransformAttribute, FTransformCollection::TransformGroup) != nullptr;
+					const bool bFoundBoneMap = Collection.FindAttributeTyped<int32>("BoneMap", FGeometryCollection::VerticesGroup) != nullptr;
+					const bool bFoundParents = Collection.FindAttributeTyped<int32>(FTransformCollection::ParentAttribute, FTransformCollection::TransformGroup) != nullptr;
+					UE_LOG(LogTemp, Warning, TEXT("Render GC with found params = %d %d %d %d %d"), bFoundIndices, bFoundVertices, bFoundTransforms, bFoundBoneMap, bFoundParents);
+					bool bFoundRenderData = bFoundIndices && bFoundVertices && bFoundTransforms && bFoundBoneMap && bFoundParents
+						&& Collection.NumElements(FTransformCollection::TransformGroup) > 0;
+
+					const bool bFoundVertexStart = Collection.FindAttributeTyped<int32>("VertexStart", FGeometryCollection::GeometryGroup) != nullptr;
+					const bool bFoundVertexCount = Collection.FindAttributeTyped<int32>("VertexCount", FGeometryCollection::GeometryGroup) != nullptr;
+					const bool bFoundFaceStart = Collection.FindAttributeTyped<int32>("FaceStart", FGeometryCollection::GeometryGroup) != nullptr;
+					const bool bFoundFaceCount = Collection.FindAttributeTyped<int32>("FaceCount", FGeometryCollection::GeometryGroup) != nullptr;
+					UE_LOG(LogTemp, Warning, TEXT("Render GC with found mesh group params = %d %d %d %d"), bFoundVertexStart, bFoundVertexCount, bFoundFaceStart, bFoundFaceCount);
+					bool bFoundGeometryAttributes = bFoundVertexStart && bFoundVertexCount && bFoundFaceStart && bFoundFaceCount
+						&& Collection.NumElements(FGeometryCollection::GeometryGroup) > 0;
+
+					FCollectionAttributeKey DefaultKey;
+					FName SecondaryOutput = State.GetRenderOutputs()[1]; // "AttributeKey"
+					const FCollectionAttributeKey& AttributeKey = State.GetValue<FCollectionAttributeKey>(SecondaryOutput, DefaultKey);
+
+					const bool bFoundVertexColor = Collection.FindAttributeTyped<FLinearColor>("Color", FGeometryCollection::VerticesGroup) != nullptr;
+					const bool bFoundFloatScalar = Collection.FindAttributeTyped<float>(FName(AttributeKey.Attribute), FName(AttributeKey.Group)) != nullptr;
+					bool bFoundVertexScalarAndColors = bFoundVertexColor && bFoundFloatScalar && AttributeKey.Group.Equals(FGeometryCollection::VerticesGroup.ToString());
+
+					TArray<FLinearColor>* Colors = nullptr;
+					if (bFoundVertexScalarAndColors)
+					{
+						auto RangeValue = [](const TManagedArray<float>* FloatArray)
+						{
+							float Min = FLT_MAX;
+							float Max = -FLT_MAX;
+							for (int i = 0; i < FloatArray->Num(); i++) {
+								Min = FMath::Min(Min, (*FloatArray)[i]);
+								Max = FMath::Max(Max, (*FloatArray)[i]);
+							}
+							return TPair<float, float>(Min, Max);
+						};
+
+						const TManagedArray<float>* FloatArray = Collection.FindAttributeTyped<float>(FName(AttributeKey.Attribute), FName(AttributeKey.Group));
+						if (FloatArray && FloatArray->Num())
+						{
+							Colors = new TArray<FLinearColor>();
+							Colors->AddUninitialized(FloatArray->Num());
+
+							TPair<float, float> Range = RangeValue(FloatArray);
+							float Delta = FMath::Abs(Range.Get<1>() - Range.Get<0>());
+							if (Delta > FLT_EPSILON)
+							{
+								for (int32 VertexIdx = 0; VertexIdx < FloatArray->Num(); ++VertexIdx)
+								{
+									(*Colors)[VertexIdx] = FLinearColor::White * ((*FloatArray)[VertexIdx] - Range.Get<0>()) / Delta;
+								}
+							}
+							else
+							{
+								for (int32 VertexIdx = 0; VertexIdx < FloatArray->Num(); ++VertexIdx)
+								{
+									(*Colors)[VertexIdx] = FLinearColor::Black;
+								}
+							}
+						}
+					}
+
+					if (bFoundRenderData && bFoundGeometryAttributes)
+					{
+						RenderMeshIndexedGeometryCollection(RenderCollection, State, Colors);
+					}
+					else if (bFoundRenderData)
+					{
+						RenderBasicGeometryCollection(RenderCollection, State, Colors);
+					}
+
+					if (Colors)
+					{
+						delete Colors;
+					}
+				}
+			});
+
+
 		/**
 		* DataflowNode (FDynamicMesh3) Rendering
 		*
+		*		@param Name : "SurfaceRender"
 		*		@param Type : FName("FDynamicMesh3")
-
+		*
 		*		@param Outputs : {FDynamicMesh3 : "Mesh"}
 		*/
 		FRenderingFactory::GetInstance()->RegisterOutput({ "SurfaceRender", FName("FDynamicMesh3") },
@@ -340,8 +463,9 @@ namespace Dataflow
 		/**
 		* DataflowNode (FBox) Rendering
 		*
+		*		@param Name : "SurfaceRender"
 		*		@param Type : FName("FBox")
-
+		*
 		*		@param Outputs : {FBox : "Box"}
 		*/
 		FRenderingFactory::GetInstance()->RegisterOutput({ "SurfaceRender", FName("FBox") },
@@ -401,8 +525,9 @@ namespace Dataflow
 		/**
 		* DataflowNode (FFieldCollection) Rendering
 		*
+		*		@param Name : "VolumeRender"
 		*		@param Type : FName("FFieldCollection")
-
+		*
 		*		@param Outputs : {FFieldCollection : "VectorField"}
 		*/
 		FRenderingFactory::GetInstance()->RegisterOutput({ "VolumeRender", FFieldCollection::StaticType() },
diff --git a/Engine/Plugins/Experimental/Dataflow/Source/DataflowEnginePlugin/Public/Dataflow/DataflowConnectionTypes.h b/Engine/Plugins/Experimental/Dataflow/Source/DataflowEnginePlugin/Public/Dataflow/DataflowConnectionTypes.h
new file mode 100644
--- /dev/null
+++ b/Engine/Plugins/Experimental/Dataflow/Source/DataflowEnginePlugin/Public/Dataflow/DataflowConnectionTypes.h
@@ -0,0 +1,24 @@
+// Copyright Epic Games, Inc. All Rights Reserved.
+#pragma once
+
+#include "CoreMinimal.h"
+
+#include "DataflowConnectionTypes.generated.h"
+
+
+USTRUCT(BlueprintType)
+struct DATAFLOWENGINEPLUGIN_API FCollectionAttributeKey
+{
+	GENERATED_USTRUCT_BODY()
+public:
+	FCollectionAttributeKey() : Attribute(""), Group("") {}
+	FCollectionAttributeKey(FString InAttribute, FString InGroup)
+		: Attribute(InAttribute), Group(InGroup) {}
+	TPair<FName, FName> GetNamedKey() { return { FName(Attribute), FName(Group) }; }
+
+	UPROPERTY(EditAnywhere, Category = "Collection Key")
+	FString Attribute;
+
+	UPROPERTY(EditAnywhere, Category = "Collection Key")
+	FString Group;
+};
diff --git a/Engine/Plugins/Experimental/Dataflow/Source/DataflowNodes/DataflowNodes.Build.cs b/Engine/Plugins/Experimental/Dataflow/Source/DataflowNodes/DataflowNodes.Build.cs
--- a/Engine/Plugins/Experimental/Dataflow/Source/DataflowNodes/DataflowNodes.Build.cs
+++ b/Engine/Plugins/Experimental/Dataflow/Source/DataflowNodes/DataflowNodes.Build.cs
@@ -15,6 +15,7 @@ namespace UnrealBuildTool.Rules
 					"Chaos",
 					"DataflowCore",
 					"DataflowEngine",
+					"DataflowEnginePlugin",
 					"Engine"
 				}
 			);
diff --git a/Engine/Plugins/Experimental/Dataflow/Source/DataflowNodes/Public/Dataflow/DataflowCollectionAttributeKeyNodes.h b/Engine/Plugins/Experimental/Dataflow/Source/DataflowNodes/Public/Dataflow/DataflowCollectionAttributeKeyNodes.h
--- a/Engine/Plugins/Experimental/Dataflow/Source/DataflowNodes/Public/Dataflow/DataflowCollectionAttributeKeyNodes.h
+++ b/Engine/Plugins/Experimental/Dataflow/Source/DataflowNodes/Public/Dataflow/DataflowCollectionAttributeKeyNodes.h
@@ -4,31 +4,13 @@
 
 #include "CoreMinimal.h"
 #include "Dataflow/DataflowEngine.h"
+#include "Dataflow/DataflowConnectionTypes.h"
 
 #include "DataflowCollectionAttributeKeyNodes.generated.h"
 
 class USkeletalMesh;
 
 
-USTRUCT(BlueprintType)
-struct DATAFLOWNODES_API FCollectionAttributeKey
-{
-	GENERATED_USTRUCT_BODY()
-public:
-	FCollectionAttributeKey() : Attribute(""), Group("") {}
-	FCollectionAttributeKey(FString InAttribute, FString InGroup)
-		: Attribute(InAttribute), Group(InGroup) {}
-	TPair<FName, FName> GetNamedKey() { return { FName(Attribute), FName(Group)}; }
-
-	UPROPERTY(EditAnywhere, Category = "Collection Key")
-	FString Attribute;
-
-	UPROPERTY(EditAnywhere, Category = "Collection Key")
-	FString Group;
-
-};
-
-
 USTRUCT(meta = (Dataflow))
 struct FMakeAttributeKeyDataflowNode : public FDataflowNode
 {
diff --git a/Engine/Plugins/Experimental/GeometryCollectionPlugin/Source/GeometryCollectionNodes/GeometryCollectionNodes.Build.cs b/Engine/Plugins/Experimental/GeometryCollectionPlugin/Source/GeometryCollectionNodes/GeometryCollectionNodes.Build.cs
--- a/Engine/Plugins/Experimental/GeometryCollectionPlugin/Source/GeometryCollectionNodes/GeometryCollectionNodes.Build.cs
+++ b/Engine/Plugins/Experimental/GeometryCollectionPlugin/Source/GeometryCollectionNodes/GeometryCollectionNodes.Build.cs
@@ -15,6 +15,7 @@ namespace UnrealBuildTool.Rules
 					"Chaos",
 					"DataflowCore",
 					"DataflowEngine",
+					"DataflowEnginePlugin",
 					"DataflowNodes",
 					"GeometryCollectionEngine",
 					"GeometryCore",
diff --git a/Engine/Plugins/Experimental/GeometryCollectionPlugin/Source/GeometryCollectionNodes/Public/Dataflow/GeometryCollectionTransferVertexScalarAttributeNode.h b/Engine/Plugins/Experimental/GeometryCollectionPlugin/Source/GeometryCollectionNodes/Public/Dataflow/GeometryCollectionTransferVertexScalarAttributeNode.h
--- a/Engine/Plugins/Experimental/GeometryCollectionPlugin/Source/GeometryCollectionNodes/Public/Dataflow/GeometryCollectionTransferVertexScalarAttributeNode.h
+++ b/Engine/Plugins/Experimental/GeometryCollectionPlugin/Source/GeometryCollectionNodes/Public/Dataflow/GeometryCollectionTransferVertexScalarAttributeNode.h
@@ -73,6 +73,7 @@ struct FGeometryCollectionTransferVertexScalarAttributeNode : public FDataflowNo
 {
 	GENERATED_USTRUCT_BODY()
 	DATAFLOW_NODE_DEFINE_INTERNAL(FGeometryCollectionTransferVertexScalarAttributeNode, "TransferVertexScalarAttribute", "GeometryCollection", "Transfer a named scalar attribute from the Source Collection to the Target Collection")
+	DATAFLOW_NODE_RENDER_TYPE("SurfaceWeightsRender", FGeometryCollection::StaticType(), "Collection", "AttributeKey")
 
 public:
 
diff --git a/Engine/Source/Runtime/Experimental/Dataflow/Core/Public/Dataflow/DataflowNode.h b/Engine/Source/Runtime/Experimental/Dataflow/Core/Public/Dataflow/DataflowNode.h
--- a/Engine/Source/Runtime/Experimental/Dataflow/Core/Public/Dataflow/DataflowNode.h
+++ b/Engine/Source/Runtime/Experimental/Dataflow/Core/Public/Dataflow/DataflowNode.h
@@ -584,10 +584,10 @@ private:
 #define DATAFLOW_NODE_REGISTER_CREATION_FACTORY(A)									\
 	::Dataflow::FNodeFactory::RegisterNodeFromType<A>();
 
-#define DATAFLOW_NODE_RENDER_TYPE(A, B, C)												\
+#define DATAFLOW_NODE_RENDER_TYPE(A, B, ...)												\
 	virtual TArray<::Dataflow::FRenderingParameter> GetRenderParametersImpl() const {		\
 		TArray<::Dataflow::FRenderingParameter> Array;								\
-		Array.Add({ A, B, {C,} });														\
+		Array.Add({ A, B, {__VA_ARGS__,} });														\
 		return Array;}
 
 #define DATAFLOW_NODE_DEFINE_INTERNAL(TYPE, DISPLAY_NAME, CATEGORY, TAGS)			\
