Dataflow : Editor SNode Interface to connect the NodeFactory to the SNode Construction

diff --git a/Engine/Plugins/Experimental/Dataflow/Source/DataflowEditor/Private/Dataflow/DataflowEditorModule.cpp b/Engine/Plugins/Experimental/Dataflow/Source/DataflowEditor/Private/Dataflow/DataflowEditorModule.cpp
--- a/Engine/Plugins/Experimental/Dataflow/Source/DataflowEditor/Private/Dataflow/DataflowEditorModule.cpp
+++ b/Engine/Plugins/Experimental/Dataflow/Source/DataflowEditor/Private/Dataflow/DataflowEditorModule.cpp
@@ -1,6 +1,8 @@
 // Copyright Epic Games, Inc. All Rights Reserved.
 
 #include "Dataflow/DataflowEditorModule.h"
+
+#include "Dataflow/DataflowGraphEditor.h"
 #include "Dataflow/DataflowEditorStyle.h"
 #include "Dataflow/DataflowEditorMode.h"
 #include "Dataflow/DataflowEditorToolkit.h"
@@ -19,17 +21,10 @@ void FDataflowEditorModule::StartupModule()
 	FDataflowEditorStyle::Get();
 	
 	FPropertyEditorModule& PropertyEditorModule = FModuleManager::LoadModuleChecked<FPropertyEditorModule>(TEXT("PropertyEditor"));
-	
-	DataflowSNodeFactory = MakeShareable(new FDataflowSNodeFactory());
-	FEdGraphUtilities::RegisterVisualNodeFactory(DataflowSNodeFactory);
-}
+	}
 
 void FDataflowEditorModule::ShutdownModule()
 {	
-	if (UObjectInitialized())
-	{
-		FEdGraphUtilities::UnregisterVisualNodeFactory(DataflowSNodeFactory);
-	}
 	FEditorModeRegistry::Get().UnregisterMode(UDataflowEditorMode::EM_DataflowEditorModeId);
 }
 
diff --git a/Engine/Plugins/Experimental/Dataflow/Source/DataflowEditor/Private/Dataflow/DataflowEditorToolkit.cpp b/Engine/Plugins/Experimental/Dataflow/Source/DataflowEditor/Private/Dataflow/DataflowEditorToolkit.cpp
--- a/Engine/Plugins/Experimental/Dataflow/Source/DataflowEditor/Private/Dataflow/DataflowEditorToolkit.cpp
+++ b/Engine/Plugins/Experimental/Dataflow/Source/DataflowEditor/Private/Dataflow/DataflowEditorToolkit.cpp
@@ -859,7 +859,8 @@ TSharedRef<SDataflowGraphEditor> FDataflowEditorToolkit::CreateGraphEditorWidget
 		.GraphToEdit(DataflowToEdit)
 		.GraphEvents(InEvents)
 		.DetailsView(InNodeDetailsEditor)
-		.EvaluateGraph(Evaluate);
+		.EvaluateGraph(Evaluate)
+		.DataflowEditor(DataflowEditor);
 
 	OnSelectionChangedMulticastDelegateHandle = NewGraphEditor->OnSelectionChangedMulticast.AddSP(this, &FDataflowEditorToolkit::OnNodeSelectionChanged);
 	OnNodeDeletedMulticastDelegateHandle = NewGraphEditor->OnNodeDeletedMulticast.AddSP(this, &FDataflowEditorToolkit::OnNodeDeleted);
diff --git a/Engine/Plugins/Experimental/Dataflow/Source/DataflowEditor/Private/Dataflow/DataflowGraphEditor.cpp b/Engine/Plugins/Experimental/Dataflow/Source/DataflowEditor/Private/Dataflow/DataflowGraphEditor.cpp
--- a/Engine/Plugins/Experimental/Dataflow/Source/DataflowEditor/Private/Dataflow/DataflowGraphEditor.cpp
+++ b/Engine/Plugins/Experimental/Dataflow/Source/DataflowEditor/Private/Dataflow/DataflowGraphEditor.cpp
@@ -3,8 +3,10 @@
 #include "Dataflow/DataflowGraphEditor.h"
 
 #include "BoneDragDropOp.h"
+#include "Dataflow/DataflowEditor.h"
 #include "Dataflow/DataflowEditorToolkit.h"
 #include "Dataflow/DataflowEngine.h"
+#include "Dataflow/DataflowSEditorInterface.h"
 #include "Dataflow/DataflowSNodeFactories.h"
 #include "Dataflow/DataflowSCommentNode.h"
 #include "Dataflow/DataflowNodeParameters.h"
@@ -19,6 +21,8 @@
 
 #define LOCTEXT_NAMESPACE "DataflowGraphEditor"
 
+TSharedPtr<FDataflowGraphEditorNodeFactory> SDataflowGraphEditor::NodeFactory;
+
 void SDataflowGraphEditor::Construct(const FArguments& InArgs, UObject* InAssetOwner)
 {
 	check(InArgs._GraphToEdit);
@@ -27,6 +31,7 @@ void SDataflowGraphEditor::Construct(const FArguments& InArgs, UObject* InAssetO
 	DetailsView = InArgs._DetailsView;
 	EvaluateGraphCallback = InArgs._EvaluateGraph;
 	OnDragDropEventCallback = InArgs._OnDragDropEvent;
+	DataflowEditor = InArgs._DataflowEditor;
 
 	FGraphAppearanceInfo AppearanceInfo;
 	AppearanceInfo.CornerText = FText::FromString("Dataflow");
@@ -152,8 +157,20 @@ void SDataflowGraphEditor::Construct(const FArguments& InArgs, UObject* InAssetO
 
 	SGraphEditor::Construct(Arguments);
 
+	SetNodeFactory( MakeShared<FDataflowGraphNodeFactory>(this) );
 }
 
+const TSharedPtr<Dataflow::FEngineContext> SDataflowGraphEditor::GetDataflowContext() const
+{
+	if (DataflowEditor)
+	{
+		if (DataflowEditor->GetEditorContent())
+		{
+			return DataflowEditor->GetEditorContent()->GetDataflowContext();
+		}
+	}
+	return TSharedPtr<Dataflow::FEngineContext>(nullptr);
+}
 
 void SDataflowGraphEditor::EvaluateNode()
 {
diff --git a/Engine/Plugins/Experimental/Dataflow/Source/DataflowEditor/Private/Dataflow/DataflowSEditorInterface.cpp b/Engine/Plugins/Experimental/Dataflow/Source/DataflowEditor/Private/Dataflow/DataflowSEditorInterface.cpp
new file mode 100644
--- /dev/null
+++ b/Engine/Plugins/Experimental/Dataflow/Source/DataflowEditor/Private/Dataflow/DataflowSEditorInterface.cpp
@@ -0,0 +1,4 @@
+// Copyright Epic Games, Inc. All Rights Reserved.
+
+#include "Dataflow/DataflowSEditorInterface.h"
+#include "Dataflow/DataflowEditor.h"
diff --git a/Engine/Plugins/Experimental/Dataflow/Source/DataflowEditor/Private/Dataflow/DataflowSNode.cpp b/Engine/Plugins/Experimental/Dataflow/Source/DataflowEditor/Private/Dataflow/DataflowSNode.cpp
--- a/Engine/Plugins/Experimental/Dataflow/Source/DataflowEditor/Private/Dataflow/DataflowSNode.cpp
+++ b/Engine/Plugins/Experimental/Dataflow/Source/DataflowEditor/Private/Dataflow/DataflowSNode.cpp
@@ -31,6 +31,7 @@ void SDataflowEdNode::Construct(const FArguments& InArgs, UDataflowEdNode* InNod
 {
 	GraphNode = InNode;
 	DataflowGraphNode = Cast<UDataflowEdNode>(InNode);
+	DataflowInterface = InArgs._DataflowInterface;
 	UpdateGraphNode();
 
 	
diff --git a/Engine/Plugins/Experimental/Dataflow/Source/DataflowEditor/Private/Dataflow/DataflowSNodeFactories.cpp b/Engine/Plugins/Experimental/Dataflow/Source/DataflowEditor/Private/Dataflow/DataflowSNodeFactories.cpp
--- a/Engine/Plugins/Experimental/Dataflow/Source/DataflowEditor/Private/Dataflow/DataflowSNodeFactories.cpp
+++ b/Engine/Plugins/Experimental/Dataflow/Source/DataflowEditor/Private/Dataflow/DataflowSNodeFactories.cpp
@@ -3,13 +3,15 @@
 #include "Dataflow/DataflowSNodeFactories.h"
 #include "Dataflow/DataflowSNode.h"
 #include "Dataflow/DataflowEdNode.h"
+#include "Dataflow/DataflowEditor.h"
 #include "Dataflow/DataflowCoreNodes.h"
 #include "EdGraphNode_Comment.h"
 #include "Dataflow/DataflowSCommentNode.h"
 #include "Dataflow/DataflowSchema.h"
 #include "SGraphNodeKnot.h"
 
-TSharedPtr<class SGraphNode> FDataflowSNodeFactory::CreateNode(UEdGraphNode* InNode) const
+
+TSharedPtr<SGraphNode> FDataflowGraphNodeFactory::CreateNodeWidget(UEdGraphNode* InNode)
 {
 	if (UDataflowEdNode* Node = Cast<UDataflowEdNode>(InNode))
 	{
@@ -20,8 +22,9 @@ TSharedPtr<class SGraphNode> FDataflowSNodeFactory::CreateNode(UEdGraphNode* InN
 				return SNew(SGraphNodeKnot, Node);
 			}
 		}
-		return SNew(SDataflowEdNode, Node);
-	}	
+		return SNew(SDataflowEdNode, Node)
+			.DataflowInterface(DataflowInterface);
+	}
 	else if (UEdGraphNode_Comment* CommentNode = Cast<UEdGraphNode_Comment>(InNode))
 	{
 		if (CommentNode->GetSchema()->IsA(UDataflowSchema::StaticClass()))
diff --git a/Engine/Plugins/Experimental/Dataflow/Source/DataflowEditor/Public/Dataflow/DataflowGraphEditor.h b/Engine/Plugins/Experimental/Dataflow/Source/DataflowEditor/Public/Dataflow/DataflowGraphEditor.h
--- a/Engine/Plugins/Experimental/Dataflow/Source/DataflowEditor/Public/Dataflow/DataflowGraphEditor.h
+++ b/Engine/Plugins/Experimental/Dataflow/Source/DataflowEditor/Public/Dataflow/DataflowGraphEditor.h
@@ -5,6 +5,7 @@
 #include "CoreMinimal.h"
 
 #include "Dataflow/DataflowEditorCommands.h"
+#include "Dataflow/DataflowSEditorInterface.h"
 #include "EdGraphUtilities.h"
 #include "Framework/Commands/GenericCommands.h"
 #include "Framework/Commands/UICommandList.h"
@@ -19,8 +20,10 @@
 
 
 class FDataflowEditorToolkit;
+class FDataflowGraphEditorNodeFactory;
 class FDataflowSNodeFactory;
 class UDataflow;
+class UDataflowEditor;
 struct FDataflowConnection;
 namespace Dataflow {
 	class FContext;
@@ -31,7 +34,7 @@ namespace Dataflow {
  * 
  * see(SDataprepGraphEditor for reference)
  */
-class DATAFLOWEDITOR_API SDataflowGraphEditor : public SGraphEditor, public FGCObject
+class DATAFLOWEDITOR_API SDataflowGraphEditor : public SGraphEditor, public FGCObject, public FDataflowSEditorInterface
 {
 public:
 
@@ -47,6 +50,7 @@ public:
 	SLATE_ARGUMENT(TSharedPtr<IStructureDetailsView>, DetailsView)
 	SLATE_ARGUMENT(FDataflowEditorCommands::FGraphEvaluationCallback, EvaluateGraph)
 	SLATE_ARGUMENT(FDataflowEditorCommands::FOnDragDropEventCallback, OnDragDropEvent)
+	SLATE_ARGUMENT(UDataflowEditor*, DataflowEditor)
 	SLATE_END_ARGS()
 
 	// This delegate exists in SGraphEditor but it is not multicast, and we are going to bind it to OnSelectedNodesChanged().
@@ -137,6 +141,9 @@ public:
 	virtual void AddReferencedObjects(FReferenceCollector& Collector) override;
 	virtual FString GetReferencerName() const override { return TEXT("SDataflowGraphEditor"); }
 
+	/** FDataflowSNodeInterface */
+	virtual const TSharedPtr<Dataflow::FEngineContext> GetDataflowContext() const override;
+
 private:
 	/** Add an additional option pin to all selected Dataflow nodes for those that overrides the AddPin function. */
 	void OnAddOptionPin();
@@ -162,12 +169,15 @@ private:
 	/** Command list associated with this graph editor */
 	TSharedPtr<FUICommandList> GraphEditorCommands;
 
-	/** Factory to create the associated SGraphNode classes for Dataflow graph's UEdGraph classes */
-	static TSharedPtr<FDataflowSNodeFactory> NodeFactory;
-
 	/** The details view that responds to this widget. */
 	TSharedPtr<IStructureDetailsView> DetailsView;
 
+	/** Factory to create the associated SGraphNode classes for Dataprep graph's UEdGraph classes */
+	static TSharedPtr<FDataflowGraphEditorNodeFactory> NodeFactory;
+
+	/** Editor for the content */
+	UDataflowEditor* DataflowEditor = nullptr;
+
 	bool VKeyDown = false;
 	bool LeftControlKeyDown = false;
 	bool RightControlKeyDown = false;
diff --git a/Engine/Plugins/Experimental/Dataflow/Source/DataflowEditor/Public/Dataflow/DataflowSEditorInterface.h b/Engine/Plugins/Experimental/Dataflow/Source/DataflowEditor/Public/Dataflow/DataflowSEditorInterface.h
new file mode 100644
--- /dev/null
+++ b/Engine/Plugins/Experimental/Dataflow/Source/DataflowEditor/Public/Dataflow/DataflowSEditorInterface.h
@@ -0,0 +1,25 @@
+// Copyright Epic Games, Inc. All Rights Reserved.
+
+#pragma once
+
+#include "CoreMinimal.h"
+#include "Dataflow/DataflowObjectInterface.h"
+
+/**
+* FDataflowSEditorInterface
+* 
+*/
+
+class UDataflowBaseContent;
+class UDataflowEditor;
+
+class FDataflowSEditorInterface
+{
+public:
+	FDataflowSEditorInterface() {}
+
+	/** Dataflow editor content accessors */
+	virtual const TSharedPtr<Dataflow::FEngineContext> GetDataflowContext() const { return TSharedPtr<Dataflow::FEngineContext>(nullptr); };
+
+};
+
diff --git a/Engine/Plugins/Experimental/Dataflow/Source/DataflowEditor/Public/Dataflow/DataflowSNode.h b/Engine/Plugins/Experimental/Dataflow/Source/DataflowEditor/Public/Dataflow/DataflowSNode.h
--- a/Engine/Plugins/Experimental/Dataflow/Source/DataflowEditor/Public/Dataflow/DataflowSNode.h
+++ b/Engine/Plugins/Experimental/Dataflow/Source/DataflowEditor/Public/Dataflow/DataflowSNode.h
@@ -6,6 +6,7 @@
 #include "SGraphNode.h"
 #include "UObject/GCObject.h"
 #include "Dataflow/DataflowCore.h"
+#include "Dataflow/DataflowSEditorInterface.h"
 
 #include "DataflowSNode.generated.h"
 
@@ -23,11 +24,11 @@ class DATAFLOWEDITOR_API SDataflowEdNode : public SGraphNode , public FGCObject
 public:
 	typedef TFunction<void(UEdGraphNode* InNode, bool InEnabled)> FToggleRenderCallback;
 
-
 	SLATE_BEGIN_ARGS(SDataflowEdNode)
 		: _GraphNodeObj(nullptr)
 	{}
 	SLATE_ARGUMENT(UDataflowEdNode*, GraphNodeObj)
+	SLATE_ARGUMENT(FDataflowSEditorInterface*, DataflowInterface)
 	SLATE_END_ARGS()
 	
 	void Construct(const FArguments& InArgs, UDataflowEdNode* InNode);
@@ -66,6 +67,9 @@ private:
 
 	//FCheckBoxStyle CacheStatusStyle;
 	//TSharedPtr<SCheckBox> CacheStatus;
+
+	FDataflowSEditorInterface* DataflowInterface;
+
 };
 
 //
diff --git a/Engine/Plugins/Experimental/Dataflow/Source/DataflowEditor/Public/Dataflow/DataflowSNodeFactories.h b/Engine/Plugins/Experimental/Dataflow/Source/DataflowEditor/Public/Dataflow/DataflowSNodeFactories.h
--- a/Engine/Plugins/Experimental/Dataflow/Source/DataflowEditor/Public/Dataflow/DataflowSNodeFactories.h
+++ b/Engine/Plugins/Experimental/Dataflow/Source/DataflowEditor/Public/Dataflow/DataflowSNodeFactories.h
@@ -3,11 +3,25 @@
 #pragma once
 
 #include "EdGraphUtilities.h"
-#include "SGraphNode.h"
 #include "EdGraph/EdGraphNode.h"
+#include "Dataflow/DataflowSEditorInterface.h"
+#include "NodeFactory.h"
+#include "SGraphNode.h"
+
+class UDataflowEditor;
 
-class DATAFLOWEDITOR_API FDataflowSNodeFactory : public FGraphPanelNodeFactory
+class FDataflowGraphNodeFactory : public FGraphNodeFactory, public TSharedFromThis<FDataflowGraphNodeFactory>
 {
-	virtual TSharedPtr<class SGraphNode> CreateNode(class UEdGraphNode* InNode) const override;
+public:
+	virtual ~FDataflowGraphNodeFactory() = default;
+
+	FDataflowGraphNodeFactory(FDataflowSEditorInterface* InDataflowInterface)
+		: DataflowInterface(InDataflowInterface)
+	{}
+
+	/** Create a widget for the supplied node */
+	virtual TSharedPtr<SGraphNode> CreateNodeWidget(UEdGraphNode* InNode) override;
 
-};
+private:
+	FDataflowSEditorInterface* DataflowInterface;
+};
\ No newline at end of file
diff --git a/Engine/Source/Runtime/Experimental/Dataflow/Engine/Public/Dataflow/DataflowContent.h b/Engine/Source/Runtime/Experimental/Dataflow/Engine/Public/Dataflow/DataflowContent.h
--- a/Engine/Source/Runtime/Experimental/Dataflow/Engine/Public/Dataflow/DataflowContent.h
+++ b/Engine/Source/Runtime/Experimental/Dataflow/Engine/Public/Dataflow/DataflowContent.h
@@ -155,6 +155,7 @@ public:
 	*   Dataflow context stores the evaluated state of the graph. 
 	*/
 	void SetDataflowContext(const TSharedPtr<Dataflow::FEngineContext>& InContext);
+	const TSharedPtr<Dataflow::FEngineContext>& GetDataflowContext() const { return DataflowContext; }
 	TSharedPtr<Dataflow::FEngineContext>& GetDataflowContext() { return DataflowContext; }
 
 	/** Rebuild the owner dependent datas  */
