Dataflow Construction Scene: explicitly destroy DynamicMeshComponents when they are removed from the scene. This fixes a crash in Dataflow Editor

diff --git a/Engine/Plugins/Experimental/Dataflow/Source/DataflowEditor/Private/Dataflow/DataflowConstructionScene.cpp b/Engine/Plugins/Experimental/Dataflow/Source/DataflowEditor/Private/Dataflow/DataflowConstructionScene.cpp
--- a/Engine/Plugins/Experimental/Dataflow/Source/DataflowEditor/Private/Dataflow/DataflowConstructionScene.cpp
+++ b/Engine/Plugins/Experimental/Dataflow/Source/DataflowEditor/Private/Dataflow/DataflowConstructionScene.cpp
@@ -230,13 +230,14 @@ void FDataflowConstructionScene::ResetDynamicMeshComponents()
 			DynamicMeshComponent->PushSelectionToProxy();
 		}
 		RemoveComponent(DynamicMeshComponent);
+		DynamicMeshComponent->DestroyComponent();
 	}
 	DynamicMeshComponents.Reset();
 }
 
 TObjectPtr<UDynamicMeshComponent>& FDataflowConstructionScene::AddDynamicMeshComponent(FDataflowRenderKey InKey, const FString& MeshName, UE::Geometry::FDynamicMesh3&& DynamicMesh, const TArray<UMaterialInterface*>& MaterialSet)
 {
-	// Dont use the MakeUniqueObjectName for the component, we need to keep the name alinged with the collection so selection will work in 
+	// Dont use the MakeUniqueObjectName for the component, we need to keep the name aligned with the collection so selection will work in 
 	// other editors. 
 	// const FName UniqueObjectName = MakeUniqueObjectName(RootSceneActor, UDataflowEditorCollectionComponent::StaticClass(), FName(MeshName));
 	TObjectPtr<UDataflowEditorCollectionComponent> DynamicMeshComponent = NewObject<UDataflowEditorCollectionComponent>(RootSceneActor, FName(MeshName));
