Dataflow Editor: add customizable View Modes to the construction view. We provide default 3D and 2D views, but allow other modules to define their own as well.

Also:
- Implement rendering of FClothCollection output
- Replace the default Construction View toolbar with one like the Cloth Editor

diff --git a/Engine/Plugins/ChaosClothAssetEditor/Source/ChaosClothAssetDataflowNodes/Private/ChaosClothAsset/ChaosClothAssetDataflowNodesModule.cpp b/Engine/Plugins/ChaosClothAssetEditor/Source/ChaosClothAssetDataflowNodes/Private/ChaosClothAsset/ChaosClothAssetDataflowNodesModule.cpp
--- a/Engine/Plugins/ChaosClothAssetEditor/Source/ChaosClothAssetDataflowNodes/Private/ChaosClothAsset/ChaosClothAssetDataflowNodesModule.cpp
+++ b/Engine/Plugins/ChaosClothAssetEditor/Source/ChaosClothAssetDataflowNodes/Private/ChaosClothAsset/ChaosClothAssetDataflowNodesModule.cpp
@@ -65,8 +65,10 @@
 #include "ChaosClothAsset/WeightMapToSelectionNode.h"
 #include "ChaosClothAsset/ConnectableValueCustomization.h"
 #include "ChaosClothAsset/ConnectableValue.h"
+#include "ChaosClothAsset/ClothDataflowViewModes.h"
 #include "Dataflow/DataflowNodeColorsRegistry.h"
 #include "Dataflow/DataflowNodeFactory.h"
+#include "Dataflow/DataflowRenderingFactory.h"
 #include "Modules/ModuleManager.h"
 
 namespace UE::Chaos::ClothAsset
@@ -139,6 +141,118 @@ namespace UE::Chaos::ClothAsset
 			DATAFLOW_NODE_REGISTER_CREATION_FACTORY(FChaosClothAssetWeightMapToSelectionNode);
 			DATAFLOW_NODE_REGISTER_CREATION_FACTORY_NODE_COLORS_BY_CATEGORY("Cloth", FColorScheme::NodeHeader, FColorScheme::NodeBody);
 		}
+
+
+		class FClothSurfaceRenderCallbacks : public Dataflow::FRenderingFactory::ICallbackInterface
+		{
+		public:
+
+			static Dataflow::FRenderKey RenderKey;
+
+		private:
+
+			virtual Dataflow::FRenderKey GetRenderKey() const override
+			{
+				return RenderKey;
+			}
+
+			virtual bool CanRender(const Dataflow::IDataflowConstructionViewMode& ViewMode) const override
+			{
+				const FName& ViewModeName = ViewMode.GetName();
+				return (ViewModeName == FCloth2DSimViewMode::Name ||
+						ViewModeName == FCloth3DSimViewMode::Name ||
+						ViewModeName == FClothRenderViewMode::Name);
+			}
+
+			virtual void Render(GeometryCollection::Facades::FRenderingFacade& RenderCollection, const Dataflow::FGraphRenderingState& State) override
+			{
+				if (State.GetRenderOutputs().Num())
+				{
+					const FManagedArrayCollection Default;
+					checkf(State.GetRenderOutputs().Num() == 1, TEXT("Expected FGraphRenderingState object to have one render output"));
+					const FName PrimaryOutput = State.GetRenderOutputs()[0];
+					FManagedArrayCollection Collection = State.GetValue<FManagedArrayCollection>(PrimaryOutput, Default);
+
+					const TSharedRef<const FManagedArrayCollection> ClothCollection = MakeShared<FManagedArrayCollection>(MoveTemp(Collection));
+					const FCollectionClothConstFacade ClothFacade(ClothCollection);
+					if (!ClothFacade.IsValid())
+					{
+						// Cloth collection may not be valid for all nodes
+						return;
+					}
+
+					TArray<FVector3f> Vertices;
+					TArray<FVector3f> Normals;
+					TArray<FIntVector> Indices;
+
+					if (State.GetViewMode().GetName() == FCloth3DSimViewMode::Name)
+					{
+						Vertices = ClothFacade.GetSimPosition3D();
+						Indices = ClothFacade.GetSimIndices3D();
+						Normals = ClothFacade.GetSimNormal();
+					}
+					else if (State.GetViewMode().GetName() == FCloth2DSimViewMode::Name)
+					{
+						//
+						// Flip the Y coordinate to get the desired visualization with our LVT_OrthoXY viewport
+						//
+						float MinY = UE_BIG_NUMBER;
+						float MaxY = -UE_BIG_NUMBER;
+						for (const FVector2f& Vertex2D : ClothFacade.GetSimPosition2D())
+						{
+							MinY = FMath::Min(MinY, Vertex2D[1]);
+							MaxY = FMath::Max(MaxY, Vertex2D[1]);
+						}
+						for (const FVector2f& Vertex2D : ClothFacade.GetSimPosition2D())
+						{
+							Vertices.Add({ Vertex2D[0], MinY + (MaxY - Vertex2D[1]), 0.0 });
+						}
+
+						Indices = ClothFacade.GetSimIndices2D();
+						Normals = ClothFacade.GetSimNormal();		// TODO: Should we hard code a 2D normal?
+					}
+					else if (State.GetViewMode().GetName() == FClothRenderViewMode::Name)
+					{
+						Vertices = ClothFacade.GetRenderPosition();
+						Indices = ClothFacade.GetRenderIndices();
+						Normals = ClothFacade.GetRenderNormal();
+						// TODO: Get materials from the render mesh
+					}
+					else
+					{
+						checkf(false, TEXT("Invalid View Mode for FClothCollection rendering"));
+					}
+
+					TArray<FLinearColor> Colors;
+					Colors.Init(FLinearColor::Gray, Vertices.Num());	// TODO: Choose a vertex color
+
+					const int32 GeometryIndex = RenderCollection.StartGeometryGroup(State.GetGuid().ToString());
+					RenderCollection.AddSurface(MoveTemp(Vertices), MoveTemp(Indices), MoveTemp(Normals), MoveTemp(Colors));
+					RenderCollection.EndGeometryGroup(GeometryIndex);
+				}
+			}
+		};
+
+		Dataflow::FRenderKey FClothSurfaceRenderCallbacks::RenderKey = { "SurfaceRender", FName("FClothCollection") };
+
+		static void RegisterRenderingCallbacks()
+		{
+			Dataflow::FRenderingViewModeFactory::GetInstance().RegisterViewMode(MakeUnique<FCloth2DSimViewMode>());
+			Dataflow::FRenderingViewModeFactory::GetInstance().RegisterViewMode(MakeUnique<FCloth3DSimViewMode>());
+			Dataflow::FRenderingViewModeFactory::GetInstance().RegisterViewMode(MakeUnique<FClothRenderViewMode>());
+
+			Dataflow::FRenderingFactory::GetInstance()->RegisterCallbacks(MakeUnique<FClothSurfaceRenderCallbacks>());
+		}
+
+		static void DeregisterRenderingCallbacks()
+		{
+			Dataflow::FRenderingFactory::GetInstance()->DeregisterCallbacks(FClothSurfaceRenderCallbacks::RenderKey);
+
+			Dataflow::FRenderingViewModeFactory::GetInstance().DeregisterViewMode(FCloth2DSimViewMode::Name);
+			Dataflow::FRenderingViewModeFactory::GetInstance().DeregisterViewMode(FCloth3DSimViewMode::Name);
+			Dataflow::FRenderingViewModeFactory::GetInstance().DeregisterViewMode(FClothRenderViewMode::Name);
+		}
+
 	}  // End namespace Private
 
 	class FChaosClothAssetDataflowNodesModule : public IModuleInterface
@@ -169,10 +283,14 @@ namespace UE::Chaos::ClothAsset
 				PropertyModule->RegisterCustomPropertyTypeLayout(FChaosClothAssetImportedFloatValue::StaticStruct()->GetFName(), FOnGetPropertyTypeCustomizationInstance::CreateStatic(&FImportedValueCustomization::MakeInstance));
 				PropertyModule->RegisterCustomPropertyTypeLayout(FChaosClothAssetImportedIntValue::StaticStruct()->GetFName(), FOnGetPropertyTypeCustomizationInstance::CreateStatic(&FImportedValueCustomization::MakeInstance));
 			}
+
+			Private::RegisterRenderingCallbacks();
 		}
 
 		virtual void ShutdownModule() override
 		{
+			Private::DeregisterRenderingCallbacks();
+
 			// Unregister type customizations
 			if (UObjectInitialized() && !IsEngineExitRequested())
 			{
diff --git a/Engine/Plugins/ChaosClothAssetEditor/Source/ChaosClothAssetDataflowNodes/Private/ChaosClothAsset/ClothDataflowViewModes.cpp b/Engine/Plugins/ChaosClothAssetEditor/Source/ChaosClothAssetDataflowNodes/Private/ChaosClothAsset/ClothDataflowViewModes.cpp
new file mode 100644
--- /dev/null
+++ b/Engine/Plugins/ChaosClothAssetEditor/Source/ChaosClothAssetDataflowNodes/Private/ChaosClothAsset/ClothDataflowViewModes.cpp
@@ -0,0 +1,62 @@
+// Copyright Epic Games, Inc. All Rights Reserved.
+
+#include "ChaosClothAsset/ClothDataflowViewModes.h"
+
+#define LOCTEXT_NAMESPACE "ChaosClothAssetDataflowViewModes"
+
+namespace UE::Chaos::ClothAsset
+{
+	const FName FCloth2DSimViewMode::Name = FName("Cloth2DSimView");
+
+	FName FCloth2DSimViewMode::GetName() const
+	{
+		return Name;
+	}
+
+	FText FCloth2DSimViewMode::GetButtonText() const
+	{
+		return FText(LOCTEXT("Cloth2DSimViewButtonText", "2DSim"));
+	}
+
+	FText FCloth2DSimViewMode::GetTooltipText() const
+	{
+		return FText(LOCTEXT("Cloth2DSimViewTooltipText", "2D Simulation Mesh"));
+	}
+
+	const FName FCloth3DSimViewMode::Name = FName("Cloth3DSimView");
+
+	FName FCloth3DSimViewMode::GetName() const
+	{
+		return Name;
+	}
+
+	FText FCloth3DSimViewMode::GetButtonText() const
+	{
+		return FText(LOCTEXT("Cloth3DSimViewButtonText", "3DSim"));
+	}
+
+	FText FCloth3DSimViewMode::GetTooltipText() const
+	{
+		return FText(LOCTEXT("Cloth3DSimViewTooltipText", "3D Simulation Mesh"));
+	}
+
+	const FName FClothRenderViewMode::Name = FName("ClothRenderView");
+
+	FName FClothRenderViewMode::GetName() const
+	{
+		return Name;
+	}
+
+	FText FClothRenderViewMode::GetButtonText() const
+	{
+		return FText(LOCTEXT("ClothRenderViewButtonText", "Render"));
+	}
+
+	FText FClothRenderViewMode::GetTooltipText() const
+	{
+		return FText(LOCTEXT("ClothRenderViewTooltipText", "3D Render Mesh"));
+	}
+
+}  // End namespace UE::Chaos::ClothAsset
+
+#undef LOCTEXT_NAMESPACE 
diff --git a/Engine/Plugins/ChaosClothAssetEditor/Source/ChaosClothAssetDataflowNodes/Private/ChaosClothAsset/ClothDataflowViewModes.h b/Engine/Plugins/ChaosClothAssetEditor/Source/ChaosClothAssetDataflowNodes/Private/ChaosClothAsset/ClothDataflowViewModes.h
new file mode 100644
--- /dev/null
+++ b/Engine/Plugins/ChaosClothAssetEditor/Source/ChaosClothAssetDataflowNodes/Private/ChaosClothAsset/ClothDataflowViewModes.h
@@ -0,0 +1,42 @@
+// Copyright Epic Games, Inc. All Rights Reserved.
+
+#pragma once
+
+#include "Dataflow/DataflowRenderingViewMode.h"
+
+namespace UE::Chaos::ClothAsset
+{
+	class FCloth2DSimViewMode : public Dataflow::FDataflowConstruction2DViewModeBase
+	{
+	public:
+		static const FName Name;
+		virtual ~FCloth2DSimViewMode() = default;
+	private:
+		virtual FName GetName() const override;
+		virtual FText GetButtonText() const override;
+		virtual FText GetTooltipText() const override;
+	};
+
+	class FCloth3DSimViewMode : public Dataflow::FDataflowConstruction3DViewModeBase
+	{
+	public:
+		static const FName Name;
+		virtual ~FCloth3DSimViewMode() = default;
+	private:
+		virtual FName GetName() const override;
+		virtual FText GetButtonText() const override;
+		virtual FText GetTooltipText() const override;
+	};
+
+	class FClothRenderViewMode : public Dataflow::FDataflowConstruction3DViewMode
+	{
+	public:
+		static const FName Name;
+		virtual ~FClothRenderViewMode() = default;
+	private:
+		virtual FName GetName() const override;
+		virtual FText GetButtonText() const override;
+		virtual FText GetTooltipText() const override;
+	};
+
+}  // End namespace UE::Chaos::ClothAsset
diff --git a/Engine/Plugins/ChaosClothAssetEditor/Source/ChaosClothAssetDataflowNodes/Public/ChaosClothAsset/AddWeightMapNode.h b/Engine/Plugins/ChaosClothAssetEditor/Source/ChaosClothAssetDataflowNodes/Public/ChaosClothAsset/AddWeightMapNode.h
--- a/Engine/Plugins/ChaosClothAssetEditor/Source/ChaosClothAssetDataflowNodes/Public/ChaosClothAsset/AddWeightMapNode.h
+++ b/Engine/Plugins/ChaosClothAssetEditor/Source/ChaosClothAssetDataflowNodes/Public/ChaosClothAsset/AddWeightMapNode.h
@@ -48,6 +48,7 @@ struct FChaosClothAssetAddWeightMapNode : public FDataflowTerminalNode
 {
 	GENERATED_USTRUCT_BODY()
 	DATAFLOW_NODE_DEFINE_INTERNAL(FChaosClothAssetAddWeightMapNode, "AddWeightMap", "Cloth", "Cloth Add Weight Map")
+	DATAFLOW_NODE_RENDER_TYPE("SurfaceRender", FName("FClothCollection"), "Collection")
 
 public:
 	static constexpr float ReplaceChangedPassthroughValue = UE_BIG_NUMBER;
diff --git a/Engine/Plugins/ChaosClothAssetEditor/Source/ChaosClothAssetDataflowNodes/Public/ChaosClothAsset/StaticMeshImportNode.h b/Engine/Plugins/ChaosClothAssetEditor/Source/ChaosClothAssetDataflowNodes/Public/ChaosClothAsset/StaticMeshImportNode.h
--- a/Engine/Plugins/ChaosClothAssetEditor/Source/ChaosClothAssetDataflowNodes/Public/ChaosClothAsset/StaticMeshImportNode.h
+++ b/Engine/Plugins/ChaosClothAssetEditor/Source/ChaosClothAssetDataflowNodes/Public/ChaosClothAsset/StaticMeshImportNode.h
@@ -68,6 +68,7 @@ struct FChaosClothAssetStaticMeshImportNode : public FDataflowNode
 {
 	GENERATED_USTRUCT_BODY()
 	DATAFLOW_NODE_DEFINE_INTERNAL(FChaosClothAssetStaticMeshImportNode, "StaticMeshImport", "Cloth", "Cloth Static Mesh Import")
+	DATAFLOW_NODE_RENDER_TYPE("SurfaceRender", FName("FClothCollection"), "Collection")
 
 public:
 	UPROPERTY(Meta = (DataflowOutput))
diff --git a/Engine/Plugins/Experimental/Dataflow/Source/DataflowAssetTools/Private/Dataflow/CollectionPatternToDynamicMesh.cpp b/Engine/Plugins/Experimental/Dataflow/Source/DataflowAssetTools/Private/Dataflow/CollectionPatternToDynamicMesh.cpp
deleted file mode 100644
diff --git a/Engine/Plugins/Experimental/Dataflow/Source/DataflowAssetTools/Public/Dataflow/CollectionPatternToDynamicMesh.h b/Engine/Plugins/Experimental/Dataflow/Source/DataflowAssetTools/Public/Dataflow/CollectionPatternToDynamicMesh.h
deleted file mode 100644
diff --git a/Engine/Plugins/Experimental/Dataflow/Source/DataflowEditor/Private/Dataflow/DataflowConstructionScene.cpp b/Engine/Plugins/Experimental/Dataflow/Source/DataflowEditor/Private/Dataflow/DataflowConstructionScene.cpp
--- a/Engine/Plugins/Experimental/Dataflow/Source/DataflowEditor/Private/Dataflow/DataflowConstructionScene.cpp
+++ b/Engine/Plugins/Experimental/Dataflow/Source/DataflowEditor/Private/Dataflow/DataflowConstructionScene.cpp
@@ -141,7 +141,8 @@ void FDataflowConstructionScene::UpdateDynamicMeshComponents()
 					GeometryCollection::Facades::FRenderingFacade Facade(*RenderCollection);
 					Facade.DefineSchema();
 
-					Target->Render(Facade, DataflowContext);
+					check(EditorContent->GetConstructionViewMode());
+					Target->Render(Facade, DataflowContext.ToSharedRef(), *EditorContent->GetConstructionViewMode());
 
 					const int32 NumGeometry = Facade.NumGeometry();
 					for (int32 MeshIndex = 0; MeshIndex < NumGeometry; ++MeshIndex)
@@ -179,7 +180,8 @@ void FDataflowConstructionScene::UpdateDynamicMeshComponents()
 					GeometryCollection::Facades::FRenderingFacade Facade(*RenderCollection);
 					Facade.DefineSchema();
 
-					Target->Render(Facade, DataflowContext);
+					check(EditorContent->GetConstructionViewMode());
+					Target->Render(Facade, DataflowContext.ToSharedRef(), *EditorContent->GetConstructionViewMode());
 
 					const int32 NumGeometry = Facade.NumGeometry();
 					for (int32 MeshIndex = 0; MeshIndex < NumGeometry; ++MeshIndex)
diff --git a/Engine/Plugins/Experimental/Dataflow/Source/DataflowEditor/Private/Dataflow/DataflowConstructionViewport.cpp b/Engine/Plugins/Experimental/Dataflow/Source/DataflowEditor/Private/Dataflow/DataflowConstructionViewport.cpp
--- a/Engine/Plugins/Experimental/Dataflow/Source/DataflowEditor/Private/Dataflow/DataflowConstructionViewport.cpp
+++ b/Engine/Plugins/Experimental/Dataflow/Source/DataflowEditor/Private/Dataflow/DataflowConstructionViewport.cpp
@@ -3,6 +3,7 @@
 #include "Dataflow/DataflowConstructionViewport.h"
 
 #include "Dataflow/DataflowActor.h"
+#include "Dataflow/DataflowEditorCommands.h"
 #include "Dataflow/DataflowEditorMode.h"
 #include "Dataflow/DataflowConstructionViewportClient.h"
 #include "Dataflow/DataflowEditorToolkit.h"
@@ -11,6 +12,7 @@
 #include "Dataflow/DataflowConstructionViewportToolbar.h"
 #include "Dataflow/DataflowConstructionScene.h"
 #include "Dataflow/DataflowEditorPreviewSceneBase.h"
+#include "Dataflow/DataflowRenderingViewMode.h"
 #include "Dataflow/DataflowSimulationPanel.h"
 
 #define LOCTEXT_NAMESPACE "SDataflowConstructionViewport"
@@ -30,7 +32,8 @@ void SDataflowConstructionViewport::Construct(const FArguments& InArgs, const FA
 
 TSharedPtr<SWidget> SDataflowConstructionViewport::MakeViewportToolbar()
 {
-	return SNew(SDataflowConstructionViewportSelectionToolBar, SharedThis(this));
+	return SNew(SDataflowConstructionViewportSelectionToolBar, SharedThis(this))
+		.CommandList(CommandList);
 }
 
 void SDataflowConstructionViewport::OnFocusViewportToSelection()
@@ -57,6 +60,38 @@ UDataflowEditorMode* SDataflowConstructionViewport::GetEdMode() const
 void SDataflowConstructionViewport::BindCommands()
 {
 	SAssetEditorViewport::BindCommands();
+
+	const FDataflowEditorCommandsImpl& CommandInfos = FDataflowEditorCommands::Get();
+
+	for (const TPair<FName, TSharedPtr<FUICommandInfo>>& SetViewModeCommand : CommandInfos.SetConstructionViewModeCommands)
+	{
+		CommandList->MapAction(
+			SetViewModeCommand.Value,
+			FExecuteAction::CreateLambda([this, ViewModeName = SetViewModeCommand.Key]()
+			{
+				if (UDataflowEditorMode* const EdMode = GetEdMode())
+				{
+					EdMode->SetConstructionViewMode(ViewModeName);
+				}
+			}),
+			FCanExecuteAction::CreateLambda([this, ViewModeName = SetViewModeCommand.Key]()
+			{ 
+				if (const UDataflowEditorMode* const EdMode = GetEdMode())
+				{
+					return EdMode->CanChangeConstructionViewModeTo(ViewModeName);
+				}
+				return false; 
+			}),
+			FIsActionChecked::CreateLambda([this, ViewModeName = SetViewModeCommand.Key]()
+			{
+				if (const UDataflowEditorMode* const EdMode = GetEdMode())
+				{
+					return EdMode->GetConstructionViewMode()->GetName() == ViewModeName;
+				}
+				return false;
+			})
+		);
+	}
 }
 
 bool SDataflowConstructionViewport::IsVisible() const
diff --git a/Engine/Plugins/Experimental/Dataflow/Source/DataflowEditor/Private/Dataflow/DataflowConstructionViewportClient.cpp b/Engine/Plugins/Experimental/Dataflow/Source/DataflowEditor/Private/Dataflow/DataflowConstructionViewportClient.cpp
--- a/Engine/Plugins/Experimental/Dataflow/Source/DataflowEditor/Private/Dataflow/DataflowConstructionViewportClient.cpp
+++ b/Engine/Plugins/Experimental/Dataflow/Source/DataflowEditor/Private/Dataflow/DataflowConstructionViewportClient.cpp
@@ -10,6 +10,7 @@
 #include "Dataflow/DataflowGraphEditor.h"
 #include "Dataflow/DataflowConstructionScene.h"
 #include "Dataflow/DataflowEditorPreviewSceneBase.h"
+#include "Dataflow/DataflowRenderingViewMode.h"
 #include "EditorModeManager.h"
 #include "EdModeInteractiveToolsContext.h"
 #include "GraphEditor.h"
@@ -19,6 +20,19 @@
 #include "SNodePanel.h"
 
 
+static_assert(std::is_same<std::underlying_type_t<ELevelViewportType>, int>::value);
+static_assert(std::is_same<std::underlying_type_t<Dataflow::EDataflowViewportType>, int>::value);
+static_assert(static_cast<int>(ELevelViewportType::LVT_OrthoXY) == static_cast<int>(Dataflow::EDataflowViewportType::DVT_OrthoXY));
+static_assert(static_cast<int>(ELevelViewportType::LVT_OrthoXZ) == static_cast<int>(Dataflow::EDataflowViewportType::DVT_OrthoXZ));
+static_assert(static_cast<int>(ELevelViewportType::LVT_OrthoYZ) == static_cast<int>(Dataflow::EDataflowViewportType::DVT_OrthoYZ));
+static_assert(static_cast<int>(ELevelViewportType::LVT_Perspective) == static_cast<int>(Dataflow::EDataflowViewportType::DVT_Perspective));
+static_assert(static_cast<int>(ELevelViewportType::LVT_OrthoFreelook) == static_cast<int>(Dataflow::EDataflowViewportType::DVT_OrthoFreelook));
+static_assert(static_cast<int>(ELevelViewportType::LVT_OrthoNegativeXY) == static_cast<int>(Dataflow::EDataflowViewportType::DVT_OrthoNegativeXY));
+static_assert(static_cast<int>(ELevelViewportType::LVT_OrthoNegativeXZ) == static_cast<int>(Dataflow::EDataflowViewportType::DVT_OrthoNegativeXZ));
+static_assert(static_cast<int>(ELevelViewportType::LVT_OrthoNegativeYZ) == static_cast<int>(Dataflow::EDataflowViewportType::DVT_OrthoNegativeYZ));
+static_assert(static_cast<int>(ELevelViewportType::LVT_MAX) == static_cast<int>(Dataflow::EDataflowViewportType::DVT_MAX));
+
+#define LOCTEXT_NAMESPACE "DataflowConstructionViewportClient"
 
 FDataflowConstructionViewportClient::FDataflowConstructionViewportClient(FEditorModeTools* InModeTools,
                                                              FPreviewScene* InPreviewScene,  const bool bCouldTickScene,
@@ -189,51 +203,44 @@ void FDataflowConstructionViewportClient::ProcessClick(FSceneView& View, HHitPro
 	OnSelectionChangedMulticast.Broadcast(CurrentlySelectedComponents);
 }
 
-void FDataflowConstructionViewportClient::SetConstructionViewMode(Dataflow::EDataflowPatternVertexType InViewMode)
+void FDataflowConstructionViewportClient::SetConstructionViewMode(const Dataflow::IDataflowConstructionViewMode* InViewMode)
 {
-	// @todo(Dataflow) : Add support for Sim2D
-	//const bool bSwitching2D3D = (ConstructionViewMode == Dataflow::EDataflowPatternVertexType::Sim2D) != (InViewMode == Dataflow::EDataflowPatternVertexType::Sim2D);
-	//if (bSwitching2D3D)
-	//{
-	//	Swap(SavedInactiveViewTransform, ViewTransformPerspective);
-	//}
+	checkf(InViewMode, TEXT("SetConstructionViewMode received null IDataflowConstructionViewMode pointer"));
 
-	ConstructionViewMode = Dataflow::EDataflowPatternVertexType::Sim3D;
+	if (ConstructionViewMode)
+	{
+		SavedInactiveViewTransforms.FindOrAdd(ConstructionViewMode->GetName()) = GetViewTransform();
+	}
 
-	
-	//if (ConstructionViewMode == EDataflowPatternVertexType::Sim2D)
-	//{
-	//	for (UInputBehavior* const Behavior : BehaviorsFor2DMode)
-	//	{
-	//		BehaviorSet->Add(Behavior);
-	//	}
-	//
-	//	const double AbsZ = FMath::Abs(ViewTransformPerspective.GetLocation().Z);
-	//	constexpr double CameraFarPlaneWorldZ = -10.0;
-	//	constexpr double CameraNearPlaneProportionZ = 0.8;
-	//	OverrideFarClipPlane(static_cast<float>(AbsZ - CameraFarPlaneWorldZ));
-	//	OverrideNearClipPlane(static_cast<float>(AbsZ * (1.0 - CameraNearPlaneProportionZ)));
-	//}
-	//else
-	//{
-	OverrideFarClipPlane(0);
-	OverrideNearClipPlane(UE_KINDA_SMALL_NUMBER);
-	//}
-
-	//ModeTools->GetInteractiveToolsContext()->InputRouter->DeregisterSource(this);
-	//ModeTools->GetInteractiveToolsContext()->InputRouter->RegisterSource(this);
+	ConstructionViewMode = InViewMode;
 
-}
+	SetViewportType((ELevelViewportType)ConstructionViewMode->GetViewportType());
 
+	if (const FViewportCameraTransform* const FoundPreviousTransform = SavedInactiveViewTransforms.Find(InViewMode->GetName()))
+	{
+		if (ConstructionViewMode->IsPerspective())
+		{
+			ViewTransformPerspective = *FoundPreviousTransform;
+		}
+		else
+		{
+			ViewTransformOrthographic = *FoundPreviousTransform;
+		}
+	}
+	else
+	{
+		// TODO: Default view transform
+	}
 
-Dataflow::EDataflowPatternVertexType FDataflowConstructionViewportClient::GetConstructionViewMode() const
-{
-	return Dataflow::EDataflowPatternVertexType::Sim3D;// ConstructionViewMode;
+	bDrawAxes = ConstructionViewMode->IsPerspective();
+	Invalidate();
 }
 
+
 void FDataflowConstructionViewportClient::AddReferencedObjects(FReferenceCollector& Collector)
 {
 	Super::AddReferencedObjects(Collector);
 	Collector.AddReferencedObject(BehaviorSet);
 }
 
+#undef LOCTEXT_NAMESPACE 
diff --git a/Engine/Plugins/Experimental/Dataflow/Source/DataflowEditor/Private/Dataflow/DataflowConstructionViewportToolbar.cpp b/Engine/Plugins/Experimental/Dataflow/Source/DataflowEditor/Private/Dataflow/DataflowConstructionViewportToolbar.cpp
--- a/Engine/Plugins/Experimental/Dataflow/Source/DataflowEditor/Private/Dataflow/DataflowConstructionViewportToolbar.cpp
+++ b/Engine/Plugins/Experimental/Dataflow/Source/DataflowEditor/Private/Dataflow/DataflowConstructionViewportToolbar.cpp
@@ -1,10 +1,163 @@
 // Copyright Epic Games, Inc. All Rights Reserved.
+
 #include "Dataflow/DataflowConstructionViewportToolbar.h"
+
 #include "Dataflow/DataflowConstructionViewport.h"
+#include "Dataflow/DataflowEditorCommands.h"
+#include "Dataflow/DataflowEditorMode.h"
+#include "Dataflow/DataflowRenderingViewMode.h"
+#include "EditorModeManager.h"
+#include "SEditorViewportViewMenu.h"
+#include "ToolMenus.h"
+#include "Widgets/Layout/SBorder.h"
+
+#define LOCTEXT_NAMESPACE "SDataflowConstructionViewportSelectionToolBar"
 
 void SDataflowConstructionViewportSelectionToolBar::Construct(const FArguments& InArgs, TSharedPtr<SDataflowConstructionViewport> InDataflowViewport)
 {
 	EditorViewport = InDataflowViewport;
-	
 	SCommonEditorViewportToolbarBase::Construct(SCommonEditorViewportToolbarBase::FArguments(), InDataflowViewport);
-}
\ No newline at end of file
+
+	CommandList = InArgs._CommandList;
+
+	UToolMenu* const ViewModeMenu = UToolMenus::Get()->RegisterMenu("DataflowEditor.ConstructionViewModeMenu");
+
+	// Dynamically populate the view mode menu based on which ViewModes are available for the current node selection
+	ViewModeMenu->AddDynamicSection("DataflowConstructionViewModeMenuSection",
+		FNewToolMenuDelegate::CreateLambda([this](UToolMenu* ViewModeMenu)
+		{	
+			const TSharedPtr<const SDataflowConstructionViewport> PinnedViewport = EditorViewport.Pin();
+			if (PinnedViewport.IsValid())
+			{
+				FToolMenuSection& ViewModesSection = ViewModeMenu->AddSection(NAME_None, FText());
+
+				const TSharedPtr<const FEditorViewportClient> ViewportClient = PinnedViewport->GetViewportClient();
+				check(ViewportClient.IsValid());
+
+				const FEditorModeTools* const EditorModeTools = ViewportClient->GetModeTools();
+				check(EditorModeTools);
+
+				if (const UDataflowEditorMode* const DataflowEdMode = Cast<UDataflowEditorMode>(EditorModeTools->GetActiveScriptableMode(UDataflowEditorMode::EM_DataflowEditorModeId)))
+				{
+					for (const TPair<FName, TSharedPtr<FUICommandInfo>>& NameAndCommand : FDataflowEditorCommandsImpl::Get().SetConstructionViewModeCommands)
+					{
+						const FName ViewModeName = NameAndCommand.Key;
+						if (DataflowEdMode->CanChangeConstructionViewModeTo(ViewModeName))
+						{
+							ViewModesSection.AddMenuEntry(NameAndCommand.Value);
+						}
+					}
+				}
+			}
+		}));
+
+	const FMargin ToolbarSlotPadding(4.0f, 1.0f);
+	TSharedPtr<SHorizontalBox> MainBoxPtr;
+
+	ChildSlot
+	[
+		SNew( SBorder )
+		.BorderImage(FAppStyle::Get().GetBrush("EditorViewportToolBar.Background"))
+		.Cursor(EMouseCursor::Default)
+		[
+			SNew(SVerticalBox)
+				+ SVerticalBox::Slot()
+				.AutoHeight()
+			[
+				SAssignNew( MainBoxPtr, SHorizontalBox )
+			]
+		]
+	];
+
+	// Display (Lit, Unlit, Wireframe, etc.)
+	MainBoxPtr->AddSlot()
+		.Padding(ToolbarSlotPadding)
+		.HAlign(HAlign_Left)
+		[
+			MakeDisplayToolBar(InArgs._Extenders)
+		];
+
+	// View mode (Sim2D/Sim3D/Render)
+	MainBoxPtr->AddSlot()
+		.Padding(ToolbarSlotPadding)
+		.HAlign(HAlign_Right)
+		[
+			MakeToolBar(InArgs._Extenders)
+		];
+
+	// See SCommonEditorViewportToolbarBase::Construct for more possible menus to add
+}
+
+
+
+TSharedRef<SWidget> SDataflowConstructionViewportSelectionToolBar::MakeToolBar(const TSharedPtr<FExtender> InExtenders)
+{
+	// The following is modeled after portions of STransformViewportToolBar, which gets 
+	// used in SCommonEditorViewportToolbarBase.
+	// The buttons are hooked up to actual functions via command bindings in SChaosClothAssetEditorRestSpaceViewport::BindCommands(),
+	// and the toolbar gets built in SDataflowConstructionViewport::MakeViewportToolbar().
+
+	FSlimHorizontalToolBarBuilder ToolbarBuilder(CommandList, FMultiBoxCustomization::None, InExtenders);
+
+	// Use a custom style
+	const FName ToolBarStyle = "EditorViewportToolBar";
+	ToolbarBuilder.SetStyle(&FAppStyle::Get(), ToolBarStyle);
+	ToolbarBuilder.SetLabelVisibility(EVisibility::Collapsed);
+
+	ToolbarBuilder.BeginBlockGroup();
+	{
+		// View mode selector (2D/3D/Render)
+		ToolbarBuilder.AddWidget(
+			SAssignNew(ViewModeDropDown, SEditorViewportToolbarMenu)
+			.ParentToolBar(SharedThis(this))
+			.Cursor(EMouseCursor::Default)
+			.Label(this, &SDataflowConstructionViewportSelectionToolBar::GetViewModeMenuLabel)
+			.LabelIcon(this, &SDataflowConstructionViewportSelectionToolBar::GetViewModeMenuLabelIcon)
+			.OnGetMenuContent(this, &SDataflowConstructionViewportSelectionToolBar::GenerateViewModeMenuContent)
+		);
+	}
+	ToolbarBuilder.EndBlockGroup();
+	ToolbarBuilder.EndSection(); // View Controls
+
+	return ToolbarBuilder.MakeWidget();
+}
+
+FText SDataflowConstructionViewportSelectionToolBar::GetViewModeMenuLabel() const
+{
+	FText Label = LOCTEXT("DataflowConstructionViewMenuTitle_Default", "View");
+
+	if (const TSharedPtr<SDataflowConstructionViewport> PinnedViewport = EditorViewport.Pin())
+	{
+		const TSharedPtr<FEditorViewportClient> ViewportClient = PinnedViewport->GetViewportClient();
+		check(ViewportClient.IsValid());
+		const FEditorModeTools* const EditorModeTools = ViewportClient->GetModeTools();
+		const UDataflowEditorMode* const DataflowEdMode = Cast<UDataflowEditorMode>(EditorModeTools->GetActiveScriptableMode(UDataflowEditorMode::EM_DataflowEditorModeId));
+		if (DataflowEdMode)
+		{
+			Label = DataflowEdMode->GetConstructionViewMode()->GetButtonText();
+		}
+	}
+
+	return Label;
+}
+
+const FSlateBrush* SDataflowConstructionViewportSelectionToolBar::GetViewModeMenuLabelIcon() const
+{
+	return FStyleDefaults::GetNoBrush();
+}
+
+TSharedRef<SWidget> SDataflowConstructionViewportSelectionToolBar::GenerateViewModeMenuContent() const
+{
+	return UToolMenus::Get()->GenerateWidget("DataflowEditor.ConstructionViewModeMenu", FToolMenuContext(CommandList));
+}
+
+TSharedRef<SWidget> SDataflowConstructionViewportSelectionToolBar::MakeDisplayToolBar(const TSharedPtr<FExtender> InExtenders)
+{
+	const TSharedRef<SEditorViewport> ViewportRef = StaticCastSharedPtr<SEditorViewport>(EditorViewport.Pin()).ToSharedRef();
+
+	return SNew(SEditorViewportViewMenu, ViewportRef, SharedThis(this))
+		.Cursor(EMouseCursor::Default)
+		.MenuExtenders(InExtenders);
+}
+
+#undef LOCTEXT_NAMESPACE
diff --git a/Engine/Plugins/Experimental/Dataflow/Source/DataflowEditor/Private/Dataflow/DataflowEditorCommands.cpp b/Engine/Plugins/Experimental/Dataflow/Source/DataflowEditor/Private/Dataflow/DataflowEditorCommands.cpp
--- a/Engine/Plugins/Experimental/Dataflow/Source/DataflowEditor/Private/Dataflow/DataflowEditorCommands.cpp
+++ b/Engine/Plugins/Experimental/Dataflow/Source/DataflowEditor/Private/Dataflow/DataflowEditorCommands.cpp
@@ -11,6 +11,7 @@
 #include "Dataflow/DataflowNodeFactory.h"
 #include "Dataflow/DataflowObject.h"
 #include "Dataflow/DataflowOverrideNode.h"
+#include "Dataflow/DataflowRenderingViewMode.h"
 #include "Dataflow/DataflowSCommentNode.h"
 #include "Dataflow/DataflowSNode.h"
 #include "DataflowEditorTools/DataflowEditorWeightMapPaintTool.h"
@@ -65,6 +66,25 @@ void FDataflowEditorCommandsImpl::RegisterCommands()
 	UI_COMMAND(AddWeightMapNode, "Add Weight Map", "Paint weight maps on the mesh", EUserInterfaceActionType::Button, FInputChord());
 	UI_COMMAND(UpdateSimulationCache, "UpdateSimulationCache", "Update Simulation Cache", EUserInterfaceActionType::ToggleButton, FInputChord());
 
+	for (const TPair<FName, TUniquePtr<Dataflow::IDataflowConstructionViewMode>>& NameAndMode : Dataflow::FRenderingViewModeFactory::GetInstance().GetViewModes())
+	{
+		TSharedPtr< FUICommandInfo > SetViewModeCommand;
+		
+		const Dataflow::IDataflowConstructionViewMode* const ViewMode = NameAndMode.Value.Get();
+		checkf(ViewMode, TEXT("Registered mode in FRenderingViewModeFactory has no associated IDataflowConstructionViewMode object. Registered name: %s"), *NameAndMode.Key.ToString());
+
+		FUICommandInfo::MakeCommandInfo(
+			this->AsShared(),
+			SetViewModeCommand,
+			ViewMode->GetName(),
+			ViewMode->GetButtonText(),
+			ViewMode->GetTooltipText(),
+			FSlateIcon(),
+			EUserInterfaceActionType::Button,
+			FInputChord()
+		);
+		SetConstructionViewModeCommands.Add(ViewMode->GetName(), SetViewModeCommand);
+	}
 
 	if (Dataflow::FNodeFactory* Factory = Dataflow::FNodeFactory::GetInstance())
 	{
diff --git a/Engine/Plugins/Experimental/Dataflow/Source/DataflowEditor/Private/Dataflow/DataflowEditorMode.cpp b/Engine/Plugins/Experimental/Dataflow/Source/DataflowEditor/Private/Dataflow/DataflowEditorMode.cpp
--- a/Engine/Plugins/Experimental/Dataflow/Source/DataflowEditor/Private/Dataflow/DataflowEditorMode.cpp
+++ b/Engine/Plugins/Experimental/Dataflow/Source/DataflowEditor/Private/Dataflow/DataflowEditorMode.cpp
@@ -18,6 +18,7 @@
 #include "Dataflow/DataflowGraphEditor.h"
 #include "Dataflow/DataflowEditorPreviewSceneBase.h"
 #include "Dataflow/DataflowConstructionScene.h"
+#include "Dataflow/DataflowRenderingViewMode.h"
 #include "Dataflow/DataflowSimulationScene.h"
 #include "Dataflow/DataflowSimulationViewportClient.h"
 #include "Dataflow/DataflowSNode.h"
@@ -75,6 +76,9 @@ void UDataflowEditorMode::Enter()
 
 	// Register gizmo ContextObject for use inside interactive tools
 	UE::TransformGizmoUtil::RegisterTransformGizmoContextObject(GetInteractiveToolsContext());
+
+	// Initialize view mode to a default
+	ConstructionViewMode = Dataflow::FRenderingViewModeFactory::GetInstance().GetViewMode(Dataflow::FDataflowConstruction3DViewMode::Name);
 }
 
 void UDataflowEditorMode::SetDataflowEditor(UDataflowEditor* InDataflowEditor) 
@@ -123,34 +127,6 @@ void UDataflowEditorMode::RegisterDataflowTool(TSharedPtr<FUICommandInfo> UIComm
 	CommandList->MapAction(UICommand, 
 		FExecuteAction::CreateWeakLambda(ToolsContext, [this, ToolsContext, ToolIdentifier, DataflowToolBuilder]()
 		{
-			// Check if we need to switch view modes before starting the tool
-			TArray<Dataflow::EDataflowPatternVertexType> SupportedModes;
-			DataflowToolBuilder->GetSupportedViewModes(SupportedModes);
-
-			if (SupportedModes.Num() > 0 && !SupportedModes.Contains(this->GetConstructionViewMode()))
-			{
-				if (!bShouldRestoreSavedConstructionViewMode)
-				{
-					// remember the current view mode so we can restore it later
-					SavedConstructionViewMode = this->GetConstructionViewMode();
-					bShouldRestoreSavedConstructionViewMode = true;
-				}
-
-				// switch to the preferred view mode for the tool that's about to start
-				this->SetConstructionViewMode(SupportedModes[0]);
-			}
-
-			// Check if we need to disable wireframe mode before starting tool.
-			const bool bCanSetWireframeActive = DataflowToolBuilder->CanSetConstructionViewWireframeActive();
-			if (!bCanSetWireframeActive)
-			{
-				if (!bShouldRestoreConstructionViewWireframe)
-				{
-					bShouldRestoreConstructionViewWireframe = bConstructionViewWireframe;
-				}
-				bConstructionViewWireframe = false;
-			}
-
 			ActiveToolsContext = ToolsContext;
 			ToolsContext->StartTool(ToolIdentifier);
 		}),
@@ -393,18 +369,6 @@ void UDataflowEditorMode::RefocusConstructionViewportClient()
 		PinnedVC->ToggleOrbitCamera(false);
 
 		const FBox SceneBounds = SceneBoundingBox();
-		const bool bPattern2DMode = (ConstructionViewMode == Dataflow::EDataflowPatternVertexType::Sim2D);
-		if (bPattern2DMode)
-		{
-			// 2D pattern
-			PinnedVC->SetInitialViewTransform(ELevelViewportType::LVT_Perspective, FVector(0, 0, -100), FRotator(90, -90, 0), DEFAULT_ORTHOZOOM);
-		}
-		else
-		{
-			// 3D rest space
-			PinnedVC->SetInitialViewTransform(ELevelViewportType::LVT_Perspective, FVector(0, 150, 200), FRotator(0, 0, 0), DEFAULT_ORTHOZOOM);
-		}
-
 		constexpr bool bInstant = true;
 		PinnedVC->FocusViewportOnBox(SceneBounds, bInstant);
 
@@ -435,7 +399,7 @@ void UDataflowEditorMode::FirstTimeFocusConstructionViewport()
 {
 	// If this is the first time seeing a valid 2D or 3D mesh, refocus the camera on it.
 	const bool bIsValid = (ConstructionScene->HasRenderableGeometry());
-	const bool bIs2D = ConstructionViewMode == Dataflow::EDataflowPatternVertexType::Sim3D;
+	const bool bIs2D = !ConstructionViewMode->IsPerspective();
 
 	if (bIsValid)
 	{
@@ -552,9 +516,10 @@ FBox UDataflowEditorMode::SelectionBoundingBox() const
 	return SceneBoundingBox();
 }
 
-void UDataflowEditorMode::SetConstructionViewMode(Dataflow::EDataflowPatternVertexType InMode)
+void UDataflowEditorMode::SetConstructionViewMode(const FName& NewViewModeName)
 {
 	// We will first check if there is an active tool. If so, we'll shut down the tool and save the results to the Node, then change view modes, then restart the tool again.
+	
 	bool bEndedActiveTool = false;
 	UInteractiveToolManager* const ToolManager = GetInteractiveToolsContext()->ToolManager;
 	checkf(ToolManager, TEXT("No valid ToolManager found for UDataflowEditorMode"));
@@ -571,7 +536,20 @@ void UDataflowEditorMode::SetConstructionViewMode(Dataflow::EDataflowPatternVert
 		bShouldRestoreSavedConstructionViewMode = bTempShouldRestoreVal;
 	}
 
-	ConstructionViewMode = InMode;
+	// Next get the ViewMode pointer from the given name
+
+	const Dataflow::FRenderingViewModeFactory& ViewModes = Dataflow::FRenderingViewModeFactory::GetInstance();
+	const Dataflow::IDataflowConstructionViewMode* const NewMode = ViewModes.GetViewMode(NewViewModeName);
+	if (!NewMode)
+	{
+		UE_LOG(LogChaos, Warning, TEXT("Warning : Unknown rendering view mode: %s"), *NewViewModeName.ToString());
+		return;
+	}
+
+	// Do the actual view mode updates
+
+	ConstructionViewMode = NewMode;
+	ConstructionScene->GetEditorContent()->SetConstructionViewMode(ConstructionViewMode);
 	ConstructionScene->UpdateConstructionScene();
 
 	const TSharedPtr<FDataflowConstructionViewportClient> VC = ConstructionViewportClient.Pin();
@@ -597,35 +575,60 @@ void UDataflowEditorMode::SetConstructionViewMode(Dataflow::EDataflowPatternVert
 	}
 }
 
-Dataflow::EDataflowPatternVertexType UDataflowEditorMode::GetConstructionViewMode() const
+const Dataflow::IDataflowConstructionViewMode* UDataflowEditorMode::GetConstructionViewMode() const
 {
 	return ConstructionViewMode;
 }
 
-
-bool UDataflowEditorMode::CanChangeConstructionViewModeTo(Dataflow::EDataflowPatternVertexType NewViewMode) const
+bool UDataflowEditorMode::CanChangeConstructionViewModeTo(const FName& NewViewModeName) const
 {
-	check(false);
-
 	if (!GetToolManager()->HasActiveTool(EToolSide::Left))
 	{
-		return true;
-	}
+		if (const TSharedPtr<const SDataflowGraphEditor> PinnedDataflowGraphEditor = DataflowGraphEditor.Pin())
+		{
+			if (const UEdGraphNode* const SelectedNode = PinnedDataflowGraphEditor->GetSingleSelectedNode())
+			{
+				if (const UDataflowEdNode* const SelectedDataflowNode = Cast<UDataflowEdNode>(SelectedNode))
+				{
+					const Dataflow::FRenderingViewModeFactory& ViewModeFactory = Dataflow::FRenderingViewModeFactory::GetInstance();
 
-	return false;
-	/*
-	const UInteractiveToolBuilder* const ActiveToolBuilder = GetToolManager()->GetActiveToolBuilder(EToolSide::Left);
-	checkf(ActiveToolBuilder, TEXT("No Active Tool Builder found despite having an Active Tool"));
+					if (const Dataflow::IDataflowConstructionViewMode* const ViewMode = ViewModeFactory.GetViewMode(NewViewModeName))
+					{
+						if (const TObjectPtr<UDataflowBaseContent> EditorContent = ConstructionScene->GetEditorContent())
+						{
+							if (const TSharedPtr<Dataflow::FEngineContext> Context = EditorContent->GetDataflowContext())
+							{
+								if (SelectedDataflowNode->CanRender(Context.ToSharedRef(), *ViewMode))
+								{
+									return true;
+								}
+							}
+						}
+					}
+				}
+			}
+		}
 
-	const IDataflowEditorToolBuilder* const DataflowToolBuilder = Cast<const IDataflowEditorToolBuilder>(ActiveToolBuilder);
-	checkf(ClothToolBuilder, TEXT("Cloth Editor has an active Tool Builder that does not implement IDataflowEditorToolBuilder"));
+		return false;
+	}
+
+	// TODO: Check active tool to see if we can switch modes while the tool is running
+//	/*
+//	const UInteractiveToolBuilder* const ActiveToolBuilder = GetToolManager()->GetActiveToolBuilder(EToolSide::Left);
+//	checkf(ActiveToolBuilder, TEXT("No Active Tool Builder found despite having an Active Tool"));
+//
+//	const IDataflowEditorToolBuilder* const DataflowToolBuilder = Cast<const IDataflowEditorToolBuilder>(ActiveToolBuilder);
+//	checkf(ClothToolBuilder, TEXT("Cloth Editor has an active Tool Builder that does not implement IDataflowEditorToolBuilder"));
+//
+//	TArray<Dataflow::EDataflowPatternVertexType> SupportedViewModes;
+//	ClothToolBuilder->GetSupportedViewModes(SupportedViewModes);
+//	return SupportedViewModes.Contains(NewViewMode);
+//	*/
 
-	TArray<Dataflow::EDataflowPatternVertexType> SupportedViewModes;
-	ClothToolBuilder->GetSupportedViewModes(SupportedViewModes);
-	return SupportedViewModes.Contains(NewViewMode);
-	*/
+	return false;
 }
 
+
 void UDataflowEditorMode::ToggleConstructionViewWireframe()
 {
 	check(false);
@@ -695,7 +698,6 @@ void UDataflowEditorMode::InitializeContextObject()
 		}
 
 		check(ContextObject);
-
 		ContextObject->SetConstructionViewMode(ConstructionViewMode);
 	}
 }
diff --git a/Engine/Plugins/Experimental/Dataflow/Source/DataflowEditor/Private/Dataflow/DataflowEditorToolkit.cpp b/Engine/Plugins/Experimental/Dataflow/Source/DataflowEditor/Private/Dataflow/DataflowEditorToolkit.cpp
--- a/Engine/Plugins/Experimental/Dataflow/Source/DataflowEditor/Private/Dataflow/DataflowEditorToolkit.cpp
+++ b/Engine/Plugins/Experimental/Dataflow/Source/DataflowEditor/Private/Dataflow/DataflowEditorToolkit.cpp
@@ -656,6 +656,13 @@ void FDataflowEditorToolkit::OnNodeSelectionChanged(const TSet<UObject*>& InNewS
 		};
 
 		const TSet<TObjectPtr<UObject>> PreviouslySelectedNodes = SelectedDataflowNodes;
+		for (UObject* const PreviouslySelectedNode : SelectedDataflowNodes)
+		{
+			if (UDataflowEdNode* const PreviouslySelectedEdNode = Cast<UDataflowEdNode>(PreviouslySelectedNode))
+			{
+				PreviouslySelectedEdNode->SetShouldRenderNode(false);
+			}
+		}
 
 		// Only keep UDataflowEdNode from NewSelection
 		TSet< TObjectPtr<UObject> > NodeSelection = FindDataflowNodesInSet(AsObjectPointers(InNewSelection));
@@ -704,7 +711,15 @@ void FDataflowEditorToolkit::OnNodeSelectionChanged(const TSet<UObject*>& InNewS
 
 			SelectedDataflowNodes = NodeSelection;
 		}
-	
+
+		for (UObject* const SelectedNode : NodeSelection)
+		{
+			if (UDataflowEdNode* SelectedEdNode = Cast<UDataflowEdNode>(SelectedNode))
+			{
+				SelectedEdNode->SetShouldRenderNode(true);
+			}
+		}
+
 		if (bPrimarySelectionChanged)
 		{
 			for (UObject* const PreviouslySelectedNode : PreviouslySelectedNodes)
@@ -745,6 +760,62 @@ void FDataflowEditorToolkit::OnNodeSelectionChanged(const TSet<UObject*>& InNewS
 			}
 		}
 	}
+
+	//
+	// Check if the current view mode can render the selected node. If not, try to find a view mode that can. 
+	//
+
+	if (UDataflowEditorMode* const DataflowMode = CastChecked<UDataflowEditorMode>(EditorModeManager->GetActiveScriptableMode(UDataflowEditorMode::EM_DataflowEditorModeId)))
+	{
+		bool bFoundViewMode = true;
+
+		if (const TObjectPtr<UDataflowBaseContent>& EditorContent = GetEditorContent())
+		{
+			if (const TSharedPtr<Dataflow::FEngineContext> Context = EditorContent->GetDataflowContext())
+			{
+				if (PrimarySelection)
+				{
+					if (!PrimarySelection->CanRender(Context.ToSharedRef(), *DataflowMode->GetConstructionViewMode()))
+					{
+						// Selected node can't render with the current view mode. Check through available view modes and see if it can render with any of them
+
+						bFoundViewMode = false;
+
+						TArray<Dataflow::FRenderingParameter> RenderingParameters = PrimarySelection->GetRenderParameters();
+						for (const Dataflow::FRenderingParameter& Param : RenderingParameters)
+						{
+							const FName NodeOutputTypeName = Param.Type;
+
+							for (const TPair<FName, TUniquePtr<Dataflow::IDataflowConstructionViewMode>>& ViewMode : Dataflow::FRenderingViewModeFactory::GetInstance().GetViewModes())
+							{
+								check(ViewMode.Value.IsValid());
+
+								const bool bCanRender = PrimarySelection->CanRender(Context.ToSharedRef(), *ViewMode.Value);
+
+								if (bCanRender)
+								{
+									DataflowMode->SetConstructionViewMode(ViewMode.Key);
+									bFoundViewMode = true;
+									break;
+								}
+							}
+
+							if (bFoundViewMode)
+							{
+								break;
+							}
+						}
+					}
+				}
+			}
+		}
+
+		if (!bFoundViewMode)
+		{
+			// TODO: Clear and disable View Mode Button. For now set default mode to the built-in 3D view mode.
+			DataflowMode->SetConstructionViewMode(Dataflow::FDataflowConstruction3DViewMode::Name);
+		}
+	}
 }
 
 void FDataflowEditorToolkit::OnNodeSingleClicked(UObject* ClickedNode) const
diff --git a/Engine/Plugins/Experimental/Dataflow/Source/DataflowEditor/Private/DataflowEditorTools/DataflowEditorWeightMapPaintTool.cpp b/Engine/Plugins/Experimental/Dataflow/Source/DataflowEditor/Private/DataflowEditorTools/DataflowEditorWeightMapPaintTool.cpp
--- a/Engine/Plugins/Experimental/Dataflow/Source/DataflowEditor/Private/DataflowEditorTools/DataflowEditorWeightMapPaintTool.cpp
+++ b/Engine/Plugins/Experimental/Dataflow/Source/DataflowEditor/Private/DataflowEditorTools/DataflowEditorWeightMapPaintTool.cpp
@@ -67,9 +67,9 @@ namespace Dataflow::Private
 /*
  * ToolBuilder
  */
-void UDataflowEditorWeightMapPaintToolBuilder::GetSupportedViewModes(TArray<Dataflow::EDataflowPatternVertexType>& Modes) const
+void UDataflowEditorWeightMapPaintToolBuilder::GetSupportedViewModes(TArray<Dataflow::IDataflowConstructionViewMode*>& Modes) const
 {
-	Modes.Add(Dataflow::EDataflowPatternVertexType::Sim3D);
+	//Modes.Add(Dataflow::EDataflowPatternVertexType::Sim3D);
 	//Modes.Add(Dataflow::EDataflowPatternVertexType::Sim2D);
 }
 
diff --git a/Engine/Plugins/Experimental/Dataflow/Source/DataflowEditor/Public/Dataflow/DataflowConstructionViewportClient.h b/Engine/Plugins/Experimental/Dataflow/Source/DataflowEditor/Public/Dataflow/DataflowConstructionViewportClient.h
--- a/Engine/Plugins/Experimental/Dataflow/Source/DataflowEditor/Public/Dataflow/DataflowConstructionViewportClient.h
+++ b/Engine/Plugins/Experimental/Dataflow/Source/DataflowEditor/Public/Dataflow/DataflowConstructionViewportClient.h
@@ -18,8 +18,13 @@ class FTransformGizmoDataBinder;
 class FDataflowPreviewSceneBase;
 class UInputBehaviorSet;
 class USelection;
+class FDataflowConstructionViewportClient;
+namespace Dataflow
+{
+	class IDataflowConstructionViewMode;
+}
 
-class DATAFLOWEDITOR_API FDataflowConstructionViewportClient : public FEditorViewportClient //, public IInputBehaviorSource
+class DATAFLOWEDITOR_API FDataflowConstructionViewportClient : public FEditorViewportClient
 {
 public:
 	using Super = FEditorViewportClient;
@@ -27,8 +32,7 @@ public:
 	FDataflowConstructionViewportClient(FEditorModeTools* InModeTools, FPreviewScene* InPreviewScene,  const bool bCouldTickScene,
 								  const TWeakPtr<SEditorViewport> InEditorViewportWidget = nullptr);
 
-	void SetConstructionViewMode(Dataflow::EDataflowPatternVertexType InViewMode);
-	Dataflow::EDataflowPatternVertexType GetConstructionViewMode() const;
+	void SetConstructionViewMode(const Dataflow::IDataflowConstructionViewMode* InViewMode);
 
 	// IInputBehaviorSource
 	// virtual const UInputBehaviorSet* GetInputBehaviors() const override;
@@ -66,11 +70,14 @@ private:
 	TWeakPtr<FUICommandList> ToolCommandList;
 
 	/** Construction view mode */
-	Dataflow::EDataflowPatternVertexType ConstructionViewMode = Dataflow::EDataflowPatternVertexType::Sim3D;
+	const Dataflow::IDataflowConstructionViewMode* ConstructionViewMode = nullptr;
 
 	/** Behavior set for the behavior UI */
 	TObjectPtr<UInputBehaviorSet> BehaviorSet;
 	
 	/** Flag to enable scene ticking from the client */
 	bool bEnableSceneTicking = false;
+
+	// Saved view transforms for the currently inactive view modes (e.g. store the 3D camera here while in 2D mode and vice-versa)
+	TMap<FName, FViewportCameraTransform> SavedInactiveViewTransforms;
 };
diff --git a/Engine/Plugins/Experimental/Dataflow/Source/DataflowEditor/Public/Dataflow/DataflowConstructionViewportToolbar.h b/Engine/Plugins/Experimental/Dataflow/Source/DataflowEditor/Public/Dataflow/DataflowConstructionViewportToolbar.h
--- a/Engine/Plugins/Experimental/Dataflow/Source/DataflowEditor/Public/Dataflow/DataflowConstructionViewportToolbar.h
+++ b/Engine/Plugins/Experimental/Dataflow/Source/DataflowEditor/Public/Dataflow/DataflowConstructionViewportToolbar.h
@@ -6,18 +6,37 @@
 
 // ----------------------------------------------------------------------------------
 class SDataflowConstructionViewport;
+class SEditorViewportToolbarMenu;
 
 /** Base toolbar for the dataflow. Should be extended to add more features */
 class SDataflowConstructionViewportSelectionToolBar : public SCommonEditorViewportToolbarBase
 {
 public:
 	SLATE_BEGIN_ARGS(SDataflowConstructionViewportSelectionToolBar) {}
+		SLATE_ARGUMENT(TSharedPtr<class FUICommandList>, CommandList)
+		SLATE_ARGUMENT(TSharedPtr<class FExtender>, Extenders)
 	SLATE_END_ARGS()
 
 	/** Constructs this widget with the given parameters */
 	void Construct(const FArguments& InArgs, TSharedPtr<SDataflowConstructionViewport> InDataflowViewport);
-	
+
+
 private:
+
+	TSharedRef<SWidget> MakeToolBar(const TSharedPtr<FExtender> InExtenders);
+
+	TSharedPtr<SEditorViewportToolbarMenu> ViewModeDropDown;
+
+	// View mode
+	FText GetViewModeMenuLabel() const;
+	const FSlateBrush* GetViewModeMenuLabelIcon() const;
+	TSharedRef<SWidget> GenerateViewModeMenuContent() const;
+
+	// Display (Lit, Unlit, Wireframe, etc.)
+	TSharedRef<SWidget> MakeDisplayToolBar(const TSharedPtr<FExtender> InExtenders);
+
 	/** Reference to the parent viewport */
 	TWeakPtr<SDataflowConstructionViewport> EditorViewport;
+
+	TSharedPtr<FUICommandList> CommandList;
 };
\ No newline at end of file
diff --git a/Engine/Plugins/Experimental/Dataflow/Source/DataflowEditor/Public/Dataflow/DataflowEditorCommands.h b/Engine/Plugins/Experimental/Dataflow/Source/DataflowEditor/Public/Dataflow/DataflowEditorCommands.h
--- a/Engine/Plugins/Experimental/Dataflow/Source/DataflowEditor/Public/Dataflow/DataflowEditorCommands.h
+++ b/Engine/Plugins/Experimental/Dataflow/Source/DataflowEditor/Public/Dataflow/DataflowEditorCommands.h
@@ -53,6 +53,8 @@ public:
 	TSharedPtr<FUICommandInfo> RemoveOptionPin;
 	TSharedPtr<FUICommandInfo> ZoomToFitGraph;
 
+	TMap<FName, TSharedPtr<FUICommandInfo>> SetConstructionViewModeCommands;
+
 	TMap< FName, TSharedPtr<FUICommandInfo> > CreateNodesMap;
 
 	const static FString BeginWeightMapPaintToolIdentifier;
diff --git a/Engine/Plugins/Experimental/Dataflow/Source/DataflowEditor/Public/Dataflow/DataflowEditorMode.h b/Engine/Plugins/Experimental/Dataflow/Source/DataflowEditor/Public/Dataflow/DataflowEditorMode.h
--- a/Engine/Plugins/Experimental/Dataflow/Source/DataflowEditor/Public/Dataflow/DataflowEditorMode.h
+++ b/Engine/Plugins/Experimental/Dataflow/Source/DataflowEditor/Public/Dataflow/DataflowEditorMode.h
@@ -17,6 +17,10 @@ struct FManagedArrayCollection;
 class UEdGraphNode;
 class FViewport;
 class UDataflowEditor;
+namespace Dataflow
+{
+	class IDataflowConstructionViewMode;
+}
 
 /**
  * The dataflow editor mode is the mode used in the cloth asset editor. It holds most of the inter-tool state.
@@ -53,9 +57,9 @@ public:
 	/**
 	* Construction View Mode
 	*/
-	void SetConstructionViewMode(Dataflow::EDataflowPatternVertexType InMode);
-	Dataflow::EDataflowPatternVertexType GetConstructionViewMode() const;
-	bool CanChangeConstructionViewModeTo(Dataflow::EDataflowPatternVertexType NewViewMode) const;
+	void SetConstructionViewMode(const FName& NewViewModeName);
+	const Dataflow::IDataflowConstructionViewMode* GetConstructionViewMode() const;
+	bool CanChangeConstructionViewModeTo(const FName& NewViewModeName) const;
 
 	void ToggleConstructionViewWireframe();
 	bool CanSetConstructionViewWireframeActive() const;
@@ -175,13 +179,14 @@ private:
 	/** Graph editor is required to add nodes once some tools are complete */
 	TWeakPtr<SDataflowGraphEditor> DataflowGraphEditor;
 
-	Dataflow::EDataflowPatternVertexType ConstructionViewMode = Dataflow::EDataflowPatternVertexType::Sim3D;
+	// Points to a view mode owned by FRenderingViewModeFactory
+	const Dataflow::IDataflowConstructionViewMode* ConstructionViewMode = nullptr;
 
 	// Whether we should restore the previous view mode when a tool ends
 	bool bShouldRestoreSavedConstructionViewMode = false;
 	
 	// The Construction view mode that was active before starting the current tool. When the tool ends, restore this view mode if bShouldRestoreSavedConstructionViewMode is true.
-	Dataflow::EDataflowPatternVertexType SavedConstructionViewMode;
+	const FName SavedConstructionViewMode;
 
 	bool bConstructionViewWireframe = false;
 	bool bShouldRestoreConstructionViewWireframe = false;
diff --git a/Engine/Plugins/Experimental/Dataflow/Source/DataflowEditor/Public/DataflowEditorTools/DataflowEditorToolBuilder.h b/Engine/Plugins/Experimental/Dataflow/Source/DataflowEditor/Public/DataflowEditorTools/DataflowEditorToolBuilder.h
--- a/Engine/Plugins/Experimental/Dataflow/Source/DataflowEditor/Public/DataflowEditorTools/DataflowEditorToolBuilder.h
+++ b/Engine/Plugins/Experimental/Dataflow/Source/DataflowEditor/Public/DataflowEditorTools/DataflowEditorToolBuilder.h
@@ -7,12 +7,10 @@
 
 #include "DataflowEditorToolBuilder.generated.h"
 
-namespace Dataflow
-{
-	enum class EDataflowPatternVertexType : uint8;
+namespace Dataflow {
+	class IDataflowConstructionViewMode;
 }
 
-
 UINTERFACE(MinimalAPI)
 class UDataflowEditorToolBuilder : public UInterface
 {
@@ -27,7 +25,7 @@ class IDataflowEditorToolBuilder
 public:
 
 	/** Returns all Construction View modes that this tool can operate in. The first element should be the preferred mode to switch to if necessary. */
-	virtual void GetSupportedViewModes(TArray<Dataflow::EDataflowPatternVertexType>& Modes) const = 0;
+	virtual void GetSupportedViewModes(TArray<Dataflow::IDataflowConstructionViewMode*>& Modes) const = 0;
 
 	/** Returns whether or not view can be set to wireframe when this tool is active.. */
 	virtual bool CanSetConstructionViewWireframeActive() const { return true; }
diff --git a/Engine/Plugins/Experimental/Dataflow/Source/DataflowEditor/Public/DataflowEditorTools/DataflowEditorWeightMapPaintTool.h b/Engine/Plugins/Experimental/Dataflow/Source/DataflowEditor/Public/DataflowEditorTools/DataflowEditorWeightMapPaintTool.h
--- a/Engine/Plugins/Experimental/Dataflow/Source/DataflowEditor/Public/DataflowEditorTools/DataflowEditorWeightMapPaintTool.h
+++ b/Engine/Plugins/Experimental/Dataflow/Source/DataflowEditor/Public/DataflowEditorTools/DataflowEditorWeightMapPaintTool.h
@@ -63,7 +63,7 @@ public:
 	void SetEditorMode(UDataflowEditorMode* InMode) { Mode = InMode; }
 
 private:
-	virtual void GetSupportedViewModes(TArray<Dataflow::EDataflowPatternVertexType>& Modes) const override;
+	virtual void GetSupportedViewModes(TArray<Dataflow::IDataflowConstructionViewMode*>& Modes) const override;
 	virtual bool CanBuildTool(const FToolBuilderState& SceneState) const override;
 	virtual UMeshSurfacePointTool* CreateNewTool(const FToolBuilderState& SceneState) const override;
 	virtual bool CanSetConstructionViewWireframeActive() const { return false; }
diff --git a/Engine/Plugins/Experimental/Dataflow/Source/DataflowEnginePlugin/Private/Dataflow/DataflowComponent.cpp b/Engine/Plugins/Experimental/Dataflow/Source/DataflowEnginePlugin/Private/Dataflow/DataflowComponent.cpp
--- a/Engine/Plugins/Experimental/Dataflow/Source/DataflowEnginePlugin/Private/Dataflow/DataflowComponent.cpp
+++ b/Engine/Plugins/Experimental/Dataflow/Source/DataflowEnginePlugin/Private/Dataflow/DataflowComponent.cpp
@@ -58,9 +58,9 @@ void UDataflowComponent::BuildRenderCollection()
 	{
 		for (const UDataflowEdNode* Target : RenderTargets)
 		{
-			if (Target)
+			if (Target && ViewMode)
 			{
-				Target->Render(Facade, Context);
+				Target->Render(Facade, Context.ToSharedRef(), *ViewMode);
 			}
 		}
 	}
@@ -84,9 +84,9 @@ void UDataflowComponent::TickComponent(float DeltaTime, enum ELevelTick TickType
 		{
 			for (const UDataflowEdNode* Target : RenderTargets)
 			{
-				if (Target)
+				if (Target && ViewMode)
 				{
-					Target->Render(Facade, Context);
+					Target->Render(Facade, Context.ToSharedRef(), *ViewMode);
 				}
 			}
 		}
diff --git a/Engine/Plugins/Experimental/Dataflow/Source/DataflowEnginePlugin/Private/Dataflow/DataflowEngineRendering.cpp b/Engine/Plugins/Experimental/Dataflow/Source/DataflowEnginePlugin/Private/Dataflow/DataflowEngineRendering.cpp
--- a/Engine/Plugins/Experimental/Dataflow/Source/DataflowEnginePlugin/Private/Dataflow/DataflowEngineRendering.cpp
+++ b/Engine/Plugins/Experimental/Dataflow/Source/DataflowEnginePlugin/Private/Dataflow/DataflowEngineRendering.cpp
@@ -249,102 +249,107 @@ namespace Dataflow
 		}
 	}
 
-	void RenderingCallbacks()
+	
+
+	class FGeometryCollectionSurfaceRenderCallbacks : public FRenderingFactory::ICallbackInterface
 	{
-		using namespace Dataflow;
+		virtual Dataflow::FRenderKey GetRenderKey() const override
+		{
+			return { "SurfaceRender", FGeometryCollection::StaticType() };
+		}
+
+		virtual bool CanRender(const Dataflow::IDataflowConstructionViewMode& ViewMode) const override
+		{
+			return (ViewMode.GetName() == FDataflowConstruction3DViewMode::Name);
+		}
 
-		/**
-		* DataflowNode (FGeometryCollection) Rendering
-		*
-		*		@param Name : "SurfaceRender"
-		*		@param Type : FGeometryCollection::StaticType()
-		* 
-		*		@param Outputs : {FManagedArrayCollection : "Collection"}
-		*/
-		FRenderingFactory::GetInstance()->RegisterOutput({ "SurfaceRender", FGeometryCollection::StaticType()},
-			[](GeometryCollection::Facades::FRenderingFacade& RenderCollection, const Dataflow::FGraphRenderingState& State)
+		virtual void Render(GeometryCollection::Facades::FRenderingFacade& RenderCollection, const FGraphRenderingState& State) override
+		{
+			if (State.GetRenderOutputs().Num())
 			{
-				if (State.GetRenderOutputs().Num())
+				FManagedArrayCollection Default;
+				FName PrimaryOutput = State.GetRenderOutputs()[0]; // "Collection"
+				const FManagedArrayCollection& Collection = State.GetValue<FManagedArrayCollection>(PrimaryOutput, Default);
+
+				const bool bFoundIndices = Collection.FindAttributeTyped<FIntVector>("Indices", FGeometryCollection::FacesGroup) != nullptr;
+				const bool bFoundVertices = Collection.FindAttributeTyped<FVector3f>("Vertex", FGeometryCollection::VerticesGroup) != nullptr;
+				const bool bFoundTransforms = Collection.FindAttributeTyped<FTransform3f>(FTransformCollection::TransformAttribute, FTransformCollection::TransformGroup) != nullptr;
+				const bool bFoundBoneMap = Collection.FindAttributeTyped<int32>("BoneMap", FGeometryCollection::VerticesGroup) != nullptr;
+				const bool bFoundParents = Collection.FindAttributeTyped<int32>(FTransformCollection::ParentAttribute, FTransformCollection::TransformGroup) != nullptr;
+				UE_LOG(LogTemp, Warning, TEXT("Render GC with found params = %d %d %d %d %d"), bFoundIndices, bFoundVertices, bFoundTransforms, bFoundBoneMap, bFoundParents);
+				bool bFoundRenderData = bFoundIndices && bFoundVertices && bFoundTransforms && bFoundBoneMap && bFoundParents
+					&& Collection.NumElements(FTransformCollection::TransformGroup) > 0;
+
+				const bool bFoundVertexStart = Collection.FindAttributeTyped<int32>("VertexStart", FGeometryCollection::GeometryGroup) != nullptr;
+				const bool bFoundVertexCount = Collection.FindAttributeTyped<int32>("VertexCount", FGeometryCollection::GeometryGroup) != nullptr;
+				const bool bFoundFaceStart = Collection.FindAttributeTyped<int32>("FaceStart", FGeometryCollection::GeometryGroup) != nullptr;
+				const bool bFoundFaceCount = Collection.FindAttributeTyped<int32>("FaceCount", FGeometryCollection::GeometryGroup) != nullptr;
+				UE_LOG(LogTemp, Warning, TEXT("Render GC with found mesh group params = %d %d %d %d"), bFoundVertexStart, bFoundVertexCount, bFoundFaceStart, bFoundFaceCount);
+				bool bFoundGeometryAttributes = bFoundVertexStart && bFoundVertexCount && bFoundFaceStart && bFoundFaceCount
+					&& Collection.NumElements(FGeometryCollection::GeometryGroup) > 0;
+
+				if (bFoundRenderData && bFoundGeometryAttributes)
 				{
-					FManagedArrayCollection Default;
-					FName PrimaryOutput = State.GetRenderOutputs()[0]; // "Collection"
-					const FManagedArrayCollection& Collection = State.GetValue<FManagedArrayCollection>(PrimaryOutput, Default);
-
-					const bool bFoundIndices = Collection.FindAttributeTyped<FIntVector>("Indices", FGeometryCollection::FacesGroup) != nullptr;
-					const bool bFoundVertices = Collection.FindAttributeTyped<FVector3f>("Vertex", FGeometryCollection::VerticesGroup) != nullptr;
-					const bool bFoundTransforms = Collection.FindAttributeTyped<FTransform3f>(FTransformCollection::TransformAttribute, FTransformCollection::TransformGroup) != nullptr;
-					const bool bFoundBoneMap = Collection.FindAttributeTyped<int32>("BoneMap", FGeometryCollection::VerticesGroup) != nullptr;
-					const bool bFoundParents = Collection.FindAttributeTyped<int32>(FTransformCollection::ParentAttribute, FTransformCollection::TransformGroup) != nullptr;
-					UE_LOG(LogTemp, Warning, TEXT("Render GC with found params = %d %d %d %d %d"), bFoundIndices, bFoundVertices, bFoundTransforms, bFoundBoneMap, bFoundParents);
-					bool bFoundRenderData = bFoundIndices && bFoundVertices && bFoundTransforms && bFoundBoneMap && bFoundParents 
-						&& Collection.NumElements(FTransformCollection::TransformGroup) > 0;
-
-					const bool bFoundVertexStart = Collection.FindAttributeTyped<int32>("VertexStart", FGeometryCollection::GeometryGroup) != nullptr;
-					const bool bFoundVertexCount = Collection.FindAttributeTyped<int32>("VertexCount", FGeometryCollection::GeometryGroup) != nullptr;
-					const bool bFoundFaceStart = Collection.FindAttributeTyped<int32>("FaceStart", FGeometryCollection::GeometryGroup) != nullptr;
-					const bool bFoundFaceCount = Collection.FindAttributeTyped<int32>("FaceCount", FGeometryCollection::GeometryGroup) != nullptr;
-					UE_LOG(LogTemp, Warning, TEXT("Render GC with found mesh group params = %d %d %d %d"), bFoundVertexStart, bFoundVertexCount, bFoundFaceStart, bFoundFaceCount);
-					bool bFoundGeometryAttributes = bFoundVertexStart && bFoundVertexCount && bFoundFaceStart && bFoundFaceCount
-						&& Collection.NumElements(FGeometryCollection::GeometryGroup) > 0;
-
-					if (bFoundRenderData && bFoundGeometryAttributes )
-					{
-						RenderMeshIndexedGeometryCollection(RenderCollection, State);
-					}
-					else if (bFoundRenderData)
-					{
-						RenderBasicGeometryCollection(RenderCollection, State);
-					}
+					RenderMeshIndexedGeometryCollection(RenderCollection, State);
+				}
+				else if (bFoundRenderData)
+				{
+					RenderBasicGeometryCollection(RenderCollection, State);
 				}
-			});
-
-
-		/**
-		* DataflowNode (FGeometryCollection) Rendering a vertex weight map
-		*
-		*		@param Name : "SurfaceWeightsRender"
-		*		@param Type : FGeometryCollection::StaticType()
-		*		@param Outputs : {FManagedArrayCollection : "Collection"}
-		*		@param Outputs : {FAttributeKey : "AttributeKey"}
-		*/
-		FRenderingFactory::GetInstance()->RegisterOutput({ "SurfaceWeightsRender", FGeometryCollection::StaticType() },
-			[](GeometryCollection::Facades::FRenderingFacade& RenderCollection, const Dataflow::FGraphRenderingState& State)
+			}
+		}
+	};
+
+
+	class FGeometryCollectionSurfaceWeightsRenderCallbacks : public FRenderingFactory::ICallbackInterface
+	{
+		virtual Dataflow::FRenderKey GetRenderKey() const override
+		{
+			return { "SurfaceWeightsRender", FGeometryCollection::StaticType() };
+		}
+
+		virtual bool CanRender(const Dataflow::IDataflowConstructionViewMode& ViewMode) const override
+		{
+			return (ViewMode.GetName() == FDataflowConstruction3DViewMode::Name);
+		}
+
+		virtual void Render(GeometryCollection::Facades::FRenderingFacade& RenderCollection, const FGraphRenderingState& State) override
+		{
+			if (State.GetRenderOutputs().Num() >= 2)
 			{
-				if (State.GetRenderOutputs().Num()>=2)
+				FManagedArrayCollection Default;
+				FName PrimaryOutput = State.GetRenderOutputs()[0]; // "Collection"
+				const FManagedArrayCollection& Collection = State.GetValue<FManagedArrayCollection>(PrimaryOutput, Default);
+
+				const bool bFoundIndices = Collection.FindAttributeTyped<FIntVector>("Indices", FGeometryCollection::FacesGroup) != nullptr;
+				const bool bFoundVertices = Collection.FindAttributeTyped<FVector3f>("Vertex", FGeometryCollection::VerticesGroup) != nullptr;
+				const bool bFoundTransforms = Collection.FindAttributeTyped<FTransform3f>(FTransformCollection::TransformAttribute, FTransformCollection::TransformGroup) != nullptr;
+				const bool bFoundBoneMap = Collection.FindAttributeTyped<int32>("BoneMap", FGeometryCollection::VerticesGroup) != nullptr;
+				const bool bFoundParents = Collection.FindAttributeTyped<int32>(FTransformCollection::ParentAttribute, FTransformCollection::TransformGroup) != nullptr;
+				UE_LOG(LogTemp, Warning, TEXT("Render GC with found params = %d %d %d %d %d"), bFoundIndices, bFoundVertices, bFoundTransforms, bFoundBoneMap, bFoundParents);
+				bool bFoundRenderData = bFoundIndices && bFoundVertices && bFoundTransforms && bFoundBoneMap && bFoundParents
+					&& Collection.NumElements(FTransformCollection::TransformGroup) > 0;
+
+				const bool bFoundVertexStart = Collection.FindAttributeTyped<int32>("VertexStart", FGeometryCollection::GeometryGroup) != nullptr;
+				const bool bFoundVertexCount = Collection.FindAttributeTyped<int32>("VertexCount", FGeometryCollection::GeometryGroup) != nullptr;
+				const bool bFoundFaceStart = Collection.FindAttributeTyped<int32>("FaceStart", FGeometryCollection::GeometryGroup) != nullptr;
+				const bool bFoundFaceCount = Collection.FindAttributeTyped<int32>("FaceCount", FGeometryCollection::GeometryGroup) != nullptr;
+				UE_LOG(LogTemp, Warning, TEXT("Render GC with found mesh group params = %d %d %d %d"), bFoundVertexStart, bFoundVertexCount, bFoundFaceStart, bFoundFaceCount);
+				bool bFoundGeometryAttributes = bFoundVertexStart && bFoundVertexCount && bFoundFaceStart && bFoundFaceCount
+					&& Collection.NumElements(FGeometryCollection::GeometryGroup) > 0;
+
+				FCollectionAttributeKey DefaultKey;
+				FName SecondaryOutput = State.GetRenderOutputs()[1]; // "AttributeKey"
+				const FCollectionAttributeKey& AttributeKey = State.GetValue<FCollectionAttributeKey>(SecondaryOutput, DefaultKey);
+
+				const bool bFoundVertexColor = Collection.FindAttributeTyped<FLinearColor>("Color", FGeometryCollection::VerticesGroup) != nullptr;
+				const bool bFoundFloatScalar = Collection.FindAttributeTyped<float>(FName(AttributeKey.Attribute), FName(AttributeKey.Group)) != nullptr;
+				bool bFoundVertexScalarAndColors = bFoundVertexColor && bFoundFloatScalar && AttributeKey.Group.Equals(FGeometryCollection::VerticesGroup.ToString());
+
+				TArray<FLinearColor>* Colors = nullptr;
+				if (bFoundVertexScalarAndColors)
 				{
-					FManagedArrayCollection Default;
-					FName PrimaryOutput = State.GetRenderOutputs()[0]; // "Collection"
-					const FManagedArrayCollection & Collection = State.GetValue<FManagedArrayCollection>(PrimaryOutput, Default);
-
-					const bool bFoundIndices = Collection.FindAttributeTyped<FIntVector>("Indices", FGeometryCollection::FacesGroup) != nullptr;
-					const bool bFoundVertices = Collection.FindAttributeTyped<FVector3f>("Vertex", FGeometryCollection::VerticesGroup) != nullptr;
-					const bool bFoundTransforms = Collection.FindAttributeTyped<FTransform3f>(FTransformCollection::TransformAttribute, FTransformCollection::TransformGroup) != nullptr;
-					const bool bFoundBoneMap = Collection.FindAttributeTyped<int32>("BoneMap", FGeometryCollection::VerticesGroup) != nullptr;
-					const bool bFoundParents = Collection.FindAttributeTyped<int32>(FTransformCollection::ParentAttribute, FTransformCollection::TransformGroup) != nullptr;
-					UE_LOG(LogTemp, Warning, TEXT("Render GC with found params = %d %d %d %d %d"), bFoundIndices, bFoundVertices, bFoundTransforms, bFoundBoneMap, bFoundParents);
-					bool bFoundRenderData = bFoundIndices && bFoundVertices && bFoundTransforms && bFoundBoneMap && bFoundParents
-						&& Collection.NumElements(FTransformCollection::TransformGroup) > 0;
-
-					const bool bFoundVertexStart = Collection.FindAttributeTyped<int32>("VertexStart", FGeometryCollection::GeometryGroup) != nullptr;
-					const bool bFoundVertexCount = Collection.FindAttributeTyped<int32>("VertexCount", FGeometryCollection::GeometryGroup) != nullptr;
-					const bool bFoundFaceStart = Collection.FindAttributeTyped<int32>("FaceStart", FGeometryCollection::GeometryGroup) != nullptr;
-					const bool bFoundFaceCount = Collection.FindAttributeTyped<int32>("FaceCount", FGeometryCollection::GeometryGroup) != nullptr;
-					UE_LOG(LogTemp, Warning, TEXT("Render GC with found mesh group params = %d %d %d %d"), bFoundVertexStart, bFoundVertexCount, bFoundFaceStart, bFoundFaceCount);
-					bool bFoundGeometryAttributes = bFoundVertexStart && bFoundVertexCount && bFoundFaceStart && bFoundFaceCount
-						&& Collection.NumElements(FGeometryCollection::GeometryGroup) > 0;
-
-					FCollectionAttributeKey DefaultKey;
-					FName SecondaryOutput = State.GetRenderOutputs()[1]; // "AttributeKey"
-					const FCollectionAttributeKey& AttributeKey = State.GetValue<FCollectionAttributeKey>(SecondaryOutput, DefaultKey);
-
-					const bool bFoundVertexColor = Collection.FindAttributeTyped<FLinearColor>("Color", FGeometryCollection::VerticesGroup) != nullptr;
-					const bool bFoundFloatScalar = Collection.FindAttributeTyped<float>(FName(AttributeKey.Attribute), FName(AttributeKey.Group)) != nullptr;
-					bool bFoundVertexScalarAndColors = bFoundVertexColor && bFoundFloatScalar && AttributeKey.Group.Equals(FGeometryCollection::VerticesGroup.ToString());
-
-					TArray<FLinearColor>* Colors = nullptr;
-					if (bFoundVertexScalarAndColors)
-					{
-						auto RangeValue = [](const TManagedArray<float>* FloatArray)
+					auto RangeValue = [](const TManagedArray<float>* FloatArray)
 						{
 							float Min = FLT_MAX;
 							float Max = -FLT_MAX;
@@ -355,225 +360,247 @@ namespace Dataflow
 							return TPair<float, float>(Min, Max);
 						};
 
-						const TManagedArray<float>* FloatArray = Collection.FindAttributeTyped<float>(FName(AttributeKey.Attribute), FName(AttributeKey.Group));
-						if (FloatArray && FloatArray->Num())
-						{
-							Colors = new TArray<FLinearColor>();
-							Colors->AddUninitialized(FloatArray->Num());
+					const TManagedArray<float>* FloatArray = Collection.FindAttributeTyped<float>(FName(AttributeKey.Attribute), FName(AttributeKey.Group));
+					if (FloatArray && FloatArray->Num())
+					{
+						Colors = new TArray<FLinearColor>();
+						Colors->AddUninitialized(FloatArray->Num());
 
-							TPair<float, float> Range = RangeValue(FloatArray);
-							float Delta = FMath::Abs(Range.Get<1>() - Range.Get<0>());
-							if (Delta > FLT_EPSILON)
+						TPair<float, float> Range = RangeValue(FloatArray);
+						float Delta = FMath::Abs(Range.Get<1>() - Range.Get<0>());
+						if (Delta > FLT_EPSILON)
+						{
+							for (int32 VertexIdx = 0; VertexIdx < FloatArray->Num(); ++VertexIdx)
 							{
-								for (int32 VertexIdx = 0; VertexIdx < FloatArray->Num(); ++VertexIdx)
-								{
-									(*Colors)[VertexIdx] = FLinearColor::White * ((*FloatArray)[VertexIdx] - Range.Get<0>()) / Delta;
-								}
+								(*Colors)[VertexIdx] = FLinearColor::White * ((*FloatArray)[VertexIdx] - Range.Get<0>()) / Delta;
 							}
-							else
+						}
+						else
+						{
+							for (int32 VertexIdx = 0; VertexIdx < FloatArray->Num(); ++VertexIdx)
 							{
-								for (int32 VertexIdx = 0; VertexIdx < FloatArray->Num(); ++VertexIdx)
-								{
-									(*Colors)[VertexIdx] = FLinearColor::Black;
-								}
+								(*Colors)[VertexIdx] = FLinearColor::Black;
 							}
 						}
 					}
+				}
 
-					if (bFoundRenderData && bFoundGeometryAttributes)
-					{
-						RenderMeshIndexedGeometryCollection(RenderCollection, State, Colors);
-					}
-					else if (bFoundRenderData)
-					{
-						RenderBasicGeometryCollection(RenderCollection, State, Colors);
-					}
+				if (bFoundRenderData && bFoundGeometryAttributes)
+				{
+					RenderMeshIndexedGeometryCollection(RenderCollection, State, Colors);
+				}
+				else if (bFoundRenderData)
+				{
+					RenderBasicGeometryCollection(RenderCollection, State, Colors);
+				}
 
-					if (Colors)
-					{
-						delete Colors;
-					}
+				if (Colors)
+				{
+					delete Colors;
 				}
-			});
-
-
-		/**
-		* DataflowNode (FDynamicMesh3) Rendering
-		*
-		*		@param Name : "SurfaceRender"
-		*		@param Type : FName("FDynamicMesh3")
-		*
-		*		@param Outputs : {FDynamicMesh3 : "Mesh"}
-		*/
-		FRenderingFactory::GetInstance()->RegisterOutput({ "SurfaceRender", FName("FDynamicMesh3") },
-			[](GeometryCollection::Facades::FRenderingFacade& RenderCollection, const Dataflow::FGraphRenderingState& State)
+			}
+		}
+	};
+
+	class FDynamicMesh3SurfaceRenderCallbacks : public FRenderingFactory::ICallbackInterface
+	{
+		virtual Dataflow::FRenderKey GetRenderKey() const override
+		{
+			return { "SurfaceRender", FName("FDynamicMesh3") };
+		}
+
+		virtual bool CanRender(const Dataflow::IDataflowConstructionViewMode& ViewMode) const override
+		{
+			return (ViewMode.GetName() == FDataflowConstruction3DViewMode::Name);
+		}
+
+		virtual void Render(GeometryCollection::Facades::FRenderingFacade& RenderCollection, const FGraphRenderingState& State) override
+		{
+			if (State.GetRenderOutputs().Num())
 			{
-				if (State.GetRenderOutputs().Num())
+				FName PrimaryOutput = State.GetRenderOutputs()[0]; // "Mesh"
+
+				TObjectPtr<UDynamicMesh> Default;
+				if (const TObjectPtr<UDynamicMesh> Mesh = State.GetValue<TObjectPtr<UDynamicMesh>>(PrimaryOutput, Default))
 				{
-					FName PrimaryOutput = State.GetRenderOutputs()[0]; // "Mesh"
+					const UE::Geometry::FDynamicMesh3& DynamicMesh = Mesh->GetMeshRef();
 
-					TObjectPtr<UDynamicMesh> Default;
-					if (const TObjectPtr<UDynamicMesh> Mesh = State.GetValue<TObjectPtr<UDynamicMesh>>(PrimaryOutput, Default))
-					{
-						const UE::Geometry::FDynamicMesh3& DynamicMesh = Mesh->GetMeshRef();
+					const int32 NumVertices = DynamicMesh.VertexCount();
+					const int32 NumTriangles = DynamicMesh.TriangleCount();
 
-						const int32 NumVertices = DynamicMesh.VertexCount();
-						const int32 NumTriangles = DynamicMesh.TriangleCount();
+					if (NumVertices > 0 && NumTriangles > 0)
+					{
+						// This will contain the valid triangles only
+						TArray<FIntVector> Tris; Tris.Reserve(DynamicMesh.TriangleCount());
 
-						if (NumVertices > 0 && NumTriangles > 0)
+						// DynamicMesh.TrianglesItr() returns the valid triangles only
+						for (UE::Geometry::FIndex3i Tri : DynamicMesh.TrianglesItr())
 						{
-							// This will contain the valid triangles only
-							TArray<FIntVector> Tris; Tris.Reserve(DynamicMesh.TriangleCount());
-														
-							// DynamicMesh.TrianglesItr() returns the valid triangles only
-							for (UE::Geometry::FIndex3i Tri : DynamicMesh.TrianglesItr())
-							{
-								Tris.Add(FIntVector(Tri.A, Tri.B, Tri.C));
-							}
-
-							// This will contain all the vertices (invalid ones too)
-							// Otherwise the IDs need to be remaped
-							TArray<FVector3f> Vertices; Vertices.AddZeroed(DynamicMesh.MaxVertexID());
+							Tris.Add(FIntVector(Tri.A, Tri.B, Tri.C));
+						}
 
-							// DynamicMesh.VertexIndicesItr() returns the valid vertices only
-							for (int32 VertexID : DynamicMesh.VertexIndicesItr())
-							{
-								Vertices[VertexID] = (FVector3f)DynamicMesh.GetVertex(VertexID);
-							}
+						// This will contain all the vertices (invalid ones too)
+						// Otherwise the IDs need to be remaped
+						TArray<FVector3f> Vertices; Vertices.AddZeroed(DynamicMesh.MaxVertexID());
 
-							// Add VertexNormal and VertexColor
-							TArray<FVector3f> VertexNormals; VertexNormals.AddUninitialized(Vertices.Num());
-							TArray<FLinearColor> VertexColors; VertexColors.AddUninitialized(Vertices.Num());
-							for (int32 VertexIdx = 0; VertexIdx < VertexNormals.Num(); ++VertexIdx)
-							{
-								// TODO: Get the normal from FDynamicMesh3
-								VertexNormals[VertexIdx] = FVector3f(0.f);
-								VertexColors[VertexIdx] = FLinearColor(IDataflowEnginePlugin::SurfaceColor);
-							}
+						// DynamicMesh.VertexIndicesItr() returns the valid vertices only
+						for (int32 VertexID : DynamicMesh.VertexIndicesItr())
+						{
+							Vertices[VertexID] = (FVector3f)DynamicMesh.GetVertex(VertexID);
+						}
 
-							int32 GeometryIndex = RenderCollection.StartGeometryGroup(State.GetGuid().ToString());
-							RenderCollection.AddSurface(MoveTemp(Vertices), MoveTemp(Tris), MoveTemp(VertexNormals), MoveTemp(VertexColors));
-							RenderCollection.EndGeometryGroup(GeometryIndex);
+						// Add VertexNormal and VertexColor
+						TArray<FVector3f> VertexNormals; VertexNormals.AddUninitialized(Vertices.Num());
+						TArray<FLinearColor> VertexColors; VertexColors.AddUninitialized(Vertices.Num());
+						for (int32 VertexIdx = 0; VertexIdx < VertexNormals.Num(); ++VertexIdx)
+						{
+							// TODO: Get the normal from FDynamicMesh3
+							VertexNormals[VertexIdx] = FVector3f(0.f);
+							VertexColors[VertexIdx] = FLinearColor(IDataflowEnginePlugin::SurfaceColor);
 						}
+
+						int32 GeometryIndex = RenderCollection.StartGeometryGroup(State.GetGuid().ToString());
+						RenderCollection.AddSurface(MoveTemp(Vertices), MoveTemp(Tris), MoveTemp(VertexNormals), MoveTemp(VertexColors));
+						RenderCollection.EndGeometryGroup(GeometryIndex);
 					}
 				}
-			});
-
-
-		/**
-		* DataflowNode (FBox) Rendering
-		*
-		*		@param Name : "SurfaceRender"
-		*		@param Type : FName("FBox")
-		*
-		*		@param Outputs : {FBox : "Box"}
-		*/
-		FRenderingFactory::GetInstance()->RegisterOutput({ "SurfaceRender", FName("FBox") },
-			[](GeometryCollection::Facades::FRenderingFacade& RenderCollection, const Dataflow::FGraphRenderingState& State)
+			}
+		}
+	};
+
+
+	class FBoxSurfaceRenderCallbacks : public FRenderingFactory::ICallbackInterface
+	{
+		virtual Dataflow::FRenderKey GetRenderKey() const override
+		{
+			return { "SurfaceRender", FName("FBox") };
+		}
+
+		virtual bool CanRender(const Dataflow::IDataflowConstructionViewMode& ViewMode) const override
+		{
+			return (ViewMode.GetName() == FDataflowConstruction3DViewMode::Name);
+		}
+
+		virtual void Render(GeometryCollection::Facades::FRenderingFacade& RenderCollection, const FGraphRenderingState& State) override
+		{
+			if (State.GetRenderOutputs().Num())
 			{
-				if (State.GetRenderOutputs().Num())
+				FName PrimaryOutput = State.GetRenderOutputs()[0]; // "Box"
+
+				FBox Default(ForceInit);
+				const FBox& Box = State.GetValue<FBox>(PrimaryOutput, Default);
+
+				const int32 NumVertices = 8;
+				const int32 NumTriangles = 12;
+
+				TArray<FVector3f> Vertices; Vertices.AddUninitialized(NumVertices);
+				TArray<FIntVector> Tris; Tris.AddUninitialized(NumTriangles);
+
+				FVector Min = Box.Min;
+				FVector Max = Box.Max;
+
+				// Add vertices
+				Vertices[0] = FVector3f(Min);
+				Vertices[1] = FVector3f(Max.X, Min.Y, Min.Z);
+				Vertices[2] = FVector3f(Max.X, Max.Y, Min.Z);
+				Vertices[3] = FVector3f(Min.X, Max.Y, Min.Z);
+				Vertices[4] = FVector3f(Min.X, Min.Y, Max.Z);
+				Vertices[5] = FVector3f(Max.X, Min.Y, Max.Z);
+				Vertices[6] = FVector3f(Max);
+				Vertices[7] = FVector3f(Min.X, Max.Y, Max.Z);
+
+				// Add triangles
+				Tris[0] = FIntVector(0, 1, 3); Tris[1] = FIntVector(1, 2, 3);
+				Tris[2] = FIntVector(0, 4, 1); Tris[3] = FIntVector(4, 5, 1);
+				Tris[4] = FIntVector(5, 2, 1); Tris[5] = FIntVector(5, 6, 2);
+				Tris[6] = FIntVector(3, 2, 6); Tris[7] = FIntVector(7, 3, 6);
+				Tris[8] = FIntVector(0, 3, 7); Tris[9] = FIntVector(4, 0, 7);
+				Tris[10] = FIntVector(5, 4, 7); Tris[11] = FIntVector(5, 7, 6);
+
+				TArray<FVector3f> VertexNormals; VertexNormals.AddUninitialized(NumVertices);
+				// TODO: Compute vertex normals
+
+				// Add VertexNormal and VertexColor
+				TArray<FLinearColor> VertexColors; VertexColors.AddUninitialized(Vertices.Num());
+				for (int32 VertexIdx = 0; VertexIdx < VertexNormals.Num(); ++VertexIdx)
 				{
-					FName PrimaryOutput = State.GetRenderOutputs()[0]; // "Box"
+					VertexNormals[VertexIdx] = FVector3f(0.f);
+					VertexColors[VertexIdx] = FLinearColor(IDataflowEnginePlugin::SurfaceColor);
+				}
+
+				int32 GeometryIndex = RenderCollection.StartGeometryGroup(State.GetGuid().ToString());
+				RenderCollection.AddSurface(MoveTemp(Vertices), MoveTemp(Tris), MoveTemp(VertexNormals), MoveTemp(VertexColors));
+				RenderCollection.EndGeometryGroup(GeometryIndex);
+			}
+		}
+	};
 
-					FBox Default(ForceInit);
-					const FBox& Box = State.GetValue<FBox>(PrimaryOutput, Default);
 
-					const int32 NumVertices = 8;
-					const int32 NumTriangles = 12;
+	class FFieldVolumeRenderCallbacks : public FRenderingFactory::ICallbackInterface
+	{
+		virtual Dataflow::FRenderKey GetRenderKey() const override
+		{
+			return { "VolumeRender", FFieldCollection::StaticType() };
+		}
 
-					TArray<FVector3f> Vertices; Vertices.AddUninitialized(NumVertices);
-					TArray<FIntVector> Tris; Tris.AddUninitialized(NumTriangles);
+		virtual bool CanRender(const Dataflow::IDataflowConstructionViewMode& ViewMode) const override
+		{
+			return (ViewMode.GetName() == FDataflowConstruction3DViewMode::Name);
+		}
 
-					FVector Min = Box.Min;
-					FVector Max = Box.Max;
-
-					// Add vertices
-					Vertices[0] = FVector3f(Min);
-					Vertices[1] = FVector3f(Max.X, Min.Y, Min.Z);
-					Vertices[2] = FVector3f(Max.X, Max.Y, Min.Z);
-					Vertices[3] = FVector3f(Min.X, Max.Y, Min.Z);
-					Vertices[4] = FVector3f(Min.X, Min.Y, Max.Z);
-					Vertices[5] = FVector3f(Max.X, Min.Y, Max.Z);
-					Vertices[6] = FVector3f(Max);
-					Vertices[7] = FVector3f(Min.X, Max.Y, Max.Z);
-
-					// Add triangles
-					Tris[0] = FIntVector(0, 1, 3); Tris[1] = FIntVector(1, 2, 3);
-					Tris[2] = FIntVector(0, 4, 1); Tris[3] = FIntVector(4, 5, 1);
-					Tris[4] = FIntVector(5, 2, 1); Tris[5] = FIntVector(5, 6, 2);
-					Tris[6] = FIntVector(3, 2, 6); Tris[7] = FIntVector(7, 3, 6);
-					Tris[8] = FIntVector(0, 3, 7); Tris[9] = FIntVector(4, 0, 7);
-					Tris[10] = FIntVector(5, 4, 7); Tris[11] = FIntVector(5, 7, 6);
+		virtual void Render(GeometryCollection::Facades::FRenderingFacade& RenderCollection, const FGraphRenderingState& State) override
+		{
+			if (State.GetRenderOutputs().Num())
+			{
+				FName PrimaryOutput = State.GetRenderOutputs()[0]; // "VectorField"
+				if (PrimaryOutput.IsEqual(FName("VectorField")))
+				{
+					FFieldCollection Default;
+					const FFieldCollection& Collection = State.GetValue<FFieldCollection>(PrimaryOutput, Default);
+					TArray<TPair<FVector3f, FVector3f>> VectorField = Collection.GetVectorField();
+					const int32 NumVertices = 3 * VectorField.Num();
+					const int32 NumTriangles = VectorField.Num();
 
+					TArray<FVector3f> Vertices; Vertices.AddUninitialized(NumVertices);
+					TArray<FIntVector> Tris; Tris.AddUninitialized(NumTriangles);
 					TArray<FVector3f> VertexNormals; VertexNormals.AddUninitialized(NumVertices);
-					// TODO: Compute vertex normals
-					
-					// Add VertexNormal and VertexColor
-					TArray<FLinearColor> VertexColors; VertexColors.AddUninitialized(Vertices.Num());
-					for (int32 VertexIdx = 0; VertexIdx < VertexNormals.Num(); ++VertexIdx)
+					TArray<FLinearColor> VertexColors; VertexColors.AddUninitialized(NumVertices);
+
+					for (int32 i = 0; i < VectorField.Num(); i++)
 					{
-						VertexNormals[VertexIdx] = FVector3f(0.f);
-						VertexColors[VertexIdx] = FLinearColor(IDataflowEnginePlugin::SurfaceColor);
+
+						FVector3f Dir = VectorField[i].Value - VectorField[i].Key;
+						FVector3f DirAdd = Dir;
+						DirAdd.X += 1.f;
+						FVector3f OrthogonalDir = (Dir ^ DirAdd).GetSafeNormal();
+						Tris[i] = FIntVector(3 * i, 3 * i + 1, 3 * i + 2);
+						Vertices[3 * i] = VectorField[i].Key;
+						Vertices[3 * i + 1] = VectorField[i].Value;
+						Vertices[3 * i + 2] = VectorField[i].Key + float(0.1) * Dir.Size() * OrthogonalDir;
+						FVector3f TriangleNormal = (OrthogonalDir ^ Dir).GetSafeNormal();
+						VertexNormals[3 * i] = TriangleNormal;
+						VertexNormals[3 * i + 1] = TriangleNormal;
+						VertexNormals[3 * i + 2] = TriangleNormal;
+						VertexColors[i] = FLinearColor(IDataflowEnginePlugin::SurfaceColor);
 					}
 
 					int32 GeometryIndex = RenderCollection.StartGeometryGroup(State.GetGuid().ToString());
 					RenderCollection.AddSurface(MoveTemp(Vertices), MoveTemp(Tris), MoveTemp(VertexNormals), MoveTemp(VertexColors));
 					RenderCollection.EndGeometryGroup(GeometryIndex);
 				}
-			});
-
-		/**
-		* DataflowNode (FFieldCollection) Rendering
-		*
-		*		@param Name : "VolumeRender"
-		*		@param Type : FName("FFieldCollection")
-		*
-		*		@param Outputs : {FFieldCollection : "VectorField"}
-		*/
-		FRenderingFactory::GetInstance()->RegisterOutput({ "VolumeRender", FFieldCollection::StaticType() },
-			[](GeometryCollection::Facades::FRenderingFacade& RenderCollection, const Dataflow::FGraphRenderingState& State)
-			{
-				if (State.GetRenderOutputs().Num())
-				{
-					FName PrimaryOutput = State.GetRenderOutputs()[0]; // "VectorField"
-					if (PrimaryOutput.IsEqual(FName("VectorField")))
-					{
-						FFieldCollection Default;
-						const FFieldCollection& Collection = State.GetValue<FFieldCollection>(PrimaryOutput, Default);
-						TArray<TPair<FVector3f, FVector3f>> VectorField = Collection.GetVectorField();
-						const int32 NumVertices = 3 * VectorField.Num();
-						const int32 NumTriangles = VectorField.Num();
-
-						TArray<FVector3f> Vertices; Vertices.AddUninitialized(NumVertices);
-						TArray<FIntVector> Tris; Tris.AddUninitialized(NumTriangles);
-						TArray<FVector3f> VertexNormals; VertexNormals.AddUninitialized(NumVertices);
-						TArray<FLinearColor> VertexColors; VertexColors.AddUninitialized(NumVertices);
-
-						for (int32 i = 0; i < VectorField.Num(); i++)
-						{
-							
-							FVector3f Dir = VectorField[i].Value - VectorField[i].Key;
-							FVector3f DirAdd = Dir;
-							DirAdd.X += 1.f;
-							FVector3f OrthogonalDir = (Dir^ DirAdd).GetSafeNormal();
-							Tris[i] = FIntVector(3*i, 3*i+1, 3*i+2);
-							Vertices[3*i] = VectorField[i].Key;
-							Vertices[3*i+1] = VectorField[i].Value;
-							Vertices[3*i+2] = VectorField[i].Key + float(0.1) * Dir.Size() * OrthogonalDir;
-							FVector3f TriangleNormal = (OrthogonalDir ^ Dir).GetSafeNormal();
-							VertexNormals[3*i] = TriangleNormal;
-							VertexNormals[3*i+1] = TriangleNormal;
-							VertexNormals[3*i+2] = TriangleNormal;
-							VertexColors[i] = FLinearColor(IDataflowEnginePlugin::SurfaceColor);
-						}
+			}
+		}
+	};
 
-						int32 GeometryIndex = RenderCollection.StartGeometryGroup(State.GetGuid().ToString());
-						RenderCollection.AddSurface(MoveTemp(Vertices), MoveTemp(Tris), MoveTemp(VertexNormals), MoveTemp(VertexColors));
-						RenderCollection.EndGeometryGroup(GeometryIndex);
-					}
-				}
-			});
+	void RenderingCallbacks()
+	{
+		using namespace Dataflow;
 
+		FRenderingFactory::GetInstance()->RegisterCallbacks(MakeUnique<FGeometryCollectionSurfaceRenderCallbacks>());
+		FRenderingFactory::GetInstance()->RegisterCallbacks(MakeUnique<FGeometryCollectionSurfaceWeightsRenderCallbacks>());
+		FRenderingFactory::GetInstance()->RegisterCallbacks(MakeUnique<FDynamicMesh3SurfaceRenderCallbacks>());
+		FRenderingFactory::GetInstance()->RegisterCallbacks(MakeUnique<FBoxSurfaceRenderCallbacks>());
+		FRenderingFactory::GetInstance()->RegisterCallbacks(MakeUnique<FFieldVolumeRenderCallbacks>());
 	}
 
 }
diff --git a/Engine/Plugins/Experimental/Dataflow/Source/DataflowEnginePlugin/Public/Dataflow/DataflowComponent.h b/Engine/Plugins/Experimental/Dataflow/Source/DataflowEnginePlugin/Public/Dataflow/DataflowComponent.h
--- a/Engine/Plugins/Experimental/Dataflow/Source/DataflowEnginePlugin/Public/Dataflow/DataflowComponent.h
+++ b/Engine/Plugins/Experimental/Dataflow/Source/DataflowEnginePlugin/Public/Dataflow/DataflowComponent.h
@@ -10,6 +10,7 @@
 
 #include "DataflowComponent.generated.h"
 
+namespace Dataflow { class IDataflowConstructionViewMode; }
 
 /**
 *	UDataflowComponent
@@ -66,6 +67,13 @@ public:
 		SelectionState = InState; 
 	}
 
+	/* View mode */
+	// NOTE: Currently UDataflowComponent is not used in the Dataflow Editor. Instead the FDataflowConstructionScene converts the FRenderingFacade to a UDynamicMeshComponent.
+	// If we do start using UDataflowComponent we will need to update the current View Mode as it's changed using this function.
+	void SetViewMode(const Dataflow::IDataflowConstructionViewMode* InViewMode)
+	{
+		ViewMode = InViewMode;
+	}
 
 private:
 	TSharedPtr<Dataflow::FContext> Context;
@@ -78,5 +86,6 @@ private:
 	bool bBoundsNeedsUpdate = true;
 	FBoxSphereBounds BoundingBox = FBoxSphereBounds(ForceInitToZero);
 	FDataflowSelectionState SelectionState = FDataflowSelectionState(FDataflowSelectionState::EMode::DSS_Dataflow_None);
+	const Dataflow::IDataflowConstructionViewMode* ViewMode;
 };
 
diff --git a/Engine/Source/Runtime/Experimental/Dataflow/Engine/Private/Dataflow/DataflowEdNode.cpp b/Engine/Source/Runtime/Experimental/Dataflow/Engine/Private/Dataflow/DataflowEdNode.cpp
--- a/Engine/Source/Runtime/Experimental/Dataflow/Engine/Private/Dataflow/DataflowEdNode.cpp
+++ b/Engine/Source/Runtime/Experimental/Dataflow/Engine/Private/Dataflow/DataflowEdNode.cpp
@@ -760,7 +760,7 @@ TArray<Dataflow::FRenderingParameter> UDataflowEdNode::GetRenderParameters() con
 }
 
 
-bool UDataflowEdNode::Render(GeometryCollection::Facades::FRenderingFacade& RenderData, const TSharedPtr<Dataflow::FContext> Context) const
+bool UDataflowEdNode::Render(GeometryCollection::Facades::FRenderingFacade& RenderData, const TSharedRef<Dataflow::FContext> Context, const Dataflow::IDataflowConstructionViewMode& ViewMode) const
 {
 	bool bNeedsRefresh = false;
 	if (DataflowGraph)
@@ -773,7 +773,7 @@ bool UDataflowEdNode::Render(GeometryCollection::Facades::FRenderingFacade& Rend
 				{
 					for (Dataflow::FRenderingParameter& Parameter : GetRenderParameters())
 					{
-						Factory->RenderNodeOutput(RenderData, { GetDataflowNodeGuid(), NodeTarget.Get(), Parameter, *Context});
+						Factory->RenderNodeOutput(RenderData, Dataflow::FGraphRenderingState{ GetDataflowNodeGuid(), NodeTarget.Get(), Parameter, Context.Get(), ViewMode});
 						bNeedsRefresh = true;
 					}
 				}
@@ -783,5 +783,27 @@ bool UDataflowEdNode::Render(GeometryCollection::Facades::FRenderingFacade& Rend
 	return bNeedsRefresh;
 }
 
+
+bool UDataflowEdNode::CanRender(const TSharedRef<Dataflow::FContext> Context, const Dataflow::IDataflowConstructionViewMode& ViewMode) const
+{
+	if (DataflowGraph)
+	{
+		if (TSharedPtr<const FDataflowNode> NodeTarget = DataflowGraph->FindBaseNode(FName(GetName())))
+		{
+			if (const Dataflow::FRenderingFactory* const Factory = Dataflow::FRenderingFactory::GetInstance())
+			{
+				for (const Dataflow::FRenderingParameter& Parameter : GetRenderParameters())
+				{
+					if (Factory->CanRenderNodeOutput(Dataflow::FGraphRenderingState{ GetDataflowNodeGuid(), NodeTarget.Get(), Parameter, Context.Get(), ViewMode }))
+					{
+						return true;
+					}
+				}
+			}
+		}
+	}
+	return false;
+}
+
 #undef LOCTEXT_NAMESPACE
 
diff --git a/Engine/Source/Runtime/Experimental/Dataflow/Engine/Private/Dataflow/DataflowRenderingFactory.cpp b/Engine/Source/Runtime/Experimental/Dataflow/Engine/Private/Dataflow/DataflowRenderingFactory.cpp
--- a/Engine/Source/Runtime/Experimental/Dataflow/Engine/Private/Dataflow/DataflowRenderingFactory.cpp
+++ b/Engine/Source/Runtime/Experimental/Dataflow/Engine/Private/Dataflow/DataflowRenderingFactory.cpp
@@ -10,12 +10,11 @@ namespace Dataflow
 {
 	FRenderingFactory* FRenderingFactory::Instance = nullptr;
 
-
 	void FRenderingFactory::RenderNodeOutput(GeometryCollection::Facades::FRenderingFacade& RenderData, const FGraphRenderingState& State)
 	{
-		if (RenderMap.Contains(State.GetRenderKey()))
+		if (CallbackMap.Contains(State.GetRenderKey()))
 		{
-			RenderMap[State.GetRenderKey()](RenderData, State);
+			CallbackMap[State.GetRenderKey()]->Render(RenderData, State);
 		}
 		else
 		{
@@ -26,5 +25,57 @@ namespace Dataflow
 				*State.GetNodeName().ToString());
 		}
 	}
+
+	bool FRenderingFactory::CanRenderNodeOutput(const FGraphRenderingState& State) const
+	{
+		if (CallbackMap.Contains(State.GetRenderKey()))
+		{
+			return CallbackMap[State.GetRenderKey()]->CanRender(State.GetViewMode());
+		}
+
+		return false;
+	}
+
+	FRenderingFactory* FRenderingFactory::GetInstance()
+	{
+		if (!Instance)
+		{
+			Instance = new FRenderingFactory();
+		}
+		return Instance;
+	}
+
+	void FRenderingFactory::RegisterCallbacks(TUniquePtr<ICallbackInterface> InCallbacks)
+	{
+		const FRenderKey Key = InCallbacks->GetRenderKey();
+		if (CallbackMap.Contains(Key))
+		{
+			UE_LOG(LogChaos, Warning,
+				TEXT("Warning : Dataflow output rendering callback registration conflicts with existing rendering callback (<%s,%s>)"),
+				*Key.Get<0>(),
+				*Key.Get<1>().ToString());
+		}
+		else
+		{
+			CallbackMap.Add(Key, MoveTemp(InCallbacks));
+		}
+	}
+
+	void FRenderingFactory::DeregisterCallbacks(const FRenderKey& Key)
+	{
+		if (!CallbackMap.Contains(Key))
+		{
+			UE_LOG(LogChaos, Warning,
+				TEXT("Warning : Dataflow output rendering callback deregistration. Rendering callback not registered: (<%s,%s>)"),
+				*Key.Get<0>(),
+				*Key.Get<1>().ToString());
+		}
+		else
+		{
+			CallbackMap.Remove(Key);
+		}
+	}
+
+
 }
 
diff --git a/Engine/Source/Runtime/Experimental/Dataflow/Engine/Private/Dataflow/DataflowRenderingViewMode.cpp b/Engine/Source/Runtime/Experimental/Dataflow/Engine/Private/Dataflow/DataflowRenderingViewMode.cpp
new file mode 100644
--- /dev/null
+++ b/Engine/Source/Runtime/Experimental/Dataflow/Engine/Private/Dataflow/DataflowRenderingViewMode.cpp
@@ -0,0 +1,124 @@
+// Copyright Epic Games, Inc. All Rights Reserved.
+
+#include "Dataflow/DataflowRenderingViewMode.h"
+#include "Misc/LazySingleton.h"
+
+#define LOCTEXT_NAMESPACE "DataflowRenderingViewMode"
+
+namespace Dataflow
+{
+	bool IDataflowConstructionViewMode::IsPerspective() const
+	{
+		return GetViewportType() == EDataflowViewportType::DVT_Perspective;
+	}
+
+	EDataflowViewportType FDataflowConstruction2DViewModeBase::GetViewportType() const
+	{
+		return EDataflowViewportType::DVT_OrthoNegativeXY;
+	}
+
+	EDataflowViewportType FDataflowConstruction3DViewModeBase::GetViewportType() const
+	{
+		return EDataflowViewportType::DVT_Perspective;
+	}
+
+	FName FDataflowConstruction2DViewMode::Name = FName("2DView");
+
+	FName FDataflowConstruction2DViewMode::GetName() const
+	{
+		return Name;
+	}
+
+	FText FDataflowConstruction2DViewMode::GetButtonText() const
+	{
+		return FText(LOCTEXT("2DViewButtonText", "2DView"));
+	}
+
+	FText FDataflowConstruction2DViewMode::GetTooltipText() const
+	{
+		return FText(LOCTEXT("2DViewTooltipText", "Default 2D View"));
+	}
+	
+	FName FDataflowConstruction3DViewMode::Name = FName("3DView");
+
+	FName FDataflowConstruction3DViewMode::GetName() const
+	{
+		return Name;
+	}
+
+	FText FDataflowConstruction3DViewMode::GetButtonText() const
+	{
+		return FText(LOCTEXT("3DViewButtonText", "3DView"));
+	}
+
+	FText FDataflowConstruction3DViewMode::GetTooltipText() const
+	{
+		return FText(LOCTEXT("3DViewTooltipText", "Default 3D View"));
+	}
+
+	TUniquePtr<FRenderingViewModeFactory> FRenderingViewModeFactory::Instance = nullptr;
+
+	//
+	// Factory
+	//
+
+	FRenderingViewModeFactory::FRenderingViewModeFactory()
+	{
+		// Default modes
+		ViewModeMap.Add(FDataflowConstruction2DViewMode::Name, MakeUnique<FDataflowConstruction2DViewMode>());
+		ViewModeMap.Add(FDataflowConstruction3DViewMode::Name, MakeUnique<FDataflowConstruction3DViewMode>());
+	}
+
+	FRenderingViewModeFactory& FRenderingViewModeFactory::GetInstance()
+	{
+		return TLazySingleton<FRenderingViewModeFactory>::Get();
+	}
+
+	void FRenderingViewModeFactory::TearDown()
+	{
+		return TLazySingleton<FRenderingViewModeFactory>::TearDown();
+	}
+
+	void FRenderingViewModeFactory::RegisterViewMode(TUniquePtr<IDataflowConstructionViewMode>&& ViewMode)
+	{
+		const FName NewViewModeName = ViewMode->GetName();
+		if (ViewModeMap.Contains(NewViewModeName))
+		{
+			UE_LOG(LogChaos, Warning, TEXT("Dataflow rendering view mode registration conflicts with existing view mode : %s"), *NewViewModeName.ToString());
+		}
+		else
+		{
+			ViewModeMap.Add(NewViewModeName, MoveTemp(ViewMode));
+		}
+	}
+
+	void FRenderingViewModeFactory::DeregisterViewMode(const FName& ViewModeName)
+	{
+		if (!ViewModeMap.Contains(ViewModeName))
+		{
+			UE_LOG(LogChaos, Warning, TEXT("Dataflow rendering view mode deregistration -- view mode not registered : %s"), *ViewModeName.ToString());
+		}
+		else
+		{
+			ViewModeMap.Remove(ViewModeName);
+		}
+	}
+
+	const IDataflowConstructionViewMode* FRenderingViewModeFactory::GetViewMode(const FName& ViewModeName) const
+	{
+		if (ViewModeMap.Contains(ViewModeName))
+		{
+			return ViewModeMap[ViewModeName].Get();
+		}
+		return nullptr;
+	}
+
+	const TMap<FName, TUniquePtr<IDataflowConstructionViewMode>>& FRenderingViewModeFactory::GetViewModes() const
+	{
+		return ViewModeMap;
+	}
+
+}
+
+#undef LOCTEXT_NAMESPACE 
+
diff --git a/Engine/Source/Runtime/Experimental/Dataflow/Engine/Public/Dataflow/DataflowContent.h b/Engine/Source/Runtime/Experimental/Dataflow/Engine/Public/Dataflow/DataflowContent.h
--- a/Engine/Source/Runtime/Experimental/Dataflow/Engine/Public/Dataflow/DataflowContent.h
+++ b/Engine/Source/Runtime/Experimental/Dataflow/Engine/Public/Dataflow/DataflowContent.h
@@ -18,16 +18,7 @@ class UDataflowBaseContent;
 class FPreviewScene;
 class UAnimSingleNodeInstance;
 class AActor;
-
-namespace Dataflow
-{
-	enum class EDataflowPatternVertexType : uint8
-	{
-		Sim2D = 0,
-		Sim3D = 1,
-		Render = 2
-	};
-}
+namespace Dataflow { class IDataflowConstructionViewMode; }
 
 namespace DataflowContextHelpers
 {
@@ -56,8 +47,8 @@ public:
 	TSharedPtr<const FManagedArrayCollection> GetPrimaryRenderCollection() const { return PrimaryRenderCollection; }
 
 	/** ViewMode Access */
-	void SetConstructionViewMode(Dataflow::EDataflowPatternVertexType InMode) {ConstructionViewMode = InMode;}
-	Dataflow::EDataflowPatternVertexType GetConstructionViewMode() const { return ConstructionViewMode; }
+	void SetConstructionViewMode(const Dataflow::IDataflowConstructionViewMode* InMode) { ConstructionViewMode = InMode; }
+	const Dataflow::IDataflowConstructionViewMode* GetConstructionViewMode() const { return ConstructionViewMode; }
 
 	/** Get a single selected node of the specified type. Return nullptr if the specified node is not selected, or if multiple nodes are selected*/
 	template<typename NodeType>
@@ -82,7 +73,7 @@ protected:
 	TObjectPtr<UDataflowEdNode> PrimarySelectedNode = nullptr;
 
 	/** Construction view mode for the context object @todo(michael) : is it only for construction or for simulation as well*/
-	Dataflow::EDataflowPatternVertexType ConstructionViewMode = Dataflow::EDataflowPatternVertexType::Sim3D;
+	const Dataflow::IDataflowConstructionViewMode* ConstructionViewMode = nullptr;
 };
 
 UINTERFACE(MinimalAPI)
diff --git a/Engine/Source/Runtime/Experimental/Dataflow/Engine/Public/Dataflow/DataflowEdNode.h b/Engine/Source/Runtime/Experimental/Dataflow/Engine/Public/Dataflow/DataflowEdNode.h
--- a/Engine/Source/Runtime/Experimental/Dataflow/Engine/Public/Dataflow/DataflowEdNode.h
+++ b/Engine/Source/Runtime/Experimental/Dataflow/Engine/Public/Dataflow/DataflowEdNode.h
@@ -20,7 +20,7 @@
 class FArchive;
 class UEdGraphPin;
 class UObject;
-namespace Dataflow { class FGraph; class FRenderingParameters; }
+namespace Dataflow { class FGraph; class FRenderingParameters; class IDataflowConstructionViewMode;  }
 namespace GeometryCollection::Facades { class FRenderingFacade; }
 
 UCLASS(MinimalAPI)
@@ -101,7 +101,9 @@ public:
 	DATAFLOWENGINE_API bool CanEnableWireframeRenderNode() const;
 
 	DATAFLOWENGINE_API TArray<Dataflow::FRenderingParameter> GetRenderParameters() const;
-	DATAFLOWENGINE_API virtual bool Render(GeometryCollection::Facades::FRenderingFacade& RenderData, const TSharedPtr<Dataflow::FContext> Context) const;
+	DATAFLOWENGINE_API virtual bool Render(GeometryCollection::Facades::FRenderingFacade& RenderData, const TSharedRef<Dataflow::FContext> Context, const Dataflow::IDataflowConstructionViewMode& ViewMode) const;
+
+	DATAFLOWENGINE_API virtual bool CanRender(const TSharedRef<Dataflow::FContext> Context, const Dataflow::IDataflowConstructionViewMode& ViewMode) const;
 
 private:
 
diff --git a/Engine/Source/Runtime/Experimental/Dataflow/Engine/Public/Dataflow/DataflowRenderingFactory.h b/Engine/Source/Runtime/Experimental/Dataflow/Engine/Public/Dataflow/DataflowRenderingFactory.h
--- a/Engine/Source/Runtime/Experimental/Dataflow/Engine/Public/Dataflow/DataflowRenderingFactory.h
+++ b/Engine/Source/Runtime/Experimental/Dataflow/Engine/Public/Dataflow/DataflowRenderingFactory.h
@@ -7,6 +7,7 @@
 #include "Dataflow/DataflowNodeParameters.h"
 #include "Dataflow/DataflowInputOutput.h"
 #include "Dataflow/DataflowNode.h"
+#include "Dataflow/DataflowRenderingViewMode.h"
 #include "GeometryCollection/Facades/CollectionRenderingFacade.h"
 
 namespace Dataflow
@@ -15,13 +16,14 @@ namespace Dataflow
 	typedef TPair<FString, FName> FRenderKey;
 
 	struct FGraphRenderingState {
-		FGraphRenderingState(const FGuid InGuid, const FDataflowNode* InNode, const FRenderingParameter& InParameters, Dataflow::FContext& InContext)
+		FGraphRenderingState(const FGuid InGuid, const FDataflowNode* InNode, const FRenderingParameter& InParameters, Dataflow::FContext& InContext, const Dataflow::IDataflowConstructionViewMode& ViewMode)
 			: NodeGuid(InGuid)
 			, Node(InNode)
 			, RenderName(InParameters.Name)
 			, RenderType(InParameters.Type)
 			, RenderOutputs(InParameters.Outputs)
 			, Context(InContext)
+			, ViewMode(ViewMode)
 		{}
 
 		const FGuid& GetGuid() const { return NodeGuid; }
@@ -42,6 +44,8 @@ namespace Dataflow
 			return Default;
 		}
 
+		const Dataflow::IDataflowConstructionViewMode& GetViewMode() const { return ViewMode; }
+
 	private:
 		const FGuid NodeGuid;
 		const FDataflowNode* Node = nullptr;
@@ -51,53 +55,44 @@ namespace Dataflow
 		TArray<FName> RenderOutputs;
 
 		Dataflow::FContext& Context;
+
+		const Dataflow::IDataflowConstructionViewMode& ViewMode;
 	};
 
-	//
-	//
-	//
+
 	class FRenderingFactory
 	{
-		typedef TFunction<void(GeometryCollection::Facades::FRenderingFacade& RenderData, const FGraphRenderingState& State)> FOutputRenderingFunction;
+	public:
 
-		// All Maps indexed by TypeName
-		TMap<FRenderKey, FOutputRenderingFunction > RenderMap;		// [TypeName] -> NewNodeFunction
-		DATAFLOWENGINE_API static FRenderingFactory* Instance;
-		FRenderingFactory() {}
+		class ICallbackInterface
+		{
+		public:
+			virtual ~ICallbackInterface() = default;
+			virtual FRenderKey GetRenderKey() const = 0;
+			virtual bool CanRender(const Dataflow::IDataflowConstructionViewMode& ViewMode) const = 0;
+			virtual void Render(GeometryCollection::Facades::FRenderingFacade& RenderData, const FGraphRenderingState& State) = 0;
+		};
 
-	public:
 		~FRenderingFactory() { delete Instance; }
 
-		static FRenderingFactory* GetInstance()
-		{
-			if (!Instance)
-			{
-				Instance = new FRenderingFactory();
-			}
-			return Instance;
-		}
+		DATAFLOWENGINE_API static FRenderingFactory* GetInstance();
 
-		void RegisterOutput(const FRenderKey& Key, FOutputRenderingFunction InFunction)
-		{
-			if (RenderMap.Contains(Key))
-			{
-				UE_LOG(LogChaos, Warning,
-					TEXT("Warning : Dataflow output rendering registration conflicts with "
-						"existing renderer(<%s,%s>)"), 
-					*Key.Get<0>(),
-					*Key.Get<1>().ToString());
-			}
-			else
-			{
-				RenderMap.Add(Key, InFunction);
-			}
-		}
+		DATAFLOWENGINE_API void RegisterCallbacks(TUniquePtr<ICallbackInterface> InCallbacks);
+		DATAFLOWENGINE_API void DeregisterCallbacks(const FRenderKey& Key);
 
+		DATAFLOWENGINE_API bool Contains(const FRenderKey& InKey) const { return CallbackMap.Contains(InKey); }
 		DATAFLOWENGINE_API void RenderNodeOutput(GeometryCollection::Facades::FRenderingFacade& RenderData, const FGraphRenderingState& State);
+		DATAFLOWENGINE_API bool CanRenderNodeOutput(const FGraphRenderingState& State) const;
 
-		bool Contains(const FRenderKey& InKey) const { return RenderMap.Contains(InKey); }
+	private:
+
+		FRenderingFactory() {}
 
+		static FRenderingFactory* Instance;
+
+		TMap<FRenderKey, TUniquePtr<ICallbackInterface>> CallbackMap;
 	};
 
 }
 
+
diff --git a/Engine/Source/Runtime/Experimental/Dataflow/Engine/Public/Dataflow/DataflowRenderingViewMode.h b/Engine/Source/Runtime/Experimental/Dataflow/Engine/Public/Dataflow/DataflowRenderingViewMode.h
new file mode 100644
--- /dev/null
+++ b/Engine/Source/Runtime/Experimental/Dataflow/Engine/Public/Dataflow/DataflowRenderingViewMode.h
@@ -0,0 +1,126 @@
+// Copyright Epic Games, Inc. All Rights Reserved.
+
+#pragma once
+
+#include "UObject/NameTypes.h"
+#include "ChaosLog.h"
+
+namespace Dataflow
+{
+	//
+	// Dataflow Construction View Modes allow the user to choose different views on the same node output. The View Mode system serves two main purposes: 
+	// 
+	// 1) In the Dataflow Editor's Construction Viewport, the currently active IDataflowConstructionViewMode will determine which ELevelViewportType should be used for rendering.
+	//    In other words, the the ViewMode will determine whether we have a 3D perspective camera or a 2D orthographic camera.
+	// 2) In the Datflow Rendering Factory (Dataflow::FRenderingFactory), the rendering callbacks that are registered will also be able to return different 
+	//    render data based on what the currently active IDataflowConstructionViewMode in the viewport is.
+	//
+	// For example, suppose you have created a new type FMyMeshType that contains both a 3D Mesh representation and a 2D Mesh representation (for, say, UVs.) If you want to
+	// write custom ViewModes for it, here is the process:
+	// 
+	// 1) Define two new classes, e.g. FMyMeshType2DViewMode (inheriting from FDataflowConstruction2DViewModeBase) and FMyMeshType3DViewMode (inheriting from FDataflowConstruction3DViewModeBase)
+	// 2) In your module initialization, register both view modes with the Dataflow::FRenderingViewModeFactory
+	// 3) Create a render function callback class and register it with Dataflow::FRenderingFactory. 
+	//    - The CanRender function should return true if the view mode is either of your two new custom view modes.
+	//    - The Render function should return different RenderCollections depending on what the current view mode is (which can be determined from the FGraphRenderingState parameter.)
+	//
+
+
+	//
+	// Equivalent to ELevelViewportType, but we don't want to add an Editor dependency here
+	//
+	enum EDataflowViewportType : int
+	{
+		DVT_OrthoXY = 0,
+		DVT_OrthoXZ = 1,
+		DVT_OrthoYZ = 2,
+		DVT_Perspective = 3,
+		DVT_OrthoFreelook = 4,
+		DVT_OrthoNegativeXY = 5,
+		DVT_OrthoNegativeXZ = 6,
+		DVT_OrthoNegativeYZ = 7,
+		DVT_MAX,
+		DVT_None = 255,
+	};
+
+
+	class IDataflowConstructionViewMode
+	{
+	public:
+		virtual ~IDataflowConstructionViewMode() = default;
+
+		DATAFLOWENGINE_API virtual FName GetName() const = 0;
+		DATAFLOWENGINE_API virtual FText GetButtonText() const = 0;
+		DATAFLOWENGINE_API virtual FText GetTooltipText() const = 0;
+		DATAFLOWENGINE_API virtual EDataflowViewportType GetViewportType() const = 0;
+		
+		DATAFLOWENGINE_API bool IsPerspective() const;
+	};
+
+	// Base 2D and 3D types
+
+	class FDataflowConstruction2DViewModeBase : public IDataflowConstructionViewMode
+	{
+	public:
+		DATAFLOWENGINE_API virtual EDataflowViewportType GetViewportType() const override;
+	};
+
+	class FDataflowConstruction3DViewModeBase : public IDataflowConstructionViewMode
+	{
+	public:
+		DATAFLOWENGINE_API virtual EDataflowViewportType GetViewportType() const override;
+	};
+
+	// Concrete default 2D and 3D types
+
+	class FDataflowConstruction2DViewMode : public FDataflowConstruction2DViewModeBase
+	{
+	public:
+		DATAFLOWENGINE_API static FName Name;
+		virtual ~FDataflowConstruction2DViewMode() = default;
+	private:
+		DATAFLOWENGINE_API virtual FName GetName() const override;
+		DATAFLOWENGINE_API virtual FText GetButtonText() const override;
+		DATAFLOWENGINE_API virtual FText GetTooltipText() const override;
+	};
+
+	class FDataflowConstruction3DViewMode : public FDataflowConstruction3DViewModeBase
+	{
+	public:
+		DATAFLOWENGINE_API static FName Name;
+		virtual ~FDataflowConstruction3DViewMode() = default;
+	private:
+		DATAFLOWENGINE_API virtual FName GetName() const override;
+		DATAFLOWENGINE_API virtual FText GetButtonText() const override;
+		DATAFLOWENGINE_API virtual FText GetTooltipText() const override;
+	};
+
+	//
+	// ViewMode registry/factory
+	//
+
+	class FRenderingViewModeFactory
+	{
+	public:
+
+		FRenderingViewModeFactory();
+
+		// FLazySingleton
+		static DATAFLOWENGINE_API FRenderingViewModeFactory& GetInstance();
+		static DATAFLOWENGINE_API void TearDown();
+
+		DATAFLOWENGINE_API void RegisterViewMode(TUniquePtr<IDataflowConstructionViewMode>&& ViewMode);
+		DATAFLOWENGINE_API void DeregisterViewMode(const FName& ViewModeName);
+
+		DATAFLOWENGINE_API const IDataflowConstructionViewMode* GetViewMode(const FName& ViewModeName) const;
+
+		DATAFLOWENGINE_API const TMap<FName, TUniquePtr<IDataflowConstructionViewMode>>& GetViewModes() const;
+
+	private:
+
+		DATAFLOWENGINE_API static TUniquePtr<FRenderingViewModeFactory> Instance;
+
+		TMap<FName, TUniquePtr<IDataflowConstructionViewMode>> ViewModeMap;
+	};
+}
+
