1 - Chaos dataflow simulation manager

diff --git a/Engine/Plugins/Experimental/ChaosFlesh/Source/ChaosFleshNodes/Private/Dataflow/ChaosFleshNodesPlugin.cpp b/Engine/Plugins/Experimental/ChaosFlesh/Source/ChaosFleshNodes/Private/Dataflow/ChaosFleshNodesPlugin.cpp
--- a/Engine/Plugins/Experimental/ChaosFlesh/Source/ChaosFleshNodes/Private/Dataflow/ChaosFleshNodesPlugin.cpp
+++ b/Engine/Plugins/Experimental/ChaosFlesh/Source/ChaosFleshNodes/Private/Dataflow/ChaosFleshNodesPlugin.cpp
@@ -33,18 +33,10 @@
 #include "Dataflow/ChaosFleshVertexConstraintNode.h"
 #include "Dataflow/ChaosFleshVisualizeFiberFieldNode.h"
 #include "Dataflow/GeometryCollectionAppendCollectionTransformNode.h"
-#include "Dataflow/ChaosFleshSimulationNodes.h"
 #include "Modules/ModuleManager.h"
 
 #define LOCTEXT_NAMESPACE "ChaosFleshNodes"
 
-namespace Dataflow::CVars
-{
-	static bool bRegisterFleshSimulationNodes = false;
-	static FAutoConsoleVariableRef CVarRegisterFleshSimulationNodes(TEXT("p.Dataflow.Flesh.RegisterSimulationNodes"), bRegisterFleshSimulationNodes, TEXT("True to register all the flesh simulation nodes."));
-}
-
-
 void IChaosFleshNodesPlugin::StartupModule()
 {
 	DATAFLOW_NODE_REGISTER_CREATION_FACTORY(FAddKinematicParticlesDataflowNode);
@@ -80,15 +72,6 @@ void IChaosFleshNodesPlugin::StartupModule()
 	DATAFLOW_NODE_REGISTER_CREATION_FACTORY(FVisualizeFiberFieldNode);
 	DATAFLOW_NODE_REGISTER_CREATION_FACTORY(FVisualizePositionTargetsNode);
 	DATAFLOW_NODE_REGISTER_CREATION_FACTORY(FVisualizeKinematicFacesNode);
-
-	// Temporary registry of the simulation nodes for
-	// testing while the simulation graph is being added
-	if(Dataflow::CVars::bRegisterFleshSimulationNodes)
-	{
-		Dataflow::RegisterChaosFleshSimulationNodes();
-		Dataflow::RegisterChaosSkeletonSimulationNodes();
-		Dataflow::RegisterChaosCommonSimulationNodes();
-	}
 }
 
 void IChaosFleshNodesPlugin::ShutdownModule()
diff --git a/Engine/Plugins/Experimental/ChaosFlesh/Source/ChaosFleshNodes/Private/Dataflow/ChaosFleshSimulationNodes.cpp b/Engine/Plugins/Experimental/ChaosFlesh/Source/ChaosFleshNodes/Private/Dataflow/ChaosFleshSimulationNodes.cpp
deleted file mode 100644
diff --git a/Engine/Plugins/Experimental/ChaosFlesh/Source/ChaosFleshNodes/Public/Dataflow/ChaosFleshSimulationNodes.h b/Engine/Plugins/Experimental/ChaosFlesh/Source/ChaosFleshNodes/Public/Dataflow/ChaosFleshSimulationNodes.h
deleted file mode 100644
diff --git a/Engine/Plugins/Experimental/Dataflow/Dataflow.uplugin b/Engine/Plugins/Experimental/Dataflow/Dataflow.uplugin
--- a/Engine/Plugins/Experimental/Dataflow/Dataflow.uplugin
+++ b/Engine/Plugins/Experimental/Dataflow/Dataflow.uplugin
@@ -35,6 +35,11 @@
 			"Name": "DataflowNodes",
 			"Type": "Runtime",
 			"LoadingPhase": "Default"
+		},
+		{
+			"Name": "DataflowSimulation",
+			"Type": "Runtime",
+			"LoadingPhase": "Default"
 		}
 	],
 	"Plugins": [
diff --git a/Engine/Plugins/Experimental/Dataflow/Source/DataflowEditor/DataflowEditor.Build.cs b/Engine/Plugins/Experimental/Dataflow/Source/DataflowEditor/DataflowEditor.Build.cs
--- a/Engine/Plugins/Experimental/Dataflow/Source/DataflowEditor/DataflowEditor.Build.cs
+++ b/Engine/Plugins/Experimental/Dataflow/Source/DataflowEditor/DataflowEditor.Build.cs
@@ -30,6 +30,7 @@ namespace UnrealBuildTool.Rules
 					"DataflowEngine",
 					"DataflowEnginePlugin",
 					"DataflowNodes",
+					"DataflowSimulation",
 					"DynamicMesh",
 					"Engine",
 					"EditorFramework",
diff --git a/Engine/Plugins/Experimental/Dataflow/Source/DataflowEditor/Private/Dataflow/DataflowConstructionScene.cpp b/Engine/Plugins/Experimental/Dataflow/Source/DataflowEditor/Private/Dataflow/DataflowConstructionScene.cpp
--- a/Engine/Plugins/Experimental/Dataflow/Source/DataflowEditor/Private/Dataflow/DataflowConstructionScene.cpp
+++ b/Engine/Plugins/Experimental/Dataflow/Source/DataflowEditor/Private/Dataflow/DataflowConstructionScene.cpp
@@ -85,11 +85,11 @@ void FDataflowConstructionScene::TickDataflowScene(const float DeltaSeconds)
 			if (const UDataflow* Dataflow = EditorContent->GetDataflowAsset())
 			{
 				const Dataflow::FTimestamp SystemTimestamp = LatestTimestamp(Dataflow, DataflowContext.Get());
-				if (SystemTimestamp >= EditorContent->GetLastModifiedTimestamp() || EditorContent->IsDirty())
+				if (SystemTimestamp >= EditorContent->GetLastModifiedTimestamp() || EditorContent->IsConstructionDirty())
 				{
 					EditorContent->SetLastModifiedTimestamp(SystemTimestamp.Value + 1);
 
-					if (EditorContent->IsDirty())
+					if (EditorContent->IsConstructionDirty())
 					{
 						UpdateConstructionScene();
 					}
@@ -376,12 +376,12 @@ void FDataflowConstructionScene::UpdateConstructionScene()
 
 	if (const TObjectPtr<UDataflowBaseContent>& EditorContent = GetEditorContent())
 	{
-		EditorContent->SetIsDirty(false);
+		EditorContent->SetConstructionDirty(false);
 	}
 
 	for(const TObjectPtr<UDataflowBaseContent>& TerminalContent : GetTerminalContents())
 	{
-		TerminalContent->SetIsDirty(false);
+		TerminalContent->SetConstructionDirty(false);
 	}
 }
 
diff --git a/Engine/Plugins/Experimental/Dataflow/Source/DataflowEditor/Private/Dataflow/DataflowEditorModeToolkit.cpp b/Engine/Plugins/Experimental/Dataflow/Source/DataflowEditor/Private/Dataflow/DataflowEditorModeToolkit.cpp
--- a/Engine/Plugins/Experimental/Dataflow/Source/DataflowEditor/Private/Dataflow/DataflowEditorModeToolkit.cpp
+++ b/Engine/Plugins/Experimental/Dataflow/Source/DataflowEditor/Private/Dataflow/DataflowEditorModeToolkit.cpp
@@ -279,7 +279,7 @@ void FDataflowEditorModeToolkit::OnToolEnded(UInteractiveToolManager* Manager, U
 	{
 		if (const TObjectPtr<UDataflowBaseContent>& EditorContent = ConstructionScene->GetEditorContent())
 		{
-			EditorContent->SetIsDirty(true);
+			EditorContent->SetConstructionDirty(true);
 		}
 	}
 
diff --git a/Engine/Plugins/Experimental/Dataflow/Source/DataflowEditor/Private/Dataflow/DataflowSimulationGenerator.cpp b/Engine/Plugins/Experimental/Dataflow/Source/DataflowEditor/Private/Dataflow/DataflowSimulationGenerator.cpp
--- a/Engine/Plugins/Experimental/Dataflow/Source/DataflowEditor/Private/Dataflow/DataflowSimulationGenerator.cpp
+++ b/Engine/Plugins/Experimental/Dataflow/Source/DataflowEditor/Private/Dataflow/DataflowSimulationGenerator.cpp
@@ -1,11 +1,10 @@
 ﻿// Copyright Epic Games, Inc. All Rights Reserved.
 
 #include "Dataflow/DataflowSimulationGenerator.h"
-
-#include "Engine/World.h"
+#include "Dataflow/DataflowSimulationManager.h"
 #include "Chaos/CacheManagerActor.h"
-#include "Dataflow/DataflowObject.h"
-#include "Dataflow/DataflowSolverNodes.h"
+#include "Engine/World.h"
+#include "Engine/Engine.h"
 #include "Misc/AsyncTaskNotification.h" 
 
 DEFINE_LOG_CATEGORY(LogDataflowSimulationGenerator);
@@ -14,379 +13,277 @@ DEFINE_LOG_CATEGORY(LogDataflowSimulationGenerator);
 
 namespace Dataflow
 {
-	float GetCacheDuration(const TObjectPtr<UDataflow>& SimulationGraph, const TSharedPtr<Dataflow::FSimulationContext>& SimulationContext, const int32 SamplingRate, float& MinTime, float& MaxTime, int32& NumFrames)
-	{
-		float CacheDuration = 0.0f;
-		if(SimulationContext.IsValid())
-		{
-			SimulationContext->SetSimulationFlag(Dataflow::ESimulationFlags::SetupSimulation);
-		}
-		if(const TSharedPtr<Dataflow::FGraph> DataflowGraph = SimulationGraph->GetDataflow())
-		{
-			MinTime = TNumericLimits<float>::Max();
-			MaxTime = TNumericLimits<float>::Lowest();
-			for(const TSharedPtr<FDataflowNode>& CacheNode : DataflowGraph->GetFilteredNodes(FDataflowCacheNode::StaticType()))
-			{
-				const FVector2f NodeRange = StaticCastSharedPtr<FDataflowCacheNode>(CacheNode)->GetTimeRange(*SimulationContext.Get());
-				MinTime = FMath::Min(MinTime, NodeRange[0]);
-				MaxTime = FMath::Max(MaxTime, NodeRange[1]);
-			}
-			NumFrames = (MaxTime > MinTime) ? FMath::Floor((MaxTime - MinTime) * SamplingRate) : 0;
-			CacheDuration = ( MaxTime > MinTime) ? (MaxTime - MinTime) : 0.0f;
-		}
-		if(SimulationContext.IsValid())
-		{
-			SimulationContext->UnsetSimulationFlag(Dataflow::ESimulationFlags::SetupSimulation);
-		}
-		return CacheDuration;
-	}
-
-	bool ShouldResetWorld(const TObjectPtr<UDataflow>& SimulationGraph, const TSharedPtr<Dataflow::FSimulationContext>& SimulationContext, Dataflow::FTimestamp& LastTimeStamp)
+void FDataflowSimulationTask::DoWork()
+{
+	const int32 NumFrames = (MaxTime-MinTime) / DeltaTime;
+	for (int32 FrameIndex = 0; FrameIndex < NumFrames; ++FrameIndex)
 	{
-		if(const TSharedPtr<Dataflow::FGraph> DataflowGraph = SimulationGraph->GetDataflow())
+		if (!TaskManager->bCancelled.load())
 		{
-			for(const TSharedPtr<FDataflowNode>& CacheNode : DataflowGraph->GetFilteredNodes(FDataflowCacheNode::StaticType()))
-			{
-				if(CacheNode->GetTimestamp() >= LastTimeStamp)
-				{
-					LastTimeStamp = SimulationContext->GetTimestamp();
-					return true;
-				}
-			}
-		}
-		return false;
-	}
+			// Compute the simulation time that will be sent to the graph
+			const float SimulationTime = MinTime + (FrameIndex+1) * DeltaTime;
 
-	void UpdateAnimationNodes(const TObjectPtr<UDataflow>& SimulationGraph, const TSharedPtr<Dataflow::FSimulationContext>& SimulationContext, const float AnimationTime)
-	{
-		if(const TSharedPtr<Dataflow::FGraph> DataflowGraph = SimulationGraph->GetDataflow())
-		{
-			for(const TSharedPtr<FDataflowNode>& AnimationNode : DataflowGraph->GetFilteredNodes(FDataflowAnimationNode::StaticType()))
+			if(bBackgroundTask)
 			{
-				StaticCastSharedPtr<FDataflowAnimationNode>(AnimationNode)->SetAnimationTime(*SimulationContext.Get(), AnimationTime);
+				// Background task : run directly the advance simulation data without coming back to the game thread
+				UDataflowSimulationManager* DataflowManager = SimulationWorld->GetSubsystem<UDataflowSimulationManager>();
+                DataflowManager->AdvanceSimulationData(DeltaTime, SimulationTime);
 			}
-		}
-	}
-	
-	void EvaluateSimulationGraph(const TObjectPtr<UDataflow>& SimulationGraph, const TSharedPtr<Dataflow::FSimulationContext>& SimulationContext,
-		const float DeltaTime, const float SimulationTime, const ESimulationFlags SimulationFlags)
-	{
-		if(SimulationContext.IsValid())
-		{
-			SimulationContext->SetSimulationFlag(SimulationFlags);
-			if(SimulationContext->HasSimulationFlag(ESimulationFlags::UpdateSimulation))
+			else
 			{
-				SimulationContext->SetTimingInfos(DeltaTime, SimulationTime);
+				// Update all the skelmesh animations at the simulation time
+				Dataflow::UpdateSkeletonAnimation(TaskManager->PreviewActor, SimulationTime);
+				
+				// Foreground task : Run the world ticking 
+				SimulationWorld->Tick(ELevelTick::LEVELTICK_All, DeltaTime);
 			}
-		}
 
-		if(SimulationGraph)
-		{
-			if(const TSharedPtr<Dataflow::FGraph> DataflowGraph = SimulationGraph->GetDataflow())
-			{
-				if(SimulationContext->HasSimulationFlag(ESimulationFlags::UpdateSimulation))
-				{
-					// Invalidation of all the simulation nodes that are always dirty
-					for(const TSharedPtr<FDataflowNode>& SimulationNode : DataflowGraph->GetFilteredNodes(FDataflowInvalidNode::StaticType()))
-					{
-						SimulationNode->Invalidate();
-					}
-				}
-
-				// Pull the graph evaluation from the cache nodes
-				for(const TSharedPtr<FDataflowNode>& CacheNode : DataflowGraph->GetFilteredNodes(FDataflowCacheNode::StaticType()))
-				{
-					SimulationContext->Evaluate(CacheNode.Get(), nullptr);
-				}
-			}
+			// Finish the frame
+			TaskManager->SimulationResource->FinishFrame();
 		}
-		if(SimulationContext.IsValid())
+		else
 		{
-			SimulationContext->UnsetSimulationFlag(SimulationFlags);
+			break;
 		}
 	}
+}
 
-	void CleanSimulationContext(TSharedPtr<Dataflow::FSimulationContext>& SimulationContext)
-	{
-		if(const TObjectPtr<AChaosCacheManager> CacheManager = Cast<AChaosCacheManager>(SimulationContext->GetRootActor()))
-		{
-			if(SimulationContext->GetCachingMode() == ECachingMode::RecordCache)
-			{
-				CacheManager->EndEvaluate();
-			}
-
-			//  Clear the observed components
-			CacheManager->ClearObservedComponents();
-			
-			TInlineComponentArray<UPrimitiveComponent*> PrimComponents;
-			CacheManager->GetComponents(PrimComponents);
+bool FDataflowTaskManager::AllocateSimulationResource(const FVector2f& TimeRange, const int32 FrameRate,
+		const TObjectPtr<UChaosCacheCollection>& CacheAsset, const TSubclassOf<AActor>& ActorClass, const TObjectPtr<UDataflowBaseContent>& DataflowContent)
+{
+	SimulationWorld = UWorld::CreateWorld(EWorldType::Editor, false);
 
-			// Unregister all the components
-			for(UPrimitiveComponent* PrimComponent : PrimComponents)
-			{
-				PrimComponent->SelectionOverrideDelegate.Unbind();
-				PrimComponent->UnregisterComponent();
-				PrimComponent->DestroyComponent();
-			}
-		}
+	FWorldContext& WorldContext = GEngine->CreateNewWorldContext( SimulationWorld->WorldType );
+	WorldContext.SetCurrentWorld(SimulationWorld);
+	
+	CacheManager = SimulationWorld->SpawnActor<AChaosCacheManager>();
 
-		// Reset the simulation context
-		SimulationContext.Reset();
-	}
+	PreviewActor = Dataflow::SpawnSimulatedActor(ActorClass, CacheManager, CacheAsset, true, DataflowContent);
 
-	void BuildSimulationContext(const TObjectPtr<UObject>& SimulationOwner, const TObjectPtr<UDataflow>& SimulationGraph, const TObjectPtr<AChaosCacheManager>& CacheManager,
-		const ECachingMode CachingMode, TSharedPtr<Dataflow::FSimulationContext>& SimulationContext)
+	// Init the cache manager
+	CacheManager->SetObservedComponentProperties(CacheManager->CacheMode);
+	CacheManager->BeginEvaluate();
+	
+	SimulationResource = MakeShared<FDataflowSimulationResource>();
+	SimulationResource->NumSimulatedFrames = &NumSimulatedFrames;
+	SimulationResource->bCancelled = &bCancelled;
+	
+	NumFrames = FMath::Floor((TimeRange[1] - TimeRange[0]) * FrameRate);
+	if(SimulationTask.IsValid())
 	{
-		if(CacheManager)
-		{
-			// Create a simulation context that will hold all the simulation datas
-			SimulationContext = MakeShared<Dataflow::FSimulationContext>(
-						SimulationOwner, SimulationGraph,
-					Dataflow::FTimestamp::Invalid, CacheManager, CachingMode);
-
-			// Evaluate the setup stage of the simulation graph
-			EvaluateSimulationGraph(SimulationGraph, SimulationContext, 
-					0.0f, 0.0f , ESimulationFlags::SetupSimulation);
+		SimulationTask->GetTask().SimulationWorld = SimulationWorld;
+		SimulationTask->GetTask().MinTime = TimeRange[0];
+		SimulationTask->GetTask().MaxTime = TimeRange[1];
+		SimulationTask->GetTask().DeltaTime = (TimeRange[1]-TimeRange[0]) / NumFrames;
+
+		UDataflowSimulationManager* DataflowManager = SimulationWorld->GetSubsystem<UDataflowSimulationManager>();
+		DataflowManager->SetSimulationEnabled(!SimulationTask->GetTask().bBackgroundTask);
+		DataflowManager->WriteSimulationData(SimulationTask->GetTask().DeltaTime);
+	}
+	Dataflow::SetupSkeletonAnimation(PreviewActor);
 
-			// Init the cache manager
-			CacheManager->SetObservedComponentProperties(CacheManager->CacheMode);
+	return true;
+}
 
-			if(SimulationContext->GetCachingMode() == ECachingMode::RecordCache)
-			{
-				CacheManager->BeginEvaluate();
-			}
-		}
+void FDataflowTaskManager::FreeSimulationResource()
+{
+	if (SimulationTask.IsValid())
+	{
+		SimulationTask->EnsureCompletion();
+		
+		UDataflowSimulationManager* DataflowManager = SimulationWorld->GetSubsystem<UDataflowSimulationManager>();
+		DataflowManager->ReadSimulationData(SimulationTask->GetTask().DeltaTime);
 	}
-	
-	void FDataflowSimulationTask::DoWork()
+
+	if(CacheManager)
 	{
-		if(SimulationGraph)
-		{
-			const int32 NumFrames = (MaxTime-MinTime) / DeltaTime;
-			for (int32 FrameIndex = 0; FrameIndex < NumFrames; ++FrameIndex)
-			{
-				if (!TaskManager->bCancelled.load())
-				{
-					// Compute the simulation time that will be sent to the graph
-					const float SimulationTime = MinTime + (FrameIndex+1) * DeltaTime;
-
-					// Evaluate the simulation graph update stage
-					EvaluateSimulationGraph(SimulationGraph, SimulationContext,
-						DeltaTime, SimulationTime, ESimulationFlags::UpdateSimulation);
-
-					// Finish the frame
-					TaskManager->SimulationResource->FinishFrame();
-				}
-				else
-				{
-					break;
-				}
-			}
-		}
+		CacheManager->EndEvaluate();
+
+		//  Clear the observed components
+        CacheManager->ClearObservedComponents();
+		SimulationWorld->DestroyActor(CacheManager);
 	}
 	
-	bool FDataflowTaskManager::AllocateSimulationResource(const TObjectPtr<UDataflow> SimulationGraph, const int32 SamplingRate, const TObjectPtr<UChaosCacheCollection>& CacheCollection)
-	{
-		SimulationWorld = UWorld::CreateWorld(EWorldType::Editor, false);
-		TObjectPtr<AChaosCacheManager> CacheManager = SimulationWorld->SpawnActor<AChaosCacheManager>(AChaosCacheManager::StaticClass());
+	SimulationResource.Reset();
 
-		CacheManager->StartMode = EStartMode::Timed;
-		CacheManager->CacheMode = ECacheMode::Record;
-		CacheManager->CacheCollection = CacheCollection;
-		
-		SimulationResource = MakeShared<FDataflowSimulationResource>();
-		SimulationResource->NumSimulatedFrames = &NumSimulatedFrames;
-		SimulationResource->bCancelled = &bCancelled;
-		
-		if(SimulationTask.IsValid())
-		{
-			// Clone the simulation graph to avoid evaluating the graph on another thread while modifying it on the game thread
-			SimulationTask->GetTask().SimulationGraph = DuplicateObject<UDataflow>(SimulationGraph, SimulationWorld, NAME_None);
+	GEngine->DestroyWorldContext(SimulationWorld);
+	SimulationWorld->DestroyWorld(false);
+}
 
-			// Build the simulation context
-			BuildSimulationContext(SimulationTask->GetTask().SimulationGraph, SimulationTask->GetTask().SimulationGraph,
-				CacheManager, ECachingMode::RecordCache, SimulationTask->GetTask().SimulationContext);
+void FDataflowTaskManager::CancelSimulationGeneration()
+{
+	bCancelled.store(true);
+	SimulationTask->TryAbandonTask();
+}
 
-			// Compute the simulation duration in order to run the simulation
-			const float Duration = GetCacheDuration(SimulationTask->GetTask().SimulationGraph, SimulationTask->GetTask().SimulationContext, SamplingRate,
-				SimulationTask->GetTask().MinTime, SimulationTask->GetTask().MaxTime, NumFrames);
+FDataflowSimulationGenerator::FDataflowSimulationGenerator()
+{}
 
-			// Update the delta time 
-			SimulationTask->GetTask().DeltaTime = Duration / NumFrames;
-		}
+FDataflowSimulationGenerator::~FDataflowSimulationGenerator()
+{
+	if (TaskManager != nullptr)
+	{
+		TaskManager->FreeSimulationResource();
+	}
+}
 
-		return true;
+void FDataflowSimulationGenerator::Tick(float DeltaTime)
+{
+	if (PendingAction == EDataflowGeneratorActions::StartGenerate)
+	{
+		StartGenerateSimulation();
 	}
-	
-	void FDataflowTaskManager::FreeSimulationResource()
+	else if (PendingAction == EDataflowGeneratorActions::TickGenerate)
 	{
-		if (SimulationTask.IsValid())
-		{
-			SimulationTask->EnsureCompletion();
-
-			CleanSimulationContext(SimulationTask->GetTask().SimulationContext);
-		}
-		
-		SimulationResource.Reset();
-		SimulationWorld->DestroyWorld(false);
+		TickGenerateSimulation();
 	}
+}
+
+TStatId FDataflowSimulationGenerator::GetStatId() const
+{
+	RETURN_QUICK_DECLARE_CYCLE_STAT(FDataflowSimulationGenerator, STATGROUP_Tickables);
+}
 
-	void FDataflowTaskManager::CancelSimulationGeneration()
+void FDataflowSimulationGenerator::StartGenerateSimulation()
+{
+	check(PendingAction == EDataflowGeneratorActions::StartGenerate);
+	
+	if (TaskManager != nullptr)
 	{
-		bCancelled.store(true);
-		SimulationTask->TryAbandonTask();
+		UE_LOG(LogDataflowSimulationGenerator, Error, TEXT("Previous generation is still running."));
+		PendingAction = EDataflowGeneratorActions::NoAction;
+		return;
 	}
+	TaskManager = TSharedPtr<FDataflowTaskManager>(new FDataflowTaskManager);
+
+	TaskManager->SimulationTask = MakeUnique<FAsyncTask<FDataflowSimulationTask>>();
+	TaskManager->SimulationTask->GetTask().TaskManager = TaskManager;
+	TaskManager->SimulationTask->GetTask().bBackgroundTask = bBackgroundTask;
 	
-	FDataflowSimulationGenerator::FDataflowSimulationGenerator()
-	{}
+	TaskManager->AllocateSimulationResource(TimeRange, FrameRate, CacheAsset, ActorClass, DataflowContent);
 
-	FDataflowSimulationGenerator::~FDataflowSimulationGenerator()
+	if(bBackgroundTask)
 	{
-		if (TaskManager != nullptr)
-		{
-			TaskManager->FreeSimulationResource();
-		}
+		TaskManager->SimulationTask->StartBackgroundTask();
 	}
-
-	void FDataflowSimulationGenerator::Tick(float DeltaTime)
+	else
 	{
-		if (PendingAction == EDataflowGeneratorActions::StartGenerate)
-		{
-			StartGenerateSimulation();
-		}
-		else if (PendingAction == EDataflowGeneratorActions::TickGenerate)
-		{
-			TickGenerateSimulation();
-		}
+		TaskManager->SimulationTask->StartSynchronousTask();
 	}
 
-	TStatId FDataflowSimulationGenerator::GetStatId() const
+	FAsyncTaskNotificationConfig NotificationConfig;
+	NotificationConfig.TitleText = LOCTEXT("SimulateDataflow", "Simulating Dataflow Content");
+	NotificationConfig.ProgressText = FText::FromString(TEXT("0%"));
+	NotificationConfig.bCanCancel = true;
+	NotificationConfig.bKeepOpenOnSuccess = true;
+	NotificationConfig.bKeepOpenOnFailure = true;
+	TaskManager->AsyncNotification = MakeUnique<FAsyncTaskNotification>(NotificationConfig);
+	TaskManager->StartTime = FDateTime::UtcNow();
+	TaskManager->LastUpdateTime = TaskManager->StartTime;
+
+	PendingAction = EDataflowGeneratorActions::TickGenerate;
+}
+
+void FDataflowSimulationGenerator::TickGenerateSimulation()
+{
+	check(PendingAction == EDataflowGeneratorActions::TickGenerate && TaskManager != nullptr);
+		
+	bool bFinished = false;
+	const bool bCancelled = TaskManager->AsyncNotification->GetPromptAction() == EAsyncTaskNotificationPromptAction::Cancel;
+	if (TaskManager->SimulationTask->IsDone())
 	{
-		RETURN_QUICK_DECLARE_CYCLE_STAT(FDataflowSimulationGenerator, STATGROUP_Tickables);
+		bFinished = true;
 	}
-
-	void FDataflowSimulationGenerator::StartGenerateSimulation()
+	else if (bCancelled)
 	{
-		check(PendingAction == EDataflowGeneratorActions::StartGenerate);
+		TaskManager->CancelSimulationGeneration();
+		bFinished = true;
+	}
 		
-		if (TaskManager != nullptr)
-		{
-			UE_LOG(LogDataflowSimulationGenerator, Error, TEXT("Previous generation is still running."));
-			PendingAction = EDataflowGeneratorActions::NoAction;
-			return;
-		}
-		if(!SimulationGraph)
+	if (!bFinished)
+	{
+		const FDateTime CurrentTime = FDateTime::UtcNow();
+		const double SinceLastUpdate = (CurrentTime - TaskManager->LastUpdateTime).GetTotalSeconds();
+		if (SinceLastUpdate < 0.2)
 		{
-			UE_LOG(LogDataflowSimulationGenerator, Error, TEXT("No simulation graph available in the generator"));
-			PendingAction = EDataflowGeneratorActions::NoAction;
 			return;
 		}
-		TaskManager = TSharedPtr<FDataflowTaskManager>(new FDataflowTaskManager);
-
-		TaskManager->SimulationTask = MakeUnique<FAsyncTask<FDataflowSimulationTask>>();
-		TaskManager->SimulationTask->GetTask().TaskManager = TaskManager;
-		
-		TaskManager->AllocateSimulationResource(SimulationGraph, SamplingRate, CacheCollection);
 		
-		TaskManager->SimulationTask->StartBackgroundTask();
-	
-		FAsyncTaskNotificationConfig NotificationConfig;
-		NotificationConfig.TitleText = LOCTEXT("SimulateDataflow", "Simulating Dataflow Content");
-		NotificationConfig.ProgressText = FText::FromString(TEXT("0%"));
-		NotificationConfig.bCanCancel = true;
-		NotificationConfig.bKeepOpenOnSuccess = true;
-		NotificationConfig.bKeepOpenOnFailure = true;
-		TaskManager->AsyncNotification = MakeUnique<FAsyncTaskNotification>(NotificationConfig);
-		TaskManager->StartTime = FDateTime::UtcNow();
-		TaskManager->LastUpdateTime = TaskManager->StartTime;
-	
-		PendingAction = EDataflowGeneratorActions::TickGenerate;
+		const int32 NumSimulatedFrames = TaskManager->NumSimulatedFrames.load();
+		const int32 NumTotalFrames = TaskManager->NumFrames;
+		const FText ProgressMessage = FText::FromString(FString::Printf(TEXT("Finished %d/%d, %.1f%%"),
+			NumSimulatedFrames, NumTotalFrames, 100.0 * NumSimulatedFrames / NumTotalFrames));
+		TaskManager->AsyncNotification->SetProgressText(ProgressMessage);
+		TaskManager->LastUpdateTime = CurrentTime;
 	}
-	
-	void FDataflowSimulationGenerator::TickGenerateSimulation()
+	else
 	{
-		check(PendingAction == EDataflowGeneratorActions::TickGenerate && TaskManager != nullptr);
-			
-		bool bFinished = false;
-		const bool bCancelled = TaskManager->AsyncNotification->GetPromptAction() == EAsyncTaskNotificationPromptAction::Cancel;
-		if (TaskManager->SimulationTask->IsDone())
-		{
-			bFinished = true;
-		}
-		else if (bCancelled)
-		{
-			TaskManager->CancelSimulationGeneration();
-			bFinished = true;
-		}
-			
-		if (!bFinished)
-		{
-			const FDateTime CurrentTime = FDateTime::UtcNow();
-			const double SinceLastUpdate = (CurrentTime - TaskManager->LastUpdateTime).GetTotalSeconds();
-			if (SinceLastUpdate < 0.2)
-			{
-				return;
-			}
-			
-			const int32 NumSimulatedFrames = TaskManager->NumSimulatedFrames.load();
-			const int32 NumTotalFrames = TaskManager->NumFrames;
-			const FText ProgressMessage = FText::FromString(FString::Printf(TEXT("Finished %d/%d, %.1f%%"), NumSimulatedFrames, NumTotalFrames, 100.0 * NumSimulatedFrames / NumTotalFrames));
-			TaskManager->AsyncNotification->SetProgressText(ProgressMessage);
-			TaskManager->LastUpdateTime = CurrentTime;
-		}
-		else
-		{
-			FreeTaskResource(bCancelled);
-			PendingAction = EDataflowGeneratorActions::NoAction;
-		}
+		FreeTaskResource(bCancelled);
+		PendingAction = EDataflowGeneratorActions::NoAction;
 	}
+}
 
-	void FDataflowSimulationGenerator::SetSimulationGraph(const TObjectPtr<UDataflow>& DataflowGraph) 
-	{
-		SimulationGraph = DataflowGraph;
-	}
+void FDataflowSimulationGenerator::SetFrameRate(const int32 InFrameRate) 
+{
+	FrameRate = InFrameRate;
+}
 
-	void FDataflowSimulationGenerator::SetSamplingRate(const int32 FrameRate) 
+void FDataflowSimulationGenerator::SetTimeRange(const FVector2f& InTimeRange) 
+{
+	TimeRange = InTimeRange;
+}
+
+void FDataflowSimulationGenerator::SetCacheAsset(const TObjectPtr<UChaosCacheCollection>& InCacheAsset) 
+{
+	CacheAsset = InCacheAsset;
+}
+
+void FDataflowSimulationGenerator::SetActorClass(const TSubclassOf<AActor>& InActorClass)
+{
+	ActorClass = InActorClass;
+}
+
+void FDataflowSimulationGenerator::SetBackgroundTask(const bool bInBackgroundTask)
+{
+	bBackgroundTask = bInBackgroundTask;
+}
+
+void FDataflowSimulationGenerator::SetDataflowContent(const TObjectPtr<UDataflowBaseContent>& InDataflowContent)
+{
+	DataflowContent = InDataflowContent;
+}
+
+void FDataflowSimulationGenerator::RequestGeneratorAction(EDataflowGeneratorActions ActionType)
+{
+	if (PendingAction != EDataflowGeneratorActions::NoAction)
 	{
-		SamplingRate = FrameRate;
+		return;
 	}
-	
-	void FDataflowSimulationGenerator::SetCacheCollection(const TObjectPtr<UChaosCacheCollection>& ChaosCache) 
+	PendingAction = ActionType;
+}
+
+void FDataflowSimulationGenerator::FreeTaskResource(bool bCancelled)
+{
+	TaskManager->AsyncNotification->SetProgressText(LOCTEXT("Finishing", "Finishing, please wait"));
+	TaskManager->FreeSimulationResource();
+	const FDateTime CurrentTime = FDateTime::UtcNow();
+	UE_LOG(LogDataflowSimulationGenerator, Log, TEXT("Simulation finished in %f seconds"), (CurrentTime - TaskManager->StartTime).GetTotalSeconds());
+
 	{
-		CacheCollection = ChaosCache;
+		// TODO : save the result to the chaos cache collection
 	}
-	
-	void FDataflowSimulationGenerator::RequestGeneratorAction(EDataflowGeneratorActions ActionType)
+	if (bCancelled)
 	{
-		if (PendingAction != EDataflowGeneratorActions::NoAction)
-		{
-			return;
-		}
-		PendingAction = ActionType;
+		TaskManager->AsyncNotification->SetProgressText(LOCTEXT("Cancelled", "Cancelled"));
+		TaskManager->AsyncNotification->SetComplete(false);
 	}
-	
-	void FDataflowSimulationGenerator::FreeTaskResource(bool bCancelled)
+	else
 	{
-		TaskManager->AsyncNotification->SetProgressText(LOCTEXT("Finishing", "Finishing, please wait"));
-		TaskManager->FreeSimulationResource();
-		const FDateTime CurrentTime = FDateTime::UtcNow();
-		UE_LOG(LogDataflowSimulationGenerator, Log, TEXT("Simulation finished in %f seconds"), (CurrentTime - TaskManager->StartTime).GetTotalSeconds());
-	
-		{
-			// TODO : save the result to the chaos cache collection
-		}
-		if (bCancelled)
-		{
-			TaskManager->AsyncNotification->SetProgressText(LOCTEXT("Cancelled", "Cancelled"));
-			TaskManager->AsyncNotification->SetComplete(false);
-		}
-		else
-		{
-			TaskManager->AsyncNotification->SetProgressText(LOCTEXT("Finished", "Finished"));
-			TaskManager->AsyncNotification->SetComplete(true);
-		}
-		TaskManager.Reset();
+		TaskManager->AsyncNotification->SetProgressText(LOCTEXT("Finished", "Finished"));
+		TaskManager->AsyncNotification->SetComplete(true);
 	}
+	TaskManager.Reset();
+}
 };
 
 #undef LOCTEXT_NAMESPACE
diff --git a/Engine/Plugins/Experimental/Dataflow/Source/DataflowEditor/Private/Dataflow/DataflowSimulationScene.cpp b/Engine/Plugins/Experimental/Dataflow/Source/DataflowEditor/Private/Dataflow/DataflowSimulationScene.cpp
--- a/Engine/Plugins/Experimental/Dataflow/Source/DataflowEditor/Private/Dataflow/DataflowSimulationScene.cpp
+++ b/Engine/Plugins/Experimental/Dataflow/Source/DataflowEditor/Private/Dataflow/DataflowSimulationScene.cpp
@@ -1,10 +1,12 @@
 // Copyright Epic Games, Inc. All Rights Reserved.
 
 #include "Dataflow/DataflowSimulationScene.h"
+#include "Dataflow/DataflowSimulationManager.h"
+#include "Dataflow/DataflowEditor.h"
 #include "Chaos/CacheManagerActor.h"
 #include "Misc/TransactionObjectEvent.h"
-#include "Dataflow/DataflowEditor.h"
 #include "EngineUtils.h"
+#include "Animation/AnimSingleNodeInstance.h"
 
 #define LOCTEXT_NAMESPACE "FDataflowSimulationScene"
 
@@ -15,32 +17,28 @@
 FDataflowSimulationScene::FDataflowSimulationScene(FPreviewScene::ConstructionValues ConstructionValues, UDataflowEditor* InEditor)
 	: FDataflowPreviewSceneBase(ConstructionValues, InEditor)
 {
-	SimulationSceneDescription = NewObject<UDataflowSimulationSceneDescription>();
-	SimulationSceneDescription->SetSimulationScene(this);
+	SceneDescription = NewObject<UDataflowSimulationSceneDescription>();
+	SceneDescription->SetSimulationScene(this);
 
 	SimulationGenerator = MakeShared<Dataflow::FDataflowSimulationGenerator>();
+	RootSceneActor = GetWorld()->SpawnActor<AChaosCacheManager>();
 
-	const TObjectPtr<AChaosCacheManager> CacheManager = GetWorld()->SpawnActor<AChaosCacheManager>(AChaosCacheManager::StaticClass());
-
-	CacheManager->StartMode = EStartMode::Timed;
-	CacheManager->CacheMode = ECacheMode::None;
-	
-	RootSceneActor = CacheManager;
+	SceneDescription->ActorClass = GetEditorContent()->GetPreviewClass();
 
-	CreateSimulationWorld();
+	CreateSimulationScene();
 }
 
 FDataflowSimulationScene::~FDataflowSimulationScene()
 {
-	ResetSimulationWorld();
+	ResetSimulationScene();
 }
 
 void FDataflowSimulationScene::UnbindSceneSelection()
 {
-	if(SimulationContext && SimulationContext->GetRootActor() && SimulationContext->GetSimulationWorld())
+	if(PreviewActor)
 	{
 		TInlineComponentArray<UPrimitiveComponent*> PrimComponents;
-		SimulationContext->GetRootActor()->GetComponents(PrimComponents);
+		PreviewActor->GetComponents(PrimComponents);
 
 		for(UPrimitiveComponent* PrimComponent : PrimComponents)
 		{
@@ -49,24 +47,72 @@ void FDataflowSimulationScene::UnbindSceneSelection()
 	}
 }
 
-void FDataflowSimulationScene::ResetSimulationWorld()
+void FDataflowSimulationScene::ResetSimulationScene()
 {
-	if(SimulationContext && SimulationContext->GetRootActor())
+	// Destroy the spawned root actor
+	if(PreviewActor && GetWorld())
 	{
-		// Unbind the scene selection
-		UnbindSceneSelection();
-		
-		// Clear the simulation context
-		CleanSimulationContext(SimulationContext);
+		GetWorld()->DestroyActor(PreviewActor);
+
+		GetWorld()->EditorDestroyActor(PreviewActor, true);
+		// Since deletion can be delayed, rename to avoid future name collision
+		// Call UObject::Rename directly on actor to avoid AActor::Rename which unnecessarily sunregister and re-register components
+		PreviewActor->UObject::Rename(nullptr, GetTransientPackage(), REN_DontCreateRedirectors | REN_ForceNoResetLoaders);
+	}
+	
+	// Unbind the scene selection
+	UnbindSceneSelection();
+}
+
+void FDataflowSimulationScene::PauseSimulationScene() const
+{
+	if(SceneDescription && (SceneDescription->CacheAsset == nullptr))
+	{
+		GetWorld()->GetSubsystem<UDataflowSimulationManager>()->SetSimulationEnabled(false);
+		Dataflow::PauseSkeletonAnimation(PreviewActor);
+	}
+}
+
+void FDataflowSimulationScene::StartSimulationScene() const
+{
+	if(SceneDescription && (SceneDescription->CacheAsset == nullptr))
+	{
+		GetWorld()->GetSubsystem<UDataflowSimulationManager>()->SetSimulationEnabled(true);
+		Dataflow::StartSkeletonAnimation(PreviewActor);
+	}
+}
+
+void FDataflowSimulationScene::StepSimulationScene() const
+{
+	if(SceneDescription && (SceneDescription->CacheAsset == nullptr))
+	{
+		GetWorld()->GetSubsystem<UDataflowSimulationManager>()->SetSimulationEnabled(true);
+		GetWorld()->GetSubsystem<UDataflowSimulationManager>()->SetSimulationStepping(true);
+		Dataflow::StepSkeletonAnimation(PreviewActor);
+	}
+}
+
+void FDataflowSimulationScene::RebuildSimulationScene(const bool bIsSimulationEnabled)
+{
+	if(SceneDescription && (SceneDescription->CacheAsset == nullptr))
+	{
+		// Unregister components, cache manager, selection...
+		ResetSimulationScene();
+
+		// Register components, cache manager, selection...
+		CreateSimulationScene();
+
+		// Override the simulation enabled flag
+		GetWorld()->GetSubsystem<UDataflowSimulationManager>()->SetSimulationEnabled(bIsSimulationEnabled);
 	}
 }
 
 void FDataflowSimulationScene::BindSceneSelection()
 {
-	if(SimulationContext && SimulationContext->GetRootActor() && SimulationContext->GetSimulationWorld())
+	if(PreviewActor)
 	{
 		TInlineComponentArray<UPrimitiveComponent*> PrimComponents;
-		SimulationContext->GetRootActor()->GetComponents(PrimComponents);
+		PreviewActor->GetComponents(PrimComponents);
 		
 		for(UPrimitiveComponent* PrimComponent : PrimComponents)
 		{
@@ -76,33 +122,31 @@ void FDataflowSimulationScene::BindSceneSelection()
 	}
 }
 
-void FDataflowSimulationScene::CreateSimulationWorld()
+void FDataflowSimulationScene::CreateSimulationScene()
 {
-	if(const TObjectPtr<AChaosCacheManager> CacheManager = Cast<AChaosCacheManager>(RootSceneActor))
+	if(SimulationGenerator && SceneDescription && GetWorld())
 	{
-		if(const TObjectPtr<UDataflowBaseContent>& EditorContent = GetEditorContent())
-		{
-			if(SimulationGenerator && SimulationSceneDescription)
-			{
-				SimulationGenerator->SetSimulationGraph(EditorContent->GetDataflowAsset());
-				SimulationGenerator->SetSamplingRate(SimulationSceneDescription->SamplingRate);
-				SimulationGenerator->SetCacheCollection(SimulationSceneDescription->CacheAsset);
+		SimulationGenerator->SetFrameRate(SceneDescription->FrameRate);
+		SimulationGenerator->SetCacheAsset(SceneDescription->CacheAsset);
+		SimulationGenerator->SetActorClass(SceneDescription->ActorClass);
+		SimulationGenerator->SetTimeRange(SceneDescription->TimeRange);
+		SimulationGenerator->SetBackgroundTask(SceneDescription->bBackgroundTask);
+		SimulationGenerator->SetDataflowContent(GetEditorContent());
+
+		TimeRange = SceneDescription->TimeRange;
+		NumFrames = (TimeRange[1] > TimeRange[0]) ? FMath::Floor((TimeRange[1] - TimeRange[0]) * SceneDescription->FrameRate) : 0;
 		
-				CacheManager->CacheCollection = SimulationSceneDescription->CacheAsset;
-			}
+		PreviewActor = Dataflow::SpawnSimulatedActor(SceneDescription->ActorClass, Cast<AChaosCacheManager>(RootSceneActor),
+			SceneDescription->CacheAsset, false, GetEditorContent());
 
-			// Build the simulation context
-			BuildSimulationContext(EditorContent->GetDataflowOwner(), EditorContent->GetDataflowAsset(),
-				CacheManager, Dataflow::ECachingMode::PlayCache, SimulationContext);
+		// Setup all the skelmesh animations
+		Dataflow::SetupSkeletonAnimation(PreviewActor);
 		
-			// Update the simulation time ranges and number of frames
-			GetCacheDuration(EditorContent->GetDataflowAsset(), SimulationContext, SimulationSceneDescription->SamplingRate,
-				TimeRange[0], TimeRange[1], NumFrames);
-
-			// update the selection binding since we are constantly editing the graph
-			BindSceneSelection();
-		}
+		GetWorld()->GetSubsystem<UDataflowSimulationManager>()->SetSimulationEnabled(SceneDescription->CacheAsset == nullptr);
 	}
+
+	// update the selection binding since we are constantly editing the graph
+	BindSceneSelection();
 }
 
 void FDataflowSimulationScene::UpdateSimulationCache()
@@ -119,20 +163,28 @@ void FDataflowSimulationScene::TickDataflowScene(const float DeltaSeconds)
 
 	if(const TObjectPtr<UDataflowBaseContent>& EditorContent = GetEditorContent())
 	{
-		if( Dataflow::ShouldResetWorld(EditorContent->GetDataflowAsset(), SimulationContext,LastTimeStamp))
+		if(Dataflow::ShouldResetWorld(EditorContent->GetDataflowAsset(), GetWorld(), LastTimeStamp) || EditorContent->IsSimulationDirty())
 		{
-			// unregister components, cache manager, selection...
-			ResetSimulationWorld();
+			// Unregister components, cache manager, selection...
+			ResetSimulationScene();
 
-			// register components, cache manager, selection...
-			CreateSimulationWorld();
-		}
+			// Register components, cache manager, selection...
+			CreateSimulationScene();
 
+			// Reset the dirty flag
+			EditorContent->SetSimulationDirty(false);
+		}
 		// Load the cache at some point in time
-		Cast<AChaosCacheManager>(RootSceneActor)->SetStartTime(SimulationTime);
-
-		// Update all the animation at the simulation time
-		Dataflow::UpdateAnimationNodes(EditorContent->GetDataflowAsset(), SimulationContext, SimulationTime);
+		if(SceneDescription->CacheAsset)
+		{
+			// Update the cached simulation at some point in time
+			if(RootSceneActor)
+			{
+				Cast<AChaosCacheManager>(RootSceneActor)->SetStartTime(SimulationTime);
+			}
+			// Update all the skelmesh animations at the simulation time
+			Dataflow::UpdateSkeletonAnimation(PreviewActor, SimulationTime);
+		}
 	}
 }
 
@@ -140,27 +192,51 @@ void FDataflowSimulationScene::AddReferencedObjects(FReferenceCollector& Collect
 {
 	FDataflowPreviewSceneBase::AddReferencedObjects(Collector);
 
-	Collector.AddReferencedObject(SimulationSceneDescription);
+	Collector.AddReferencedObject(SceneDescription);
 }
 
 void FDataflowSimulationScene::SceneDescriptionPropertyChanged(const FName& PropertyName)
 {
-	if (PropertyName == GET_MEMBER_NAME_CHECKED(UDataflowSimulationSceneDescription, SamplingRate))
+	if (PropertyName == GET_MEMBER_NAME_CHECKED(UDataflowSimulationSceneDescription, FrameRate))
 	{
 		if(SimulationGenerator)
 		{
-			SimulationGenerator->SetSamplingRate(SimulationSceneDescription->SamplingRate);
+			SimulationGenerator->SetFrameRate(SceneDescription->FrameRate);
+		}
+	}
+	else if (PropertyName == GET_MEMBER_NAME_CHECKED(UDataflowSimulationSceneDescription, TimeRange))
+	{
+		if(SimulationGenerator)
+		{
+			SimulationGenerator->SetTimeRange(SceneDescription->TimeRange);
 		}
 	}
 	else if(PropertyName == GET_MEMBER_NAME_CHECKED(UDataflowSimulationSceneDescription, CacheAsset))
 	{
 		if(SimulationGenerator)
 		{
-			SimulationGenerator->SetCacheCollection(SimulationSceneDescription->CacheAsset);
-
-			Cast<AChaosCacheManager>(RootSceneActor)->CacheCollection = SimulationSceneDescription->CacheAsset;
+			SimulationGenerator->SetCacheAsset(SceneDescription->CacheAsset);
+		}
+	}
+	else if(PropertyName == GET_MEMBER_NAME_CHECKED(UDataflowSimulationSceneDescription, ActorClass))
+	{
+		if(SimulationGenerator)
+		{
+			SimulationGenerator->SetActorClass(SceneDescription->ActorClass);
 		}
 	}
+	else if(PropertyName == GET_MEMBER_NAME_CHECKED(UDataflowSimulationSceneDescription, bBackgroundTask))
+	{
+		if(SimulationGenerator)
+		{
+			SimulationGenerator->SetBackgroundTask(SceneDescription->bBackgroundTask);
+		}
+	}
+	// Unregister components, cache manager, selection...
+	ResetSimulationScene();
+
+	// Register components, cache manager, selection...
+	CreateSimulationScene();
 }
 
 void UDataflowSimulationSceneDescription::SetSimulationScene(FDataflowSimulationScene* InSimulationScene)
diff --git a/Engine/Plugins/Experimental/Dataflow/Source/DataflowEditor/Public/Dataflow/DataflowSimulationGenerator.h b/Engine/Plugins/Experimental/Dataflow/Source/DataflowEditor/Public/Dataflow/DataflowSimulationGenerator.h
--- a/Engine/Plugins/Experimental/Dataflow/Source/DataflowEditor/Public/Dataflow/DataflowSimulationGenerator.h
+++ b/Engine/Plugins/Experimental/Dataflow/Source/DataflowEditor/Public/Dataflow/DataflowSimulationGenerator.h
@@ -8,6 +8,7 @@
 #include "Logging/LogMacros.h"
 #include "Misc/AsyncTaskNotification.h"
 #include "Dataflow/DataflowContent.h"
+#include "Dataflow/DataflowSimulationUtils.h"
 
 DECLARE_LOG_CATEGORY_EXTERN(LogDataflowSimulationGenerator, Log, All);
 
@@ -49,11 +50,11 @@ namespace Dataflow
 		/** Simulation max time */
 		float MaxTime = TNumericLimits<float>::Lowest();
 
-		/** Simulation graph to be executed */
-		TObjectPtr<UDataflow> SimulationGraph;
+		/** Simulation world */
+		UWorld* SimulationWorld = nullptr;
 
-		/** Simulation  context used during graph evaluation */
-		TSharedPtr<Dataflow::FSimulationContext> SimulationContext;
+		/** Boolean to check if we are running the task in the background */
+		bool bBackgroundTask = true;
 	};
 
 	/** Async simulation resource that will be used while simulating */
@@ -103,7 +104,9 @@ namespace Dataflow
 		FDateTime LastUpdateTime;
 
 		/** Allocate the simulation resource from the properties */
-        bool AllocateSimulationResource(const TObjectPtr<UDataflow> SimulationGraph, const int32 SamplingRate, const TObjectPtr<UChaosCacheCollection>& CacheCollection);
+        bool AllocateSimulationResource(const FVector2f& TimeRange, const int32 FrameRate,
+        	const TObjectPtr<UChaosCacheCollection>& CacheAsset, const TSubclassOf<AActor>& ActorClass,
+        	const TObjectPtr<UDataflowBaseContent>& DataflowContent);
 
         /** Free the simulation resource */
         void FreeSimulationResource();
@@ -119,6 +122,12 @@ namespace Dataflow
 		
 		/** Temporary world created to run the simulation */
 		UWorld* SimulationWorld = nullptr;
+
+		/** Temporary cache manager created to run the simulation */
+		TObjectPtr<AChaosCacheManager> CacheManager = nullptr;
+
+		/** Temporary cache manager created to run the simulation */
+		TObjectPtr<AActor> PreviewActor = nullptr;
 	};
 
 	/** Enum for all the generator actions */
@@ -142,14 +151,23 @@ namespace Dataflow
 		virtual ETickableTickType GetTickableTickType() const override { return ETickableTickType::Always; }
 		//~ End FTickableEditorObject Interface
 
-		/** Set the simulation graph to be executed on the async thread */
-		void SetSimulationGraph(const TObjectPtr<UDataflow>& DataflowGraph);
+		/** Set the frame rate for cache recording */
+		void SetFrameRate(const int32 InFrameRate);
+
+		/** Set the time range for cache recording */
+		void SetTimeRange(const FVector2f& InTimeRange);
 
-		/** Set the sampling rate for cache recording */
-		void SetSamplingRate(const int32 FrameRate);
+		/** Set the cache asset for cache recording */
+		void SetCacheAsset(const TObjectPtr<UChaosCacheCollection>& InCacheAsset);
 
-		/** Set the cache collection for cache recording */
-		void SetCacheCollection(const TObjectPtr<UChaosCacheCollection>&);
+		/** Set the actor class for cache recording */
+		void SetActorClass(const TSubclassOf<AActor>& InActorClass);
+
+		/** Set the dataflow content */
+		void SetDataflowContent(const TObjectPtr<UDataflowBaseContent>& InDataflowContent);
+
+		/** Set the background task boolean */
+		void SetBackgroundTask(const bool bInBackgroundTask);
 
 		/** Enqueue a generator action to be processed on the async thread */
 		void RequestGeneratorAction(EDataflowGeneratorActions Action);
@@ -165,14 +183,23 @@ namespace Dataflow
 		/** Free the task resource used while generating the simulation */
 		void FreeTaskResource(bool bCancelled);
 
-		/** Graph to be used to generate the simulation output */
-		TObjectPtr<UDataflow> SimulationGraph = nullptr;
+		/** Cache asset to store the caches */
+		TObjectPtr<UChaosCacheCollection> CacheAsset = nullptr;
+
+		/** Frame rate used to record the cache */
+		int32 FrameRate = 30;
+
+		/** Time range used to generate the cache */
+		FVector2f TimeRange = FVector2f(0.0f, 5.0f);
 
-		/** Cache collection to store the caches */
-		TObjectPtr<UChaosCacheCollection> CacheCollection = nullptr;
+		/** Boolean to control if the task is going to be run in the background */
+		bool bBackgroundTask = false;
 
-		/** Sampling rate used to record the cache */
-		int32 SamplingRate = 30;
+		/** Chaos cache manager BP class to be spawned */
+		TSubclassOf<AActor> ActorClass;
+
+		/** Dataflow content */
+		TObjectPtr<UDataflowBaseContent> DataflowContent;
 
 		/** Pending action to be send to the async thread */
 		EDataflowGeneratorActions PendingAction = EDataflowGeneratorActions::NoAction;
@@ -180,28 +207,6 @@ namespace Dataflow
 		/** Task manager to run the async tasks */
 		TSharedPtr<FDataflowTaskManager> TaskManager = nullptr;
 	};
-
-	/** Evaluate the simulation graph given a simulation context and a simulation flag */
-	void EvaluateSimulationGraph(const TObjectPtr<UDataflow>& SimulationGraph,
-		const TSharedPtr<Dataflow::FSimulationContext>& SimulationContext, const float DeltaTime,
-		const float SimulationTime, const ESimulationFlags SimulationFlags);
-
-	/** Extract the time range, number of frames and duration of the cache sequence */
-	float GetCacheDuration(const TObjectPtr<UDataflow>& SimulationGraph, const TSharedPtr<Dataflow::FSimulationContext>& SimulationContext,
-		const int32 SamplingRate, float& MinTime, float& MaxTime, int32& NumFrames);
-
-	/** Check if the simulation cache nodes have changed to trigger a reset */
-	bool ShouldResetWorld(const TObjectPtr<UDataflow>& SimulationGraph, const TSharedPtr<Dataflow::FSimulationContext>& SimulationContext, Dataflow::FTimestamp& LastTimeStamp);
-
-	/** Update the animation nodes given a context and an animation time */
-	void UpdateAnimationNodes(const TObjectPtr<UDataflow>& SimulationGraph, const TSharedPtr<Dataflow::FSimulationContext>& SimulationContext, const float AnimationTime);
-
-	/** clean the context (unregister componentds, clesn cache, reset context)*/
-	void CleanSimulationContext(TSharedPtr<Dataflow::FSimulationContext>& SimulationContext);
-
-	/** Build simulaiton context */
-	void BuildSimulationContext(const TObjectPtr<UObject>& SimulationOwner, const TObjectPtr<UDataflow>& SimulationGraph, const TObjectPtr<AChaosCacheManager>& CacheManager,
-		const ECachingMode CachingMode, TSharedPtr<Dataflow::FSimulationContext>& SimulationContext);
 };
 
 
diff --git a/Engine/Plugins/Experimental/Dataflow/Source/DataflowEditor/Public/Dataflow/DataflowSimulationScene.h b/Engine/Plugins/Experimental/Dataflow/Source/DataflowEditor/Public/Dataflow/DataflowSimulationScene.h
--- a/Engine/Plugins/Experimental/Dataflow/Source/DataflowEditor/Public/Dataflow/DataflowSimulationScene.h
+++ b/Engine/Plugins/Experimental/Dataflow/Source/DataflowEditor/Public/Dataflow/DataflowSimulationScene.h
@@ -34,7 +34,19 @@ public:
 
 	/** Number of sampling frames per second for caching*/
 	UPROPERTY(EditAnywhere, Category="Caching")
-	int32 SamplingRate = 30;
+	int32 FrameRate = 30;
+
+	/** Number of sampling frames per second for caching*/
+	UPROPERTY(EditAnywhere, Category="Caching")
+	FVector2f TimeRange = FVector2f(0.0f, 5.0f);
+
+	/** Caching actor class to spawn */
+	UPROPERTY(EditAnywhere, Category = "Caching")
+	TSubclassOf<AActor> ActorClass;
+
+	/** Boolean to check if the caching will be done on an async thread (if yes no GT dependency) */
+	UPROPERTY(EditAnywhere, AdvancedDisplay, Category = "Caching")
+	bool bBackgroundTask = true;
 
 private:
 
@@ -46,7 +58,6 @@ private:
 	class FDataflowSimulationScene* SimulationScene;
 };
 
-
 /**
  * Dataflow simulation scene holding all the dataflow content components
  */
@@ -68,16 +79,25 @@ public:
 	virtual bool CanRunSimulation() const { return true; }
 
 	/** Get the scene description used in the preview scene widget */
-	UDataflowSimulationSceneDescription* GetPreviewSceneDescription() const { return SimulationSceneDescription; }
+	UDataflowSimulationSceneDescription* GetPreviewSceneDescription() const { return SceneDescription; }
 
 	/** Create all the simulation world components and instances */
-	void CreateSimulationWorld();
+	void CreateSimulationScene();
 
 	/** Reset all the simulation world components and instances */
-	void ResetSimulationWorld();
+	void ResetSimulationScene();
 
-	/** Reset the simulation scene */
-	void ResetSimulationScene() {}
+	/** Pause the simulation */
+	void PauseSimulationScene() const;
+
+	/** Start the simulation */
+	void StartSimulationScene() const;
+
+	/** Step the simulation */
+	void StepSimulationScene() const;
+
+	/** Rebuild the simulation scene */
+	void RebuildSimulationScene(const bool bIsSimulationEnabled);
 
 	/** Check if there is something to render */
 	bool HasRenderableGeometry() { return true; }
@@ -106,14 +126,11 @@ private:
 	void UnbindSceneSelection();
 	
 	/** Simulation scene description */
-	TObjectPtr<UDataflowSimulationSceneDescription> SimulationSceneDescription;
+	TObjectPtr<UDataflowSimulationSceneDescription> SceneDescription;
 
 	/** Simulation generator to record the simulation result */
 	TSharedPtr<Dataflow::FDataflowSimulationGenerator> SimulationGenerator;
 
-	/** Simulation context used for the simulation graph */
-	TSharedPtr<Dataflow::FSimulationContext> SimulationContext;
-
 	/** Cache time range in seconds */
 	FVector2f TimeRange;
 
@@ -122,6 +139,9 @@ private:
 
 	/** Last context time stamp for which we regenerated the world */
 	Dataflow::FTimestamp LastTimeStamp = Dataflow::FTimestamp::Invalid;
+
+	/** Preview actor that will will be used to visualize the result of the simulation graph */
+	TObjectPtr<AActor> PreviewActor;
 };
 
 
diff --git a/Engine/Plugins/Experimental/Dataflow/Source/DataflowNodes/Private/Dataflow/DataflowNodesPlugin.cpp b/Engine/Plugins/Experimental/Dataflow/Source/DataflowNodes/Private/Dataflow/DataflowNodesPlugin.cpp
--- a/Engine/Plugins/Experimental/Dataflow/Source/DataflowNodes/Private/Dataflow/DataflowNodesPlugin.cpp
+++ b/Engine/Plugins/Experimental/Dataflow/Source/DataflowNodes/Private/Dataflow/DataflowNodesPlugin.cpp
@@ -10,7 +10,6 @@
 #include "Dataflow/DataflowSelectionNodes.h"
 #include "Dataflow/DataflowContextOverridesNodes.h"
 #include "Dataflow/DataflowCollectionAddScalarVertexPropertyNode.h"
-#include "Dataflow/DataflowSolverNodes.h"
 
 #define LOCTEXT_NAMESPACE "DataflowNodes"
 
@@ -25,10 +24,6 @@ void IDataflowNodesPlugin::StartupModule()
 	DATAFLOW_NODE_REGISTER_CREATION_FACTORY(FDataflowCollectionAddScalarVertexPropertyNode);
 
 	Dataflow::RegisterNodeFilter(FDataflowTerminalNode::StaticType());
-	Dataflow::RegisterNodeFilter(FDataflowSimulationNode::StaticType());
-	Dataflow::RegisterNodeFilter(FDataflowInvalidNode::StaticType());
-	Dataflow::RegisterNodeFilter(FDataflowCacheNode::StaticType());
-	Dataflow::RegisterNodeFilter(FDataflowAnimationNode::StaticType());
 }
 
 void IDataflowNodesPlugin::ShutdownModule()
diff --git a/Engine/Plugins/Experimental/Dataflow/Source/DataflowNodes/Private/Dataflow/DataflowSolverNodes.cpp b/Engine/Plugins/Experimental/Dataflow/Source/DataflowNodes/Private/Dataflow/DataflowSolverNodes.cpp
deleted file mode 100644
diff --git a/Engine/Plugins/Experimental/Dataflow/Source/DataflowNodes/Public/Dataflow/DataflowSolverNodes.h b/Engine/Plugins/Experimental/Dataflow/Source/DataflowNodes/Public/Dataflow/DataflowSolverNodes.h
deleted file mode 100644
diff --git a/Engine/Plugins/Experimental/Dataflow/Source/DataflowSimulation/DataflowSimulation.Build.cs b/Engine/Plugins/Experimental/Dataflow/Source/DataflowSimulation/DataflowSimulation.Build.cs
new file mode 100644
--- /dev/null
+++ b/Engine/Plugins/Experimental/Dataflow/Source/DataflowSimulation/DataflowSimulation.Build.cs
@@ -0,0 +1,38 @@
+﻿// Copyright Epic Games, Inc. All Rights Reserved.
+
+namespace UnrealBuildTool.Rules
+{
+	public class DataflowSimulation : ModuleRules
+	{
+		public DataflowSimulation(ReadOnlyTargetRules Target) : base(Target)
+		{
+			bTreatAsEngineModule = true;
+			PublicDependencyModuleNames.AddRange(
+				new string[]
+				{
+					"Core",
+					"CoreUObject",
+					"DataflowCore",
+					"DataflowEngine",
+					"DataflowNodes"
+				}
+			);
+			PrivateDependencyModuleNames.AddRange(
+				new string[]
+				{
+					// NOTE: UVEditorTools is a separate module so that it doesn't rely on the editor.
+					// So, do not add editor dependencies here.
+	
+					"Engine",
+					"RenderCore",
+					"RHI",
+					"DataflowCore",
+					"DataflowEngine",
+					"DataflowNodes",
+					"Chaos",
+					"ChaosCaching"
+				}
+			);
+		}
+	}
+}
\ No newline at end of file
diff --git a/Engine/Plugins/Experimental/Dataflow/Source/DataflowSimulation/Private/Dataflow/DataflowSimulationContext.cpp b/Engine/Plugins/Experimental/Dataflow/Source/DataflowSimulation/Private/Dataflow/DataflowSimulationContext.cpp
new file mode 100644
--- /dev/null
+++ b/Engine/Plugins/Experimental/Dataflow/Source/DataflowSimulation/Private/Dataflow/DataflowSimulationContext.cpp
@@ -0,0 +1,95 @@
+﻿// Copyright Epic Games, Inc. All Rights Reserved.
+
+#include "Dataflow/DataflowSimulationContext.h"
+#include "Dataflow/DataflowSimulationProxy.h"
+
+namespace Dataflow
+{
+template<typename Base>
+void TSimulationContext<Base>::GetSimulationProxies(const FString& ProxyType, const TArray<FString>& SimulationGroups, TArray<FDataflowSimulationProxy*>& FilteredProxies) const
+{
+	TBitArray<> ProxyGroups(false, GroupIndices.Num());
+	BuildGroupBits(SimulationGroups, ProxyGroups);
+	
+	if(const TSet<FDataflowSimulationProxy*>* TypedProxies = SimulationProxies.Find(ProxyType))
+	{
+		for(FDataflowSimulationProxy* SimulationProxy : *TypedProxies)
+		{
+			if(SimulationProxy && SimulationProxy->HasGroupBit(ProxyGroups))
+			{
+				FilteredProxies.Add({SimulationProxy});
+			}
+		}
+	}
+}
+
+template<typename Base>
+void TSimulationContext<Base>::BuildGroupBits(const TArray<FString>& SimulationGroups, TBitArray<>& GroupBits) const
+{
+	GroupBits.Init(false, GroupIndices.Num());
+    for(const FString& SimulationGroup : SimulationGroups)
+    {
+    	if(const uint32* GroupIndex = GroupIndices.Find(SimulationGroup))
+    	{
+    		GroupBits[*GroupIndex] = true;
+    	}
+    }
+}
+	
+template<typename Base>
+void TSimulationContext<Base>::RegisterProxyGroups()
+{
+	GroupIndices.Reset();
+	TArray<uint32> IndexArray;
+	
+	for(const TPair<FString, TSet<FDataflowSimulationProxy*>>& TypedProxies : SimulationProxies)
+	{
+		for(FDataflowSimulationProxy* SimulationProxy : TypedProxies.Value)
+		{
+			if(SimulationProxy)
+			{
+				IndexArray.Reset(SimulationProxy->GetSimulationGroups().Num());
+				for(const FString& SimulationGroup : SimulationProxy->GetSimulationGroups())
+				{
+					const uint32* IndexPtr = GroupIndices.Find(SimulationGroup);
+					IndexArray.Add( (IndexPtr == nullptr) ?
+							GroupIndices.Add(SimulationGroup, GroupIndices.Num()) : *IndexPtr);
+				}
+				SimulationProxy->GroupBits.Init(false, GroupIndices.Num());
+				for(const uint32& GroupIndex : IndexArray)
+				{
+					SimulationProxy->GroupBits[GroupIndex] = true;
+				}
+			}
+		}
+	}
+}
+
+template<typename Base>
+void TSimulationContext<Base>::AddSimulationProxy(const FString& ProxyType, FDataflowSimulationProxy* SimulationProxy)
+{
+	SimulationProxies.FindOrAdd(ProxyType).Add(SimulationProxy);
+}
+	
+template<typename Base>
+void TSimulationContext<Base>::RemoveSimulationProxy(const FString& ProxyType, const FDataflowSimulationProxy* SimulationProxy)
+{
+	SimulationProxies.FindOrAdd(ProxyType).Remove(SimulationProxy);
+}
+
+template<typename Base>
+void TSimulationContext<Base>::ResetSimulationProxies()
+{
+	SimulationProxies.Reset();
+}
+	
+template<typename Base>
+int32 TSimulationContext<Base>::NumSimulationProxies(const FString& ProxyType) const
+{
+	const TSet<FDataflowSimulationProxy*>* TypedProxies = SimulationProxies.Find(ProxyType);
+	return TypedProxies ? TypedProxies->Num() : 0;
+}
+
+template class TSimulationContext<FContextSingle>;
+template class TSimulationContext<FContextThreaded>;
+}
\ No newline at end of file
diff --git a/Engine/Plugins/Experimental/Dataflow/Source/DataflowSimulation/Private/Dataflow/DataflowSimulationInterface.cpp b/Engine/Plugins/Experimental/Dataflow/Source/DataflowSimulation/Private/Dataflow/DataflowSimulationInterface.cpp
new file mode 100644
--- /dev/null
+++ b/Engine/Plugins/Experimental/Dataflow/Source/DataflowSimulation/Private/Dataflow/DataflowSimulationInterface.cpp
@@ -0,0 +1,39 @@
+﻿// Copyright Epic Games, Inc. All Rights Reserved.
+
+#include "Dataflow/DataflowSimulationInterface.h"
+#include "Dataflow/DataflowSimulationManager.h"
+#include "Engine/World.h"
+	
+void IDataflowSimulationInterface::RegisterManagerInterface(const TObjectPtr<UWorld>& SimulationWorld)
+{
+	if(SimulationWorld)
+	{
+		if (UDataflowSimulationManager* SimulationManager = SimulationWorld->GetSubsystem<UDataflowSimulationManager>())
+		{
+			SimulationManager->AddSimulationInterface(this);
+		}
+	}
+}
+
+bool IDataflowSimulationInterface::IsInterfaceRegistered(const TObjectPtr<UWorld>& SimulationWorld) const
+{
+	if(SimulationWorld)
+	{
+		if (const UDataflowSimulationManager* SimulationManager = SimulationWorld->GetSubsystem<UDataflowSimulationManager>())
+		{
+			return SimulationManager->HasSimulationInterface(this);
+		}
+	}
+	return false;
+}
+
+void IDataflowSimulationInterface::UnregisterManagerInterface(const TObjectPtr<UWorld>& SimulationWorld) const
+{
+	if(SimulationWorld)
+	{
+		if (UDataflowSimulationManager* SimulationManager = SimulationWorld->GetSubsystem<UDataflowSimulationManager>())
+		{
+			SimulationManager->RemoveSimulationInterface(this); 
+		}
+	}
+}
diff --git a/Engine/Plugins/Experimental/Dataflow/Source/DataflowSimulation/Private/Dataflow/DataflowSimulationManager.cpp b/Engine/Plugins/Experimental/Dataflow/Source/DataflowSimulation/Private/Dataflow/DataflowSimulationManager.cpp
new file mode 100644
--- /dev/null
+++ b/Engine/Plugins/Experimental/Dataflow/Source/DataflowSimulation/Private/Dataflow/DataflowSimulationManager.cpp
@@ -0,0 +1,383 @@
+﻿// Copyright Epic Games, Inc. All Rights Reserved.
+
+#include "Dataflow/DataflowSimulationManager.h"
+#include "Dataflow/DataflowSimulationUtils.h"
+#include "Dataflow/DataflowSimulationInterface.h"
+#include "Dataflow/DataflowSimulationProxy.h"
+#include "Components/ActorComponent.h"
+#include "Containers/Map.h"
+#include "Engine/World.h"
+
+static FAutoConsoleTaskPriority DataflowSimulationTaskPriority(
+		TEXT("TaskGraph.TaskPriorities.DataflowSimulationTask"),
+		TEXT("Task and thread priority for the dataflow simulation."),
+		ENamedThreads::HighThreadPriority, // If we have high priority task threads, then use them...
+		ENamedThreads::NormalTaskPriority, // .. at normal task priority
+		ENamedThreads::HighTaskPriority);  // If we don't have high priority threads, then use normal priority threads at high task priority instead
+
+namespace Dataflow
+{
+	enum class ESimulationThreadingMode : uint8
+	{
+		GameThread,
+		BlockingThread,
+		AsyncThread,
+	};
+	
+	namespace CVars
+	{
+		/** Simulation threading mode */
+		int32  DataflowSimulationThreadingMode = (uint8)ESimulationThreadingMode::AsyncThread;
+		FAutoConsoleVariableRef CVarDataflowSimulationThreadingMode(TEXT("p.Dataflow.Simulation.ThreadingMode"), DataflowSimulationThreadingMode,
+			TEXT("0 : run simulation on GT | 1 : run simulation on PT (GT is blocked in manager Tick) | 2 : run simulation on PT (GT will be blocked at the end of the world tick)"));
+	}
+	
+class FDataflowSimulationTask
+{
+public:
+	FDataflowSimulationTask(const TObjectPtr<UDataflow>& InDataflowAsset, const TSharedPtr<FDataflowSimulationContext>& InSimulationContext,
+		const float InDeltaTime, const float InSimulationTime)
+		: DataflowAsset(InDataflowAsset), SimulationContext(InSimulationContext), DeltaTime(InDeltaTime), SimulationTime(InSimulationTime)
+	{
+	}
+
+	TStatId GetStatId() const
+	{
+		RETURN_QUICK_DECLARE_CYCLE_STAT(FDataflowSimulationProxyParallelTask, STATGROUP_TaskGraphTasks);
+	}
+
+	static ENamedThreads::Type GetDesiredThread()
+	{
+		if (DataflowSimulationTaskPriority.Get() != 0)
+		{
+			return DataflowSimulationTaskPriority.Get();
+		}
+		return ENamedThreads::GameThread;
+	}
+
+	static ESubsequentsMode::Type GetSubsequentsMode()
+	{
+		return ESubsequentsMode::TrackSubsequents;
+	}
+
+	void DoTask(ENamedThreads::Type CurrentThread, const FGraphEventRef& MyCompletionGraphEvent)
+	{
+		Dataflow::EvaluateSimulationGraph(DataflowAsset, SimulationContext, DeltaTime, SimulationTime);
+	}
+
+private:
+	/** Dataflow graph asset used to launch the simulation */
+	TObjectPtr<UDataflow> DataflowAsset;
+
+	/** Simulation context */
+	TSharedPtr<Dataflow::FDataflowSimulationContext> SimulationContext;
+
+	/** Delta time used to advance the simulation */
+	float DeltaTime;
+
+	/** World simulation time  */
+	float SimulationTime;
+};
+}
+
+FDelegateHandle UDataflowSimulationManager::OnObjectPropertyChangedHandle;
+FDelegateHandle UDataflowSimulationManager::OnWorldTickEndHandle;
+FDelegateHandle UDataflowSimulationManager::OnCreatePhysicsStateHandle;
+FDelegateHandle UDataflowSimulationManager::OnDestroyPhysicsStateHandle;
+
+UDataflowSimulationManager::UDataflowSimulationManager()
+{}
+
+void UDataflowSimulationManager::Tick(float DeltaTime)
+{
+	Super::Tick(DeltaTime);
+	
+	check(IsInGameThread());
+	check(SimulationTasks.IsEmpty());
+	
+	// Transfer data from GT -> PT
+    WriteSimulationData(DeltaTime);
+	
+	if(bIsSimulationEnabled)
+	{
+		if(Dataflow::CVars::DataflowSimulationThreadingMode == static_cast<uint8>(Dataflow::ESimulationThreadingMode::GameThread))
+		{
+			// Advance the simulation in time
+			AdvanceSimulationData(DeltaTime, GetWorld()->GetTimeSeconds());
+
+			// Transfer data from PT -> GT
+			ReadSimulationData(DeltaTime);
+		}
+		else
+		{
+			// Start all the simulation tasks in parallel
+			StartSimulationTasks(DeltaTime, GetWorld()->GetTimeSeconds());
+
+			if(Dataflow::CVars::DataflowSimulationThreadingMode == static_cast<uint8>(Dataflow::ESimulationThreadingMode::BlockingThread))
+            {
+            	// Wait until all tasks are complete
+            	CompleteSimulationTasks();
+            
+            	// Transfer data from PT -> GT
+            	ReadSimulationData(DeltaTime);
+            }
+		}
+	}
+}
+
+void UDataflowSimulationManager::OnStartup()
+{
+	OnWorldTickEndHandle = FWorldDelegates::OnWorldTickEnd.AddLambda(
+	[](const UWorld* SimulationWorld, ELevelTick LevelTick, const float DeltaSeconds)
+	{
+		if(SimulationWorld)
+		{
+			if (UDataflowSimulationManager* DataflowManager = SimulationWorld->GetSubsystem<UDataflowSimulationManager>())
+			{
+				if(DataflowManager->bIsSimulationEnabled)
+				{
+					if(Dataflow::CVars::DataflowSimulationThreadingMode == static_cast<uint8>(Dataflow::ESimulationThreadingMode::AsyncThread))
+					{
+						// Wait until all tasks are complete
+						DataflowManager->CompleteSimulationTasks();
+					
+						// Transfer data from PT -> GT
+						DataflowManager->ReadSimulationData(DeltaSeconds);
+					}
+				}
+			}
+		}
+		
+	});
+
+	OnCreatePhysicsStateHandle = UActorComponent::GlobalCreatePhysicsDelegate.AddLambda([](UActorComponent* ActorComponent)
+	{
+		if(IDataflowSimulationInterface* SimulationInterface = Cast<IDataflowSimulationInterface>(ActorComponent))
+		{
+			// Build the simulation proxy
+			SimulationInterface->BuildSimulationProxy();
+
+			// Register the simulation interface to the manager
+			SimulationInterface->RegisterManagerInterface(ActorComponent->GetWorld());
+		}
+	});
+
+	OnDestroyPhysicsStateHandle = UActorComponent::GlobalDestroyPhysicsDelegate.AddLambda([](UActorComponent* ActorComponent)
+	{
+		if(IDataflowSimulationInterface* SimulationInterface = Cast<IDataflowSimulationInterface>(ActorComponent))
+		{
+			// Reset the simulation proxy
+			SimulationInterface->ResetSimulationProxy();
+
+			// Unregister the simulation interface from the manager
+			SimulationInterface->UnregisterManagerInterface(ActorComponent->GetWorld());
+		}
+	});
+
+#if WITH_EDITOR
+	OnObjectPropertyChangedHandle = FCoreUObjectDelegates::OnObjectPropertyChanged.AddLambda([](UObject* ModifiedObject, FPropertyChangedEvent& ChangedProperty)
+	{
+		if(IDataflowSimulationInterface* SimulationInterface = Cast<IDataflowSimulationInterface>(ModifiedObject))
+		{
+			if(!SimulationInterface->IsInterfaceRegistered(ModifiedObject->GetWorld()))
+			{
+				// Unregister the simulation interface from the manager
+				SimulationInterface->UnregisterManagerInterface(ModifiedObject->GetWorld());
+                
+                // Register the simulation interface to the manager
+                SimulationInterface->RegisterManagerInterface(ModifiedObject->GetWorld());
+			}
+		}
+	});
+#endif
+}
+
+void UDataflowSimulationManager::OnShutdown()
+{
+#if WITH_EDITOR
+	FCoreUObjectDelegates::OnObjectPropertyChanged.Remove(OnObjectPropertyChangedHandle);
+#endif
+	FWorldDelegates::OnWorldTickEnd.Remove(OnWorldTickEndHandle);
+	UActorComponent::GlobalCreatePhysicsDelegate.Remove(OnCreatePhysicsStateHandle);
+	UActorComponent::GlobalDestroyPhysicsDelegate.Remove(OnDestroyPhysicsStateHandle);
+}
+
+void UDataflowSimulationManager::WriteSimulationData(const float DeltaTime)
+{
+	for(const TPair<TObjectPtr<UDataflow>, Dataflow::FDataflowSimulationData>& DataflowData : SimulationData)
+	{
+		for(const TPair<FString,TSet<IDataflowSimulationInterface*>>& SimulationInterfaces : DataflowData.Value.SimulationInterfaces)
+		{
+			for(IDataflowSimulationInterface* SimulationInterface : SimulationInterfaces.Value)
+			{
+				if(SimulationInterface)
+				{
+					if(!SimulationInterface->GetSimulationProxy())
+					{
+						SimulationInterface->BuildSimulationProxy();
+					}
+					if(SimulationInterface->GetSimulationProxy())
+					{
+						SimulationInterface->GetSimulationProxy()->SetSimulationGroups(SimulationInterface->GetSimulationAsset().SimulationGroups);
+						DataflowData.Value.SimulationContext->AddSimulationProxy(SimulationInterfaces.Key, SimulationInterface->GetSimulationProxy());
+					}
+					
+					SimulationInterface->WriteToSimulation(DeltaTime);
+				}
+			}
+		}
+		DataflowData.Value.SimulationContext->RegisterProxyGroups();
+	}
+}
+
+void UDataflowSimulationManager::ReadSimulationData(const float DeltaTime)
+{
+	for(const TPair<TObjectPtr<UDataflow>, Dataflow::FDataflowSimulationData>& DataflowData : SimulationData)
+	{
+		for(const TPair<FString,TSet<IDataflowSimulationInterface*>>& SimulationInterfaces : DataflowData.Value.SimulationInterfaces)
+		{
+			for(IDataflowSimulationInterface* SimulationInterface : SimulationInterfaces.Value)
+			{
+				if(SimulationInterface)
+				{
+					SimulationInterface->ReadFromSimulation(DeltaTime);
+				}
+			}
+		}
+		DataflowData.Value.SimulationContext->ResetSimulationProxies();
+	}
+	if(bStepSimulationScene)
+	{
+		bIsSimulationEnabled = false;
+		bStepSimulationScene = false;
+	}
+}
+
+void UDataflowSimulationManager::AdvanceSimulationData(const float DeltaTime, const float SimulationTime)
+{
+	for(TPair<TObjectPtr<UDataflow>, Dataflow::FDataflowSimulationData>& DataflowData : SimulationData)
+	{
+		Dataflow::EvaluateSimulationGraph(DataflowData.Key, DataflowData.Value.SimulationContext, DeltaTime, SimulationTime);
+	}
+}
+
+void UDataflowSimulationManager::StartSimulationTasks(const float DeltaTime, const float SimulationTime)
+{
+	check(IsInGameThread());
+	check(SimulationTasks.IsEmpty());
+
+	for(TPair<TObjectPtr<UDataflow>, Dataflow::FDataflowSimulationData>& DataflowData : SimulationData)
+	{
+		// Add a simulation task linked to that solver
+		SimulationTasks.Add(TGraphTask<Dataflow::FDataflowSimulationTask>::CreateTask(nullptr, ENamedThreads::GameThread).ConstructAndDispatchWhenReady(
+				DataflowData.Key, DataflowData.Value.SimulationContext, DeltaTime, SimulationTime));
+	}
+}
+
+void UDataflowSimulationManager::CompleteSimulationTasks()
+{
+	check(IsInGameThread());
+
+	for(FGraphEventRef& SimulationTask : SimulationTasks)
+	{
+		if (IsValidRef(SimulationTask))
+		{
+			// There's a simulation in flight
+			FTaskGraphInterface::Get().WaitUntilTaskCompletes(SimulationTask, ENamedThreads::GameThread);
+
+			// No longer need this task, it has completed
+			SimulationTask.SafeRelease();
+		}
+	}
+	SimulationTasks.Reset();
+}
+
+TSharedPtr<Dataflow::FDataflowSimulationContext> UDataflowSimulationManager::GetSimulationContext(const TObjectPtr<UDataflow>& DataflowAsset) const
+{
+	if(DataflowAsset)
+	{
+		if(const Dataflow::FDataflowSimulationData* DataflowData = SimulationData.Find(DataflowAsset))
+		{
+			return DataflowData->SimulationContext;
+		}
+	}
+	return nullptr;
+}
+
+bool UDataflowSimulationManager::HasSimulationInterface(const IDataflowSimulationInterface* SimulationInterface) const
+{
+	if(SimulationInterface)
+	{
+		if(const TObjectPtr<UDataflow> DataflowAsset = SimulationInterface->GetSimulationAsset().DataflowAsset)
+		{
+			if(const Dataflow::FDataflowSimulationData* DataflowData = SimulationData.Find(DataflowAsset))
+			{
+				if(const TSet<IDataflowSimulationInterface*>* SimulationInterfaces =
+					DataflowData->SimulationInterfaces.Find(SimulationInterface->GetSimulationType()))
+				{
+					return SimulationInterfaces->Find(SimulationInterface) != nullptr;
+				}
+			}
+		}
+	}
+	return false;
+}
+
+void UDataflowSimulationManager::AddSimulationInterface(IDataflowSimulationInterface* SimulationInterface)
+{
+	if(SimulationInterface)
+	{
+		if(TObjectPtr<UDataflow> DataflowAsset = SimulationInterface->GetSimulationAsset().DataflowAsset)
+		{
+			Dataflow::FDataflowSimulationData& DataflowData = SimulationData.FindOrAdd(DataflowAsset);
+			if(!DataflowData.SimulationContext.IsValid())
+			{
+				DataflowData.SimulationContext = MakeShared<Dataflow::FDataflowSimulationContext>(
+					DataflowAsset, DataflowAsset, Dataflow::FTimestamp::Invalid);
+			}
+			DataflowData.SimulationInterfaces.FindOrAdd(SimulationInterface->GetSimulationType()).Add(SimulationInterface);
+		}
+	}
+}
+
+void UDataflowSimulationManager::RemoveSimulationInterface(const IDataflowSimulationInterface* SimulationInterface)
+{
+	if(SimulationInterface)
+	{
+		for(TPair<TObjectPtr<UDataflow>, Dataflow::FDataflowSimulationData>& DataflowData : SimulationData)
+		{
+			if(TSet<IDataflowSimulationInterface*>* SimulationInterfaces =
+					DataflowData.Value.SimulationInterfaces.Find(SimulationInterface->GetSimulationType()))
+			{
+				SimulationInterfaces->Remove(SimulationInterface);
+			}
+		}
+	}
+}
+
+ETickableTickType UDataflowSimulationManager::GetTickableTickType() const
+{
+	return HasAnyFlags(RF_ClassDefaultObject) || !GetWorld() || GetWorld()->IsNetMode(NM_DedicatedServer) ? ETickableTickType::Never : ETickableTickType::Always;
+}
+
+TStatId UDataflowSimulationManager::GetStatId() const
+{
+	RETURN_QUICK_DECLARE_CYCLE_STAT(UDataflowSimulationManager, STATGROUP_Tickables);
+}
+
+bool UDataflowSimulationManager::DoesSupportWorldType(EWorldType::Type WorldType) const
+{
+	return Super::DoesSupportWorldType(WorldType) || WorldType == EWorldType::EditorPreview || WorldType == EWorldType::GamePreview || WorldType == EWorldType::GameRPC;
+}
+
+void UDataflowSimulationManager::Initialize(FSubsystemCollectionBase& Collection)
+{
+	Super::Initialize(Collection);
+}
+
+void UDataflowSimulationManager::Deinitialize()
+{
+	Super::Deinitialize();
+
+	CompleteSimulationTasks();
+}
diff --git a/Engine/Plugins/Experimental/Dataflow/Source/DataflowSimulation/Private/Dataflow/DataflowSimulationNodes.cpp b/Engine/Plugins/Experimental/Dataflow/Source/DataflowSimulation/Private/Dataflow/DataflowSimulationNodes.cpp
new file mode 100644
--- /dev/null
+++ b/Engine/Plugins/Experimental/Dataflow/Source/DataflowSimulation/Private/Dataflow/DataflowSimulationNodes.cpp
@@ -0,0 +1,94 @@
+﻿// Copyright Epic Games, Inc. All Rights Reserved.
+
+#include "Dataflow/DataflowSimulationNodes.h"
+#include "Dataflow/DataflowSimulationManager.h"
+#include "Dataflow/DataflowNodeFactory.h"
+#include "Dataflow/DataflowNodeColorsRegistry.h"
+#include "Interfaces/DataflowPhysicsSolver.h"
+
+#include UE_INLINE_GENERATED_CPP_BY_NAME(DataflowSimulationNodes)
+
+namespace Dataflow
+{
+	void RegisterDataflowSimulationNodes()
+	{
+		DATAFLOW_NODE_REGISTER_CREATION_FACTORY(FGetSimulationTimeDataflowNode);
+		DATAFLOW_NODE_REGISTER_CREATION_FACTORY(FGetPhysicsSolversDataflowNode);
+		DATAFLOW_NODE_REGISTER_CREATION_FACTORY(FAdvancePhysicsSolversDataflowNode);
+		DATAFLOW_NODE_REGISTER_CREATION_FACTORY(FSimulationProxiesTerminalDataflowNode);
+		DATAFLOW_NODE_REGISTER_CREATION_FACTORY(FFilterSimulationProxiesDataflowNode);
+		
+		static constexpr FLinearColor CDefaultNodeBodyTintColor = FLinearColor(0.f, 0.f, 0.0f, 0.5f);
+		
+		DATAFLOW_NODE_REGISTER_CREATION_FACTORY_NODE_COLORS_BY_CATEGORY("Terminal", FLinearColor(1.0f, 0.0f, 0.0f), CDefaultNodeBodyTintColor);
+		DATAFLOW_NODE_REGISTER_CREATION_FACTORY_NODE_COLORS_BY_CATEGORY("Setup", FLinearColor(1.0f, 1.0f, 0.0f), CDefaultNodeBodyTintColor);
+		DATAFLOW_NODE_REGISTER_CREATION_FACTORY_NODE_COLORS_BY_CATEGORY("Physics", FLinearColor(0.0f, 1.0f, 0.0f), CDefaultNodeBodyTintColor);
+	}
+}
+
+void FGetSimulationTimeDataflowNode::EvaluateSimulation(Dataflow::FDataflowSimulationContext& SimulationContext, const FDataflowOutput* Output) const
+{
+	SetValue(SimulationContext, FDataflowSimulationTime(SimulationContext.GetDeltaTime(), SimulationContext.GetSimulationTime()), &SimulationTime);
+}
+
+void FAdvancePhysicsSolversDataflowNode::EvaluateSimulation(Dataflow::FDataflowSimulationContext& SimulationContext, const FDataflowOutput* Output) const
+{
+	const TArray<FDataflowSimulationProperty> SolverProperties = GetValue(SimulationContext, &PhysicsSolvers);
+	const float SimulationDeltaTime = GetValue(SimulationContext, &SimulationTime).DeltaTime;
+	
+	if(!SolverProperties.IsEmpty())
+	{
+		for(const FDataflowSimulationProperty& SolverProperty : SolverProperties)
+		{
+			if(SolverProperty.SimulationProxy)
+			{
+				if(FDataflowPhysicsSolverProxy* SolverProxy = SolverProperty.SimulationProxy->AsType<FDataflowPhysicsSolverProxy>())
+				{
+					SolverProxy->AdvanceSolverDatas(SimulationDeltaTime);
+				}
+			}
+		}
+	}
+	SetValue(SimulationContext, SolverProperties, &PhysicsSolvers);
+}
+
+void FGetPhysicsSolversDataflowNode::EvaluateSimulation(Dataflow::FDataflowSimulationContext& SimulationContext, const FDataflowOutput* Output) const
+{
+	TArray<FDataflowSimulationProxy*> SimulationProxies;
+	SimulationContext.GetSimulationProxies(FDataflowPhysicsSolverProxy::StaticStruct()->GetName(), SimulationGroups, SimulationProxies);
+
+	TArray<FDataflowSimulationProperty> SolverProperties;
+	for(FDataflowSimulationProxy* SimulationProxy : SimulationProxies)
+	{
+		SolverProperties.Add({SimulationProxy});
+	}
+	
+	SetValue(SimulationContext, SolverProperties, &PhysicsSolvers);
+}
+
+void FFilterSimulationProxiesDataflowNode::EvaluateSimulation(Dataflow::FDataflowSimulationContext& SimulationContext, const FDataflowOutput* Output) const
+{
+	const TArray<FDataflowSimulationProperty> SimulationProperties = GetValue(SimulationContext, &SimulationProxies);
+
+	TArray<FDataflowSimulationProperty> FilteredProperties;
+	if(!SimulationProperties.IsEmpty())
+	{
+		TBitArray<> GroupBits;
+		SimulationContext.BuildGroupBits(SimulationGroups, GroupBits);
+
+		for(const FDataflowSimulationProperty& SimulationProperty : SimulationProperties)
+		{
+			if(SimulationProperty.SimulationProxy && SimulationProperty.SimulationProxy->HasGroupBit(GroupBits))
+			{
+				FilteredProperties.Add({SimulationProperty.SimulationProxy});
+			}
+		}
+	}
+	
+	SetValue(SimulationContext, FilteredProperties, &FilteredProxies);
+}
+
+void FSimulationProxiesTerminalDataflowNode::EvaluateSimulation(Dataflow::FDataflowSimulationContext& SimulationContext, const FDataflowOutput* Output) const
+{
+	const TArray<FDataflowSimulationProperty> SolverProperties = GetValue(SimulationContext, &SimulationProxies);
+}
diff --git a/Engine/Plugins/Experimental/Dataflow/Source/DataflowSimulation/Private/Dataflow/DataflowSimulationPlugin.cpp b/Engine/Plugins/Experimental/Dataflow/Source/DataflowSimulation/Private/Dataflow/DataflowSimulationPlugin.cpp
new file mode 100644
--- /dev/null
+++ b/Engine/Plugins/Experimental/Dataflow/Source/DataflowSimulation/Private/Dataflow/DataflowSimulationPlugin.cpp
@@ -0,0 +1,29 @@
+﻿// Copyright Epic Games, Inc. All Rights Reserved.
+
+#include "Dataflow/DataflowSimulationPlugin.h"
+#include "Dataflow/DataflowSimulationNodes.h"
+#include "Dataflow/DataflowGraph.h"
+#include "Dataflow/DataflowSimulationManager.h"
+#include "CoreMinimal.h"
+
+#define LOCTEXT_NAMESPACE "DataflowSimulation"
+
+void IDataflowSimulationPlugin::StartupModule()
+{
+	Dataflow::RegisterDataflowSimulationNodes();
+
+	Dataflow::RegisterNodeFilter(FDataflowSimulationNode::StaticType());
+	Dataflow::RegisterNodeFilter(FDataflowInvalidNode::StaticType());
+	Dataflow::RegisterNodeFilter(FDataflowExecutionNode::StaticType());
+
+	UDataflowSimulationManager::OnStartup();
+}
+
+void IDataflowSimulationPlugin::ShutdownModule()
+{
+	UDataflowSimulationManager::OnShutdown();
+}
+
+IMPLEMENT_MODULE(IDataflowSimulationPlugin, DataflowSimulation)
+
+#undef LOCTEXT_NAMESPACE
diff --git a/Engine/Plugins/Experimental/Dataflow/Source/DataflowSimulation/Private/Dataflow/DataflowSimulationProxy.cpp b/Engine/Plugins/Experimental/Dataflow/Source/DataflowSimulation/Private/Dataflow/DataflowSimulationProxy.cpp
new file mode 100644
--- /dev/null
+++ b/Engine/Plugins/Experimental/Dataflow/Source/DataflowSimulation/Private/Dataflow/DataflowSimulationProxy.cpp
@@ -0,0 +1,19 @@
+﻿// Copyright Epic Games, Inc. All Rights Reserved.
+
+#include "Dataflow/DataflowSimulationProxy.h"
+
+void FDataflowSimulationProxy::SetSimulationGroups(const TSet<FString>& InSimulationGroups)
+{
+	SimulationGroups = InSimulationGroups;
+}
+
+bool FDataflowSimulationProxy::HasSimulationGroup(const FString& SimulationGroup) const
+{
+	return (SimulationGroups.Find(SimulationGroup) != nullptr);
+}
+
+bool FDataflowSimulationProxy::HasGroupBit(const TBitArray<>& SimulationBits) const
+{
+	const TBitArray<> CommonGroups = TBitArray<>::BitwiseAND(GroupBits, SimulationBits, EBitwiseOperatorFlags::MinSize);
+	return (CommonGroups.CountSetBits() != 0);
+}
diff --git a/Engine/Plugins/Experimental/Dataflow/Source/DataflowSimulation/Private/Dataflow/DataflowSimulationUtils.cpp b/Engine/Plugins/Experimental/Dataflow/Source/DataflowSimulation/Private/Dataflow/DataflowSimulationUtils.cpp
new file mode 100644
--- /dev/null
+++ b/Engine/Plugins/Experimental/Dataflow/Source/DataflowSimulation/Private/Dataflow/DataflowSimulationUtils.cpp
@@ -0,0 +1,234 @@
+﻿// Copyright Epic Games, Inc. All Rights Reserved.
+
+#include "Dataflow/DataflowSimulationUtils.h"
+
+#include "Animation/AnimSingleNodeInstance.h"
+#include "Chaos/CacheManagerActor.h"
+#include "Chaos/Adapters/CacheAdapter.h"
+#include "Dataflow/DataflowContent.h"
+#include "Dataflow/DataflowObject.h"
+#include "Dataflow/DataflowSimulationNodes.h"
+#include "Dataflow/DataflowObjectInterface.h"
+#include "Features/IModularFeatures.h"
+#include "Components/SkeletalMeshComponent.h"
+#include "Engine/World.h"
+
+
+#define LOCTEXT_NAMESPACE "DataflowSimulationGenerator"
+
+namespace Dataflow
+{
+	bool ShouldResetWorld(const TObjectPtr<UDataflow>& SimulationGraph, const TObjectPtr<UWorld>& SimulationWorld, Dataflow::FTimestamp& LastTimeStamp)
+	{
+		if(const TSharedPtr<Dataflow::FGraph> DataflowGraph = SimulationGraph->GetDataflow())
+		{
+			Dataflow::FTimestamp MaxTimeStamp = Dataflow::FTimestamp::Invalid;
+			for(const TSharedPtr<FDataflowNode>& TerminalNode : DataflowGraph->GetFilteredNodes(FDataflowTerminalNode::StaticType()))
+			{
+				MaxTimeStamp.Value = FMath::Max(MaxTimeStamp.Value, TerminalNode->GetTimestamp().Value);
+			}
+			if(MaxTimeStamp.Value > LastTimeStamp.Value)
+			{
+				LastTimeStamp = MaxTimeStamp.Value;
+				return true;
+			}
+		}
+		return false;
+	}
+	
+	void EvaluateSimulationGraph(const TObjectPtr<UDataflow>& SimulationGraph, const TSharedPtr<Dataflow::FDataflowSimulationContext>& SimulationContext,
+		const float DeltaTime, const float SimulationTime)
+	{
+		if(SimulationContext.IsValid())
+		{
+			SimulationContext->SetTimingInfos(DeltaTime, SimulationTime);
+			
+			if(SimulationGraph)
+			{
+				if(const TSharedPtr<Dataflow::FGraph> DataflowGraph = SimulationGraph->GetDataflow())
+				{
+					// Invalidation of all the simulation nodes that are always dirty
+					for(const TSharedPtr<FDataflowNode>& InvalidNode : DataflowGraph->GetFilteredNodes(FDataflowInvalidNode::StaticType()))
+					{
+						InvalidNode->Invalidate();
+					}
+					
+					// Pull the graph evaluation from the solver nodes
+					for(const TSharedPtr<FDataflowNode>& ExecutionNode : DataflowGraph->GetFilteredNodes(FDataflowExecutionNode::StaticType()))
+					{
+						SimulationContext->Evaluate(ExecutionNode.Get(), nullptr);
+					}
+				}
+			}
+		}
+	}
+	
+	TObjectPtr<AActor> SpawnSimulatedActor(const TSubclassOf<AActor>& ActorClass,
+		const TObjectPtr<AChaosCacheManager>& CacheManager, const TObjectPtr<UChaosCacheCollection>& CacheCollection,
+		const bool bIsRecording, const TObjectPtr<UDataflowBaseContent>& DataflowContent)
+	{
+		if(CacheManager)
+		{
+			FActorSpawnParameters SpawnParameters;
+			SpawnParameters.Name = TEXT("CacheActor");
+			SpawnParameters.NameMode = FActorSpawnParameters::ESpawnActorNameMode::Requested;
+			SpawnParameters.Owner = CacheManager.Get(); 
+			SpawnParameters.bDeferConstruction = true;
+			
+			TObjectPtr<AActor> PreviewActor = CacheManager->GetWorld()->SpawnActor<AActor>(ActorClass, SpawnParameters);
+			if(PreviewActor)
+			{
+				// Link the editor content properties to the BP actor one 
+				DataflowContent->SetActorProperties(PreviewActor);
+
+				// Finish spawning
+				PreviewActor->FinishSpawning(FTransform::Identity, true);
+			}
+
+			CacheManager->CacheCollection = CacheCollection;
+			CacheManager->StartMode = EStartMode::Timed;
+			CacheManager->CacheMode = bIsRecording ? ECacheMode::Record : ECacheMode::None;
+			
+			// Get the implementation of our adapters for identifying compatible components
+			IModularFeatures&                      ModularFeatures = IModularFeatures::Get();
+			TArray<Chaos::FComponentCacheAdapter*> Adapters = ModularFeatures.GetModularFeatureImplementations<Chaos::FComponentCacheAdapter>(Chaos::FComponentCacheAdapter::FeatureName);
+
+			if(PreviewActor)
+			{
+				TInlineComponentArray<UPrimitiveComponent*> PrimComponents;
+				PreviewActor->GetComponents(PrimComponents);
+	
+				for(UPrimitiveComponent* PrimComponent : PrimComponents)
+				{
+					if(Chaos::FAdapterUtil::GetBestAdapterForClass(PrimComponent->GetClass(), false))
+					{
+						const FName ChannelName(PrimComponent->GetName());
+						CacheManager->FindOrAddObservedComponent(PrimComponent, ChannelName, true);
+					}
+				}
+			}
+			return PreviewActor;
+		}
+		return nullptr;
+	}
+	
+	void SetupSkeletonAnimation(const TObjectPtr<AActor>& PreviewActor)
+	{
+		if(PreviewActor)
+		{
+			TInlineComponentArray<UPrimitiveComponent*> PrimComponents;
+			PreviewActor->GetComponents(PrimComponents);
+
+			for(UPrimitiveComponent* PrimComponent : PrimComponents)
+			{
+				if(USkeletalMeshComponent* SkeletalMeshComponent = Cast<USkeletalMeshComponent>(PrimComponent))
+				{
+					SkeletalMeshComponent->SetAnimationMode(EAnimationMode::AnimationSingleNode);
+					SkeletalMeshComponent->InitAnim(true);
+					
+					if(UAnimSingleNodeInstance* AnimNodeInstance = SkeletalMeshComponent->GetSingleNodeInstance())
+					{
+						// Setup the animation instance
+						AnimNodeInstance->SetAnimationAsset(SkeletalMeshComponent->AnimationData.AnimToPlay);
+						AnimNodeInstance->InitializeAnimation();
+
+						// Update the anim data
+						SkeletalMeshComponent->AnimationData.PopulateFrom(AnimNodeInstance);
+#if WITH_EDITOR
+						SkeletalMeshComponent->ValidateAnimation();
+#endif
+
+						// Stop the animation 
+						AnimNodeInstance->SetLooping(true);
+						AnimNodeInstance->SetPlaying(false);
+						
+					}
+				}
+			}
+		}
+	}
+
+	void UpdateSkeletonAnimation(const TObjectPtr<AActor>& PreviewActor, const float SimulationTime)
+	{
+		if(PreviewActor)
+		{
+			TInlineComponentArray<UPrimitiveComponent*> PrimComponents;
+			PreviewActor->GetComponents(PrimComponents);
+
+			// Update all the animation time 
+			for(UPrimitiveComponent* PrimComponent : PrimComponents)
+			{
+				if(USkeletalMeshComponent* SkeletalMeshComponent = Cast<USkeletalMeshComponent>(PrimComponent))
+				{
+					SkeletalMeshComponent->SetPosition(SimulationTime);
+					SkeletalMeshComponent->TickAnimation(0.f, false /*bNeedsValidRootMotion*/);
+					SkeletalMeshComponent->RefreshBoneTransforms(nullptr /*TickFunction*/);
+				}
+			}
+		}
+	}
+
+	void StartSkeletonAnimation(const TObjectPtr<AActor>& PreviewActor)
+	{
+		if(PreviewActor)
+		{
+			TInlineComponentArray<UPrimitiveComponent*> PrimComponents;
+			PreviewActor->GetComponents(PrimComponents);
+
+			for(UPrimitiveComponent* PrimComponent : PrimComponents)
+			{
+				if(const USkeletalMeshComponent* SkeletalMeshComponent = Cast<USkeletalMeshComponent>(PrimComponent))
+				{
+					if(UAnimSingleNodeInstance* AnimNodeInstance = SkeletalMeshComponent->GetSingleNodeInstance())
+					{
+						AnimNodeInstance->SetPlaying(true);
+					}
+				}
+			}
+		}
+	}
+	
+	void PauseSkeletonAnimation(const TObjectPtr<AActor>& PreviewActor)
+	{
+		if(PreviewActor)
+		{
+			TInlineComponentArray<UPrimitiveComponent*> PrimComponents;
+			PreviewActor->GetComponents(PrimComponents);
+
+			for(UPrimitiveComponent* PrimComponent : PrimComponents)
+			{
+				if(const USkeletalMeshComponent* SkeletalMeshComponent = Cast<USkeletalMeshComponent>(PrimComponent))
+				{
+					if(UAnimSingleNodeInstance* AnimNodeInstance = SkeletalMeshComponent->GetSingleNodeInstance())
+					{
+						AnimNodeInstance->SetPlaying(false);
+					}
+				}
+			}
+		}
+	}
+	
+	void StepSkeletonAnimation(const TObjectPtr<AActor>& PreviewActor)
+	{
+		if(PreviewActor)
+		{
+			TInlineComponentArray<UPrimitiveComponent*> PrimComponents;
+			PreviewActor->GetComponents(PrimComponents);
+
+			for(UPrimitiveComponent* PrimComponent : PrimComponents)
+			{
+				if(const USkeletalMeshComponent* SkeletalMeshComponent = Cast<USkeletalMeshComponent>(PrimComponent))
+				{
+					if(UAnimSingleNodeInstance* AnimNodeInstance = SkeletalMeshComponent->GetSingleNodeInstance())
+					{
+						AnimNodeInstance->SetPlaying(false);
+						AnimNodeInstance->StepForward();
+					}
+				}
+			}
+		}
+	}
+	
+}
+
+#undef LOCTEXT_NAMESPACE
diff --git a/Engine/Plugins/Experimental/Dataflow/Source/DataflowSimulation/Private/Interfaces/DataflowCollisionObject.cpp b/Engine/Plugins/Experimental/Dataflow/Source/DataflowSimulation/Private/Interfaces/DataflowCollisionObject.cpp
new file mode 100644
--- /dev/null
+++ b/Engine/Plugins/Experimental/Dataflow/Source/DataflowSimulation/Private/Interfaces/DataflowCollisionObject.cpp
@@ -0,0 +1,9 @@
+﻿// Copyright Epic Games, Inc. All Rights Reserved.
+
+#include "Interfaces/DataflowCollisionObject.h"
+
+FDataflowCollisionObjectProxy::FDataflowCollisionObjectProxy()
+{}
+
+IDataflowCollisionObjectInterface::IDataflowCollisionObjectInterface()
+{}
diff --git a/Engine/Plugins/Experimental/Dataflow/Source/DataflowSimulation/Private/Interfaces/DataflowConstraintObject.cpp b/Engine/Plugins/Experimental/Dataflow/Source/DataflowSimulation/Private/Interfaces/DataflowConstraintObject.cpp
new file mode 100644
--- /dev/null
+++ b/Engine/Plugins/Experimental/Dataflow/Source/DataflowSimulation/Private/Interfaces/DataflowConstraintObject.cpp
@@ -0,0 +1,9 @@
+﻿// Copyright Epic Games, Inc. All Rights Reserved.
+
+#include "Interfaces/DataflowConstraintObject.h"
+
+FDataflowConstraintObjectProxy::FDataflowConstraintObjectProxy()
+{}
+
+IDataflowConstraintObjectInterface::IDataflowConstraintObjectInterface()
+{}
diff --git a/Engine/Plugins/Experimental/Dataflow/Source/DataflowSimulation/Private/Interfaces/DataflowPhysicsObject.cpp b/Engine/Plugins/Experimental/Dataflow/Source/DataflowSimulation/Private/Interfaces/DataflowPhysicsObject.cpp
new file mode 100644
--- /dev/null
+++ b/Engine/Plugins/Experimental/Dataflow/Source/DataflowSimulation/Private/Interfaces/DataflowPhysicsObject.cpp
@@ -0,0 +1,9 @@
+﻿// Copyright Epic Games, Inc. All Rights Reserved.
+
+#include "Interfaces/DataflowPhysicsObject.h"
+
+FDataflowPhysicsObjectProxy::FDataflowPhysicsObjectProxy()
+{}
+
+IDataflowPhysicsObjectInterface::IDataflowPhysicsObjectInterface()
+{}
diff --git a/Engine/Plugins/Experimental/Dataflow/Source/DataflowSimulation/Private/Interfaces/DataflowPhysicsSolver.cpp b/Engine/Plugins/Experimental/Dataflow/Source/DataflowSimulation/Private/Interfaces/DataflowPhysicsSolver.cpp
new file mode 100644
--- /dev/null
+++ b/Engine/Plugins/Experimental/Dataflow/Source/DataflowSimulation/Private/Interfaces/DataflowPhysicsSolver.cpp
@@ -0,0 +1,9 @@
+﻿// Copyright Epic Games, Inc. All Rights Reserved.
+
+#include "Interfaces/DataflowPhysicsSolver.h"
+
+FDataflowPhysicsSolverProxy::FDataflowPhysicsSolverProxy()
+{}
+
+IDataflowPhysicsSolverInterface::IDataflowPhysicsSolverInterface()
+{}
diff --git a/Engine/Plugins/Experimental/Dataflow/Source/DataflowSimulation/Public/Dataflow/DataflowSimulationContext.h b/Engine/Plugins/Experimental/Dataflow/Source/DataflowSimulation/Public/Dataflow/DataflowSimulationContext.h
new file mode 100644
--- /dev/null
+++ b/Engine/Plugins/Experimental/Dataflow/Source/DataflowSimulation/Public/Dataflow/DataflowSimulationContext.h
@@ -0,0 +1,74 @@
+﻿// Copyright Epic Games, Inc. All Rights Reserved.
+
+#pragma once
+
+#include "CoreMinimal.h"
+#include "UObject/ObjectMacros.h"
+#include "Dataflow/DataflowObjectInterface.h"
+
+struct FDataflowSimulationProxy;
+
+namespace Dataflow
+{
+	/** Simulation context that will be used by all the simulation nodes*/
+	template<typename Base = FContextSingle>
+	class DATAFLOWSIMULATION_API TSimulationContext : public TEngineContext<Base>
+	{
+	public:
+		DATAFLOW_CONTEXT_INTERNAL(TEngineContext<Base>, TSimulationContext);
+
+		TSimulationContext(const TObjectPtr<UObject>& InOwner,const TObjectPtr<UDataflow>& InGraph, FTimestamp InTimestamp)
+				: Super(InOwner, InGraph, InTimestamp)
+		{}
+		
+		virtual ~TSimulationContext() override {};
+		
+		/** Set the timing infos */
+		void SetTimingInfos(const float DeltaSeconds, const float TimeSeconds) {DeltaTime = DeltaSeconds; SimulationTime = TimeSeconds;}
+
+		/** Get the delta time in seconds */
+		float GetDeltaTime() const {return DeltaTime;}
+		
+		/** Get the simulation time in seconds */
+		float GetSimulationTime() const {return SimulationTime;}
+
+		/** Filter the physics solvers matching the groups */
+		void GetSimulationProxies(const FString& ProxyType, const TArray<FString>& SimulationGroups, TArray<FDataflowSimulationProxy*>& FilteredProxies) const;
+
+		/** Add simulation proxy to the context */
+		void AddSimulationProxy(const FString& ProxyType, FDataflowSimulationProxy* SimulationProxy);
+		
+		/** Remove simulation proxy from the context */
+    	void RemoveSimulationProxy(const FString& ProxyType, const FDataflowSimulationProxy* SimulationProxy);
+
+		/** Reset all the simulation proxies */
+		void ResetSimulationProxies();
+		
+		/** Return the number of physics solvers */
+		int32 NumSimulationProxies(const FString& ProxyType) const;
+
+		/** Register all the proxy groups used in the proxy */
+		void RegisterProxyGroups();
+
+		/** Build group bits based on the string array */
+		void BuildGroupBits(const TArray<FString>& SimulationGroups, TBitArray<>& GroupBits) const;
+		
+	private :
+
+		/** Simulation time */
+		float SimulationTime = 0.0f;
+
+		/** Delta time */
+		float DeltaTime = 0.0f;
+
+		/** List of all the simulation proxies within the context sorted by type */
+		TMap<FString, TSet<FDataflowSimulationProxy*>> SimulationProxies;
+
+		/** List of all the simulation groups indices */
+		TMap<FString, uint32> GroupIndices;
+	};
+
+	typedef TSimulationContext<FContextSingle> FDataflowSimulationContext;
+	typedef TSimulationContext<FContextThreaded> FDataflowSimulationContextThreaded;
+	
+}
diff --git a/Engine/Plugins/Experimental/Dataflow/Source/DataflowSimulation/Public/Dataflow/DataflowSimulationInterface.h b/Engine/Plugins/Experimental/Dataflow/Source/DataflowSimulation/Public/Dataflow/DataflowSimulationInterface.h
new file mode 100644
--- /dev/null
+++ b/Engine/Plugins/Experimental/Dataflow/Source/DataflowSimulation/Public/Dataflow/DataflowSimulationInterface.h
@@ -0,0 +1,84 @@
+﻿// Copyright Epic Games, Inc. All Rights Reserved.
+
+#pragma once
+
+#include "CoreMinimal.h"
+#include "Dataflow/DataflowObject.h"
+#include "UObject/Interface.h"
+
+#include "DataflowSimulationInterface.generated.h"
+
+struct FDataflowSimulationProxy;
+
+/**
+ * Dataflow simulation asset (should be in the interface children)
+ */
+USTRUCT(BlueprintType)
+struct DATAFLOWSIMULATION_API FDataflowSimulationAsset
+{
+	GENERATED_BODY()
+
+	/* Simulation dataflow asset used to advance in time on Pt */
+	UPROPERTY(EditAnywhere, BlueprintReadWrite, Category = "Physics")
+	TObjectPtr<UDataflow> DataflowAsset;
+
+	/* Simulation groups used to filter within the simulation nodes*/
+	UPROPERTY(EditAnywhere, BlueprintReadWrite, Category = "Physics")
+	TSet<FString> SimulationGroups;
+};
+
+UINTERFACE()
+class DATAFLOWSIMULATION_API UDataflowSimulationInterface : public UInterface
+{
+	GENERATED_BODY()
+};
+
+/**
+ * Dataflow simulation interface to send/receive datas (GT <-> PT)
+ */
+class DATAFLOWSIMULATION_API IDataflowSimulationInterface
+{
+	GENERATED_BODY()
+	
+public:
+	IDataflowSimulationInterface() {}
+
+	/** Get the dataflow simulation asset */
+	virtual FDataflowSimulationAsset& GetSimulationAsset() = 0;
+
+	/** Get the const dataflow simulation asset */
+	virtual const FDataflowSimulationAsset& GetSimulationAsset() const = 0;
+	
+	/** Build the simulation proxy */
+	virtual void BuildSimulationProxy() = 0;
+	
+	/** Reset the simulation proxy */
+	virtual void ResetSimulationProxy() = 0;
+
+	/** Get the const simulation proxy */
+	virtual const FDataflowSimulationProxy* GetSimulationProxy() const = 0;
+
+	/** Get the simulation proxy */
+	virtual FDataflowSimulationProxy* GetSimulationProxy() = 0;
+
+	/** Get the simulation name */
+	virtual FString GetSimulationName() const = 0;
+
+	/** Write data to be sent to the simulation proxy */
+	virtual void WriteToSimulation(const float DeltaTime) {};
+
+	/** Read data received from the simulation proxy */
+	virtual void ReadFromSimulation(const float DeltaTime) {};
+	
+	/** Get the simulation type */
+    virtual FString GetSimulationType() const {return TEXT("");};
+
+	/** Register simulation interface solver to manager */
+	void RegisterManagerInterface(const TObjectPtr<UWorld>& SimulationWorld);
+
+	/** Unregister simulation interface from the manager */
+	void UnregisterManagerInterface(const TObjectPtr<UWorld>& SimulationWorld) const;
+
+	/** Check if the interface has been registered to the manager */
+	bool IsInterfaceRegistered(const TObjectPtr<UWorld>& SimulationWorld) const;
+};
diff --git a/Engine/Plugins/Experimental/Dataflow/Source/DataflowSimulation/Public/Dataflow/DataflowSimulationManager.h b/Engine/Plugins/Experimental/Dataflow/Source/DataflowSimulation/Public/Dataflow/DataflowSimulationManager.h
new file mode 100644
--- /dev/null
+++ b/Engine/Plugins/Experimental/Dataflow/Source/DataflowSimulation/Public/Dataflow/DataflowSimulationManager.h
@@ -0,0 +1,115 @@
+﻿// Copyright Epic Games, Inc. All Rights Reserved.
+
+#pragma once
+
+#include "CoreMinimal.h"
+#include "UObject/ObjectMacros.h"
+#include "Subsystems/WorldSubsystem.h"
+#include "Templates/SharedPointer.h"
+#include "Dataflow/DataflowSimulationContext.h"
+
+#include "DataflowSimulationManager.generated.h"
+
+struct FDataflowSimulationProxy;
+class IDataflowSimulationInterface;
+
+namespace Dataflow
+{
+	/** Per dataflow graph simulation data type (data interfaces + simulation context) */
+	struct DATAFLOWSIMULATION_API  FDataflowSimulationData
+	{
+		/** List of all simulation interfaces used in this this dataflow graph */
+		TMap<FString, TSet<IDataflowSimulationInterface*>> SimulationInterfaces;
+
+		/** Simulation context used to evaluate the graph on PT */
+		TSharedPtr<Dataflow::FDataflowSimulationContext> SimulationContext;
+	};
+}
+
+UCLASS()
+class DATAFLOWSIMULATION_API UDataflowSimulationManager : public UTickableWorldSubsystem
+{
+	GENERATED_BODY()
+
+public:
+	UDataflowSimulationManager();
+	virtual ~UDataflowSimulationManager() override  = default;
+
+	/** Static function to add world delegates */
+	static void OnStartup();
+	
+	/** Static function to remove world delegates */
+	static void OnShutdown();
+
+	// Begin FTickableGameObject overrides
+	virtual void Tick(float DeltaTime) override;
+	virtual bool IsTickableInEditor() const override { return true; }
+	virtual ETickableTickType GetTickableTickType() const override;
+	virtual bool DoesSupportWorldType(const EWorldType::Type WorldType) const override;
+	virtual TStatId GetStatId() const override;
+	// End FTickableGameObject overrides
+	
+	// Begin USubsystem overrides
+	virtual void Initialize(FSubsystemCollectionBase& Collection) override;
+	virtual void Deinitialize() override;
+	// End USubsystem overrides
+
+	/** Check if the manager has a simulation interface */
+	bool HasSimulationInterface(const IDataflowSimulationInterface* SimulationInterface) const;
+
+	/** Add a dataflow simulation interface to the manager*/
+	void AddSimulationInterface(IDataflowSimulationInterface* SimulationInterface);
+
+	/** Remove a dataflow simulation interface from the manager */
+	void RemoveSimulationInterface(const IDataflowSimulationInterface* SimulationInterface);
+
+	/** Advance in time the registered simulation data (PT) */
+	void AdvanceSimulationData(const float DeltaTime, const float SimulationTime);
+
+	/** Write all the data to the simulation proxies (to be send from GT->PT) */
+	void WriteSimulationData(const float DeltaTime);
+
+	/** Read all the data from the simulation proxies (received from PT->GT) */
+	void ReadSimulationData(const float DeltaTime);
+
+	/** Complete all the simulation tasks */
+	void CompleteSimulationTasks();
+
+	/** Start the simulation tasks given a delta time */
+	void StartSimulationTasks(const float DeltaTime, const float SimulationTime);
+
+	/** Set the simulation flag to enable/disable the simulation */
+	void SetSimulationEnabled(const bool bSimulationEnabled) {bIsSimulationEnabled = bSimulationEnabled;}
+
+	/** Set the simulation scene stepping bool */
+	void SetSimulationStepping(const bool bSimulationStepping) {bStepSimulationScene = bSimulationStepping;}
+
+	/** Get the simulation context for a given asset */
+	TSharedPtr<Dataflow::FDataflowSimulationContext> GetSimulationContext(const TObjectPtr<UDataflow>& DataflowAsset) const;
+	
+private :
+
+	/** static delegate for object property changed */
+	static FDelegateHandle OnObjectPropertyChangedHandle;
+
+	/** static delegate for end of frame update */
+	static FDelegateHandle OnWorldTickEndHandle;
+
+	/** static delegate for physics state creation */
+	static FDelegateHandle OnCreatePhysicsStateHandle;
+
+	/** static delegate for physics state destruction */
+	static FDelegateHandle OnDestroyPhysicsStateHandle;
+
+	/** Dataflow simulation data registered to the manager */
+	TMap<TObjectPtr<UDataflow>, Dataflow::FDataflowSimulationData> SimulationData;
+	
+	/** Simulation tasks in which the graph will be evaluated */
+	TArray<FGraphEventRef> SimulationTasks;
+
+	/** Boolean to control if the simulation should be disabled or not */
+	bool bIsSimulationEnabled = true;
+	
+	/** Boolean to check if we are stepping the simulation scene */
+	bool bStepSimulationScene = false;
+};
diff --git a/Engine/Plugins/Experimental/Dataflow/Source/DataflowSimulation/Public/Dataflow/DataflowSimulationNodes.h b/Engine/Plugins/Experimental/Dataflow/Source/DataflowSimulation/Public/Dataflow/DataflowSimulationNodes.h
new file mode 100644
--- /dev/null
+++ b/Engine/Plugins/Experimental/Dataflow/Source/DataflowSimulation/Public/Dataflow/DataflowSimulationNodes.h
@@ -0,0 +1,291 @@
+﻿// Copyright Epic Games, Inc. All Rights Reserved.
+
+#pragma once
+
+#include "Dataflow/DataflowObject.h"
+#include "UObject/ObjectPtr.h"
+#include "Dataflow/DataflowSimulationProxy.h"
+#include "Dataflow/DataflowSimulationContext.h"
+
+#include "DataflowSimulationNodes.generated.h"
+
+/**
+ * Dataflow simulation property 
+ */
+USTRUCT(BlueprintType)
+struct DATAFLOWSIMULATION_API FDataflowSimulationProperty
+{
+	GENERATED_BODY()
+
+	/* Simulation proxy used to pass information in between nodes */
+	FDataflowSimulationProxy* SimulationProxy;
+};
+
+/**
+* FDataflowSimulationNode
+*		Base class for simulation nodes within the Dataflow graph. 
+* 
+*		Simulation nodes are used to simulate data from the calling client. 
+*/
+USTRUCT()
+struct FDataflowSimulationNode : public FDataflowNode
+{
+	GENERATED_BODY()
+
+	FDataflowSimulationNode()
+		: Super() { }
+
+	FDataflowSimulationNode(const Dataflow::FNodeParameters& Param, FGuid InGuid = FGuid::NewGuid())
+		: Super(Param,InGuid) {
+	}
+
+	/** FDataflowNode interface */
+	virtual ~FDataflowSimulationNode() {}
+	
+	static FName StaticType() { return FName("FDataflowSimulationNode"); }
+
+	virtual bool IsA(FName InType) const override 
+	{ 
+		return InType.ToString().Equals(StaticType().ToString()) 
+			|| Super::IsA(InType); 
+	}
+
+	/** Evaluate simulation dispatch */
+	virtual void Evaluate(Dataflow::FContext& Context, const FDataflowOutput* Output) const override
+	{
+		if(Context.IsA(Dataflow::FDataflowSimulationContext::StaticType()))
+		{
+			Dataflow::FDataflowSimulationContext& SimulationContext = StaticCast<Dataflow::FDataflowSimulationContext&>(Context);
+			EvaluateSimulation(SimulationContext, Output);
+			
+		}
+	}
+
+	protected :
+
+	/** Evaluate simulation node given a simulation context */
+	virtual void EvaluateSimulation(Dataflow::FDataflowSimulationContext& SimulationContext, const FDataflowOutput* Output) const {};
+};
+
+/**
+* FDataflowInvalidNode
+*		Base class for invalid nodes within the Dataflow graph. 
+* 
+*		Invalid nodes will be always invalidated while simulating
+*/
+USTRUCT()
+struct FDataflowInvalidNode : public FDataflowSimulationNode
+{
+	GENERATED_BODY()
+
+	FDataflowInvalidNode()
+		: Super() { }
+
+	FDataflowInvalidNode(const Dataflow::FNodeParameters& Param, FGuid InGuid = FGuid::NewGuid())
+		: Super(Param,InGuid) {
+	}
+
+	/** FDataflowNode interface */
+	virtual ~FDataflowInvalidNode() {}
+	
+	static FName StaticType() { return FName("FDataflowInvalidNode"); }
+
+	virtual bool IsA(FName InType) const override 
+	{ 
+		return InType.ToString().Equals(StaticType().ToString()) 
+			|| Super::IsA(InType); 
+	}
+};
+
+/**
+* FDataflowExecutionNode
+*		Base class for the execute the dataflow simulation graph. 
+* 
+*		Execution nodes are used to pull the graph from the calling client. 
+*/
+USTRUCT()
+struct FDataflowExecutionNode : public FDataflowSimulationNode
+{
+	GENERATED_BODY()
+
+	FDataflowExecutionNode()
+		: Super() { }
+
+	FDataflowExecutionNode(const Dataflow::FNodeParameters& Param, FGuid InGuid = FGuid::NewGuid())
+		: Super(Param,InGuid) {
+	}
+
+	/** FDataflowNode interface */
+	virtual ~FDataflowExecutionNode() {}
+	
+	static FName StaticType() { return FName("FDataflowExecutionNode"); }
+
+	virtual bool IsA(FName InType) const override 
+	{ 
+		return InType.ToString().Equals(StaticType().ToString()) 
+			|| Super::IsA(InType); 
+	}
+};
+
+/** Get the dataflow simulation time */
+USTRUCT()
+struct DATAFLOWSIMULATION_API FDataflowSimulationTime 
+{
+	GENERATED_BODY()
+
+public:
+	FDataflowSimulationTime(): DeltaTime(0.0f), CurrentTime(0.0f)
+	{}
+	
+	FDataflowSimulationTime(const float InDeltaTime, const float InCurrentTime) : DeltaTime(InDeltaTime), CurrentTime(InCurrentTime)
+	{}
+	
+	/** Delta time in seconds property coming from the context */
+	UPROPERTY(Transient, SkipSerialization)
+	float DeltaTime = 0.0f;
+
+	/** Current time in seconds property coming from the context */
+	UPROPERTY(Transient, SkipSerialization)
+	float CurrentTime = 0.0f;
+};
+
+/** Get the context simulation time */
+USTRUCT(meta = (DataflowSimulation))
+struct DATAFLOWSIMULATION_API FGetSimulationTimeDataflowNode : public FDataflowInvalidNode
+{
+	GENERATED_BODY()
+	DATAFLOW_NODE_DEFINE_INTERNAL(FGetSimulationTimeDataflowNode, "GetSimulationTime", "Physics|Common", UDataflow::SimulationTag)
+
+public:
+	
+	FGetSimulationTimeDataflowNode(const Dataflow::FNodeParameters& InParam, FGuid InGuid = FGuid::NewGuid())
+	: FDataflowInvalidNode(InParam, InGuid)
+	{
+		RegisterOutputConnection(&SimulationTime);
+	}
+	
+	/** Simulation time property coming from the context */
+	UPROPERTY(Transient, SkipSerialization, Meta = (DataflowOutput))
+	FDataflowSimulationTime SimulationTime = FDataflowSimulationTime(0.0f, 0.0f);
+
+	/** Evaluate simulation node given a simulation context */
+	virtual void EvaluateSimulation(Dataflow::FDataflowSimulationContext& SimulationContext, const FDataflowOutput* Output) const override;
+};
+
+/** Main terminal node for flesh solver */
+USTRUCT(meta = (DataflowSimulation))
+struct FSimulationProxiesTerminalDataflowNode : public FDataflowExecutionNode
+{
+	GENERATED_BODY()
+	DATAFLOW_NODE_DEFINE_INTERNAL(FSimulationProxiesTerminalDataflowNode, "SimulationProxiesTerminal", "Terminal|Common", UDataflow::SimulationTag)
+
+public:
+	
+	FSimulationProxiesTerminalDataflowNode(const Dataflow::FNodeParameters& InParam, FGuid InGuid = FGuid::NewGuid())
+	: FDataflowExecutionNode(InParam, InGuid)
+	{
+		RegisterInputConnection(&SimulationProxies);
+	}
+
+	/** Physics solvers to evaluate  */
+	UPROPERTY(Transient, SkipSerialization, Meta = (DataflowInput))
+	TArray<FDataflowSimulationProperty> SimulationProxies;
+	
+	/** Evaluate simulation node given a simulation context */
+	virtual void EvaluateSimulation(Dataflow::FDataflowSimulationContext& SimulationContext, const FDataflowOutput* Output) const override;
+};
+
+/** Get physics solvers from context */
+USTRUCT(meta = (DataflowSimulation))
+struct FGetPhysicsSolversDataflowNode : public FDataflowInvalidNode
+{
+	GENERATED_BODY()
+	DATAFLOW_NODE_DEFINE_INTERNAL(FGetPhysicsSolversDataflowNode, "GetPhysicsSolvers", "Physics|Solver", UDataflow::SimulationTag)
+
+public:
+	
+	FGetPhysicsSolversDataflowNode(const Dataflow::FNodeParameters& InParam, FGuid InGuid = FGuid::NewGuid())
+	: FDataflowInvalidNode(InParam, InGuid)
+	{
+		RegisterOutputConnection(&PhysicsSolvers);
+	}
+
+	/** Physics solvers coming from the context and filtered with the groups */
+	UPROPERTY(Transient, SkipSerialization, Meta = (DataflowOutput))
+	TArray<FDataflowSimulationProperty> PhysicsSolvers;
+
+	/** Simulation groups to filter the output solvers properties */
+	UPROPERTY(EditAnywhere, Category="Simulation")
+	TArray<FString> SimulationGroups;
+
+	/** Evaluate simulation node given a simulation context */
+	virtual void EvaluateSimulation(Dataflow::FDataflowSimulationContext& SimulationContext, const FDataflowOutput* Output) const override;
+};
+
+/** Filter simulation proxies from context */
+USTRUCT(meta = (DataflowSimulation))
+struct FFilterSimulationProxiesDataflowNode : public FDataflowSimulationNode
+{
+	GENERATED_BODY()
+	DATAFLOW_NODE_DEFINE_INTERNAL(FFilterSimulationProxiesDataflowNode, "FilterSimulationProxies", "Physics|Common", UDataflow::SimulationTag)
+
+public:
+	
+	FFilterSimulationProxiesDataflowNode(const Dataflow::FNodeParameters& InParam, FGuid InGuid = FGuid::NewGuid())
+	: FDataflowSimulationNode(InParam, InGuid)
+	{
+		RegisterInputConnection(&SimulationProxies);
+		RegisterOutputConnection(&FilteredProxies);
+	}
+
+	/** Simulation proxies coming from the context and filtered with the groups */
+	UPROPERTY(Transient, SkipSerialization, Meta = (DataflowInput))
+	TArray<FDataflowSimulationProperty> SimulationProxies;
+
+	/** Simulation proxies coming from the context and filtered with the groups */
+	UPROPERTY(Transient, SkipSerialization, Meta = (DataflowOutput))
+	TArray<FDataflowSimulationProperty> FilteredProxies;
+
+	/** Simulation groups to filter the output solvers properties */
+	UPROPERTY(EditAnywhere, Category="Simulation")
+	TArray<FString> SimulationGroups;
+
+	/** Evaluate simulation node given a simulation context */
+	virtual void EvaluateSimulation(Dataflow::FDataflowSimulationContext& SimulationContext, const FDataflowOutput* Output) const override;
+};
+
+/** Advance the simulation physics solver in time */
+USTRUCT(meta = (DataflowSimulation))
+struct FAdvancePhysicsSolversDataflowNode : public FDataflowSimulationNode
+{
+	GENERATED_BODY()
+	DATAFLOW_NODE_DEFINE_INTERNAL(FAdvancePhysicsSolversDataflowNode, "AdvancePhysicsSolvers", "Physics|Solver", UDataflow::SimulationTag)
+
+public:
+	
+	FAdvancePhysicsSolversDataflowNode(const Dataflow::FNodeParameters& InParam, FGuid InGuid = FGuid::NewGuid())
+	: FDataflowSimulationNode(InParam, InGuid)
+	{
+		RegisterInputConnection(&SimulationTime);
+		RegisterInputConnection(&PhysicsSolvers);
+		RegisterOutputConnection(&PhysicsSolvers, &PhysicsSolvers);
+	}
+
+	/** Delta time to use to advance the solver*/
+	UPROPERTY(Transient, SkipSerialization, Meta = (DataflowInput))
+	FDataflowSimulationTime SimulationTime = FDataflowSimulationTime(0.0f,0.0f);
+
+	/** Physics solvers to advance in time */
+	UPROPERTY(Transient, SkipSerialization, Meta = (DataflowInput, DataflowOutput, DataflowPassthrough = "PhysicsSolvers"))
+	TArray<FDataflowSimulationProperty> PhysicsSolvers;
+
+	/** Evaluate simulation node given a simulation context */
+	virtual void EvaluateSimulation(Dataflow::FDataflowSimulationContext& SimulationContext, const FDataflowOutput* Output) const override;
+};
+
+namespace Dataflow
+{
+	void RegisterDataflowSimulationNodes();
+}
+
+
diff --git a/Engine/Plugins/Experimental/Dataflow/Source/DataflowSimulation/Public/Dataflow/DataflowSimulationPlugin.h b/Engine/Plugins/Experimental/Dataflow/Source/DataflowSimulation/Public/Dataflow/DataflowSimulationPlugin.h
new file mode 100644
--- /dev/null
+++ b/Engine/Plugins/Experimental/Dataflow/Source/DataflowSimulation/Public/Dataflow/DataflowSimulationPlugin.h
@@ -0,0 +1,17 @@
+﻿// Copyright Epic Games, Inc. All Rights Reserved.
+
+#pragma once
+
+#include "CoreMinimal.h"
+#include "Modules/ModuleInterface.h"
+
+/**
+ * The public interface to this module
+ */
+class IDataflowSimulationPlugin : public IModuleInterface
+{
+public:
+	virtual void StartupModule();
+	virtual void ShutdownModule();
+};
+
diff --git a/Engine/Plugins/Experimental/Dataflow/Source/DataflowSimulation/Public/Dataflow/DataflowSimulationProxy.h b/Engine/Plugins/Experimental/Dataflow/Source/DataflowSimulation/Public/Dataflow/DataflowSimulationProxy.h
new file mode 100644
--- /dev/null
+++ b/Engine/Plugins/Experimental/Dataflow/Source/DataflowSimulation/Public/Dataflow/DataflowSimulationProxy.h
@@ -0,0 +1,65 @@
+﻿// Copyright Epic Games, Inc. All Rights Reserved.
+
+#pragma once
+
+#include "CoreMinimal.h"
+#include "DataflowSimulationProxy.generated.h"
+
+/**
+ * Dataflow simulation proxy used for simulation (PT)
+ */
+USTRUCT()
+struct DATAFLOWSIMULATION_API FDataflowSimulationProxy
+{
+	GENERATED_BODY()
+	
+	FDataflowSimulationProxy() = default;
+	virtual ~FDataflowSimulationProxy() = default;
+	
+	/** Get the proxy script struct */
+	virtual const UScriptStruct* GetScriptStruct() const
+	{
+		return StaticStruct();
+	}
+	
+	/** Cast the proxy to child */
+	template<class T>
+	const T* AsType() const
+	{
+		if(GetScriptStruct()->IsChildOf(T::StaticStruct()))
+		{
+			return static_cast<T*>(this);
+		}
+		return nullptr;
+	}
+	
+	/** Cast the const proxy to child */
+	template<class T>
+	T* AsType()
+	{
+		if(GetScriptStruct()->IsChildOf(T::StaticStruct()))
+		{
+			return static_cast<T*>(this);
+		}
+		return nullptr;
+	}
+
+	/** Get the simulation groups from the proxy */
+	const TSet<FString>& GetSimulationGroups() const {return SimulationGroups;};
+
+	/** Set the simulation groups onto the proxy */
+	void SetSimulationGroups(const TSet<FString>& InSimulationGroups);
+
+	/** Check if the given group is within the proxy simulation groups*/
+	bool HasSimulationGroup(const FString& SimulationGroup) const;
+
+	/** Check if the proxy has at least one valid bit */
+    bool HasGroupBit(const TBitArray<>& SimulationBits) const;
+
+	/** Bit array matching the simulation groups for fast access */
+	TBitArray<> GroupBits;
+
+private:
+	/** List of simulation groups this proxy belongs to */
+	TSet<FString> SimulationGroups;
+};
\ No newline at end of file
diff --git a/Engine/Plugins/Experimental/Dataflow/Source/DataflowSimulation/Public/Dataflow/DataflowSimulationUtils.h b/Engine/Plugins/Experimental/Dataflow/Source/DataflowSimulation/Public/Dataflow/DataflowSimulationUtils.h
new file mode 100644
--- /dev/null
+++ b/Engine/Plugins/Experimental/Dataflow/Source/DataflowSimulation/Public/Dataflow/DataflowSimulationUtils.h
@@ -0,0 +1,44 @@
+﻿// Copyright Epic Games, Inc. All Rights Reserved.
+
+#pragma once
+
+#include "Dataflow/DataflowObjectInterface.h"
+#include "Dataflow/DataflowSimulationContext.h"
+
+class UDataflow;
+class UChaosCacheCollection;
+class AChaosCacheManager;
+class UDataflowBaseContent;
+
+namespace Dataflow
+{
+	/** Evaluate the simulation graph given a simulation context and a simulation flag */
+	DATAFLOWSIMULATION_API void EvaluateSimulationGraph(const TObjectPtr<UDataflow>& SimulationGraph,
+		const TSharedPtr<Dataflow::FDataflowSimulationContext>& SimulationContext, const float DeltaTime, const float SimulationTime);
+
+	/** Check if the simulation cache nodes have changed to trigger a reset */
+	DATAFLOWSIMULATION_API bool ShouldResetWorld(const TObjectPtr<UDataflow>& SimulationGraph,
+		const TObjectPtr<UWorld>& SimulationWorld, Dataflow::FTimestamp& LastTimeStamp);
+
+	/** Spawn an actor given a class type and attach it to the cache manager */
+	DATAFLOWSIMULATION_API TObjectPtr<AActor> SpawnSimulatedActor(const TSubclassOf<AActor>& ActorClass,
+		const TObjectPtr<AChaosCacheManager>& CacheManager, const TObjectPtr<UChaosCacheCollection>& CacheCollection,
+		const bool bIsRecording, const TObjectPtr<UDataflowBaseContent>& DataflowContent);
+
+	/** Setup the skelmesh animations to be used in the scene/generator */
+	DATAFLOWSIMULATION_API void SetupSkeletonAnimation(const TObjectPtr<AActor>& PreviewActor);
+
+	/** Update the skelmesh animation at some point in time */
+	DATAFLOWSIMULATION_API void UpdateSkeletonAnimation(const TObjectPtr<AActor>& PreviewActor, const float SimulationTime);
+
+	/** Start the skelmesh animation */
+	DATAFLOWSIMULATION_API void StartSkeletonAnimation(const TObjectPtr<AActor>& PreviewActor);
+
+	/** Pause the skelmesh animation */
+	DATAFLOWSIMULATION_API void PauseSkeletonAnimation(const TObjectPtr<AActor>& PreviewActor);
+	
+	/** Step the skelmesh animation */
+	DATAFLOWSIMULATION_API void StepSkeletonAnimation(const TObjectPtr<AActor>& PreviewActor);
+};
+
+
diff --git a/Engine/Plugins/Experimental/Dataflow/Source/DataflowSimulation/Public/Interfaces/DataflowCollisionObject.h b/Engine/Plugins/Experimental/Dataflow/Source/DataflowSimulation/Public/Interfaces/DataflowCollisionObject.h
new file mode 100644
--- /dev/null
+++ b/Engine/Plugins/Experimental/Dataflow/Source/DataflowSimulation/Public/Interfaces/DataflowCollisionObject.h
@@ -0,0 +1,45 @@
+﻿// Copyright Epic Games, Inc. All Rights Reserved.
+
+#pragma once
+#include "Interfaces/DataflowPhysicsObject.h"
+#include "DataflowCollisionObject.generated.h"
+
+/**
+ * Dataflow collision object proxy (PT)
+ */
+USTRUCT()
+struct DATAFLOWSIMULATION_API FDataflowCollisionObjectProxy : public FDataflowPhysicsObjectProxy
+{
+	GENERATED_BODY()
+	
+	FDataflowCollisionObjectProxy();
+	virtual ~FDataflowCollisionObjectProxy() override = default;
+
+	/** Get the proxy script struct */
+	virtual const UScriptStruct* GetScriptStruct() const override
+	{
+		return StaticStruct();
+	}
+};
+
+UINTERFACE()
+class DATAFLOWSIMULATION_API UDataflowCollisionObjectInterface : public UDataflowPhysicsObjectInterface
+{
+	GENERATED_BODY()
+};
+
+/**
+ * Dataflow collision object interface to send/receive datas (GT <-> PT)
+ */
+class DATAFLOWSIMULATION_API IDataflowCollisionObjectInterface : public IDataflowPhysicsObjectInterface
+{
+	GENERATED_BODY()
+public:
+	IDataflowCollisionObjectInterface();
+
+	/** Get the simulation type */
+	virtual FString GetSimulationType() const override
+	{
+		return FDataflowCollisionObjectProxy::StaticStruct()->GetName();
+	}
+};
diff --git a/Engine/Plugins/Experimental/Dataflow/Source/DataflowSimulation/Public/Interfaces/DataflowConstraintObject.h b/Engine/Plugins/Experimental/Dataflow/Source/DataflowSimulation/Public/Interfaces/DataflowConstraintObject.h
new file mode 100644
--- /dev/null
+++ b/Engine/Plugins/Experimental/Dataflow/Source/DataflowSimulation/Public/Interfaces/DataflowConstraintObject.h
@@ -0,0 +1,46 @@
+﻿// Copyright Epic Games, Inc. All Rights Reserved.
+
+#pragma once
+#include "Interfaces/DataflowPhysicsObject.h"
+#include "DataflowConstraintObject.generated.h"
+
+/**
+ * Dataflow collision object proxy (PT)
+ */
+USTRUCT()
+struct DATAFLOWSIMULATION_API FDataflowConstraintObjectProxy : public FDataflowPhysicsObjectProxy
+{
+	GENERATED_BODY()
+	
+	FDataflowConstraintObjectProxy();
+	virtual ~FDataflowConstraintObjectProxy() override = default;
+
+	/** Get the proxy script struct */
+	virtual const UScriptStruct* GetScriptStruct() const override
+	{
+		return StaticStruct();
+	}
+};
+
+UINTERFACE()
+class DATAFLOWSIMULATION_API UDataflowConstraintObjectInterface : public UDataflowPhysicsObjectInterface
+{
+	GENERATED_BODY()
+};
+
+/**
+ * Dataflow collision object interface to send/receive datas (GT <-> PT)
+ */
+class DATAFLOWSIMULATION_API IDataflowConstraintObjectInterface : public IDataflowPhysicsObjectInterface
+{
+	GENERATED_BODY()
+	
+public:
+	IDataflowConstraintObjectInterface();
+
+	/** Get the simulation type */
+	virtual FString GetSimulationType() const override
+	{
+		return FDataflowConstraintObjectProxy::StaticStruct()->GetName();
+	}
+};
\ No newline at end of file
diff --git a/Engine/Plugins/Experimental/Dataflow/Source/DataflowSimulation/Public/Interfaces/DataflowPhysicsObject.h b/Engine/Plugins/Experimental/Dataflow/Source/DataflowSimulation/Public/Interfaces/DataflowPhysicsObject.h
new file mode 100644
--- /dev/null
+++ b/Engine/Plugins/Experimental/Dataflow/Source/DataflowSimulation/Public/Interfaces/DataflowPhysicsObject.h
@@ -0,0 +1,48 @@
+﻿// Copyright Epic Games, Inc. All Rights Reserved.
+
+#pragma once
+#include "UObject/ObjectPtr.h"
+#include "Dataflow/DataflowSimulationInterface.h"
+#include "Dataflow/DataflowSimulationProxy.h"
+#include "DataflowPhysicsObject.generated.h"
+
+/**
+ * Dataflow physics object proxy (PT)
+ */
+USTRUCT()
+struct DATAFLOWSIMULATION_API FDataflowPhysicsObjectProxy : public FDataflowSimulationProxy
+{
+	GENERATED_BODY()
+	
+	FDataflowPhysicsObjectProxy();
+	virtual ~FDataflowPhysicsObjectProxy() override = default;
+
+	/** Get the proxy script struct */
+	virtual const UScriptStruct* GetScriptStruct() const override
+	{
+		return StaticStruct();
+	}
+};
+
+UINTERFACE()
+class DATAFLOWSIMULATION_API UDataflowPhysicsObjectInterface : public UDataflowSimulationInterface
+{
+	GENERATED_BODY()
+};
+
+/**
+ * Dataflow physics object interface to send/receive datas (GT <-> PT)
+ */
+class DATAFLOWSIMULATION_API IDataflowPhysicsObjectInterface : public IDataflowSimulationInterface
+{
+	GENERATED_BODY()
+	
+public:
+	IDataflowPhysicsObjectInterface();
+
+	/** Get the simulation type */
+	virtual FString GetSimulationType() const override
+	{
+		return FDataflowPhysicsObjectProxy::StaticStruct()->GetName();
+	}
+};
diff --git a/Engine/Plugins/Experimental/Dataflow/Source/DataflowSimulation/Public/Interfaces/DataflowPhysicsSolver.h b/Engine/Plugins/Experimental/Dataflow/Source/DataflowSimulation/Public/Interfaces/DataflowPhysicsSolver.h
new file mode 100644
--- /dev/null
+++ b/Engine/Plugins/Experimental/Dataflow/Source/DataflowSimulation/Public/Interfaces/DataflowPhysicsSolver.h
@@ -0,0 +1,65 @@
+﻿// Copyright Epic Games, Inc. All Rights Reserved.
+
+#pragma once
+#include "UObject/ObjectPtr.h"
+#include "Dataflow/DataflowSimulationInterface.h"
+#include "Dataflow/DataflowSimulationProxy.h"
+#include "DataflowPhysicsSolver.generated.h"
+
+class UDataflow;
+struct FPhysicsObjectProxy;
+namespace Dataflow
+{
+	class FPhysicsObjectInterface;
+}
+
+/**
+ * Dataflow simulation physics solver proxy (PT)
+ */
+USTRUCT()
+struct DATAFLOWSIMULATION_API FDataflowPhysicsSolverProxy : public FDataflowSimulationProxy
+{
+	GENERATED_BODY()
+
+	FDataflowPhysicsSolverProxy();
+	virtual ~FDataflowPhysicsSolverProxy() override = default;
+
+	/** Get the proxy script struct */
+	virtual const UScriptStruct* GetScriptStruct() const override
+	{
+		return StaticStruct();
+	}
+	
+	/** Advance the solver datas in time */
+	virtual void AdvanceSolverDatas(const float DeltaTime) {}
+
+	/** Get the solver time step */
+	virtual float GetTimeStep() { return 0.033f;}
+};
+
+UINTERFACE()
+class DATAFLOWSIMULATION_API UDataflowPhysicsSolverInterface : public UDataflowSimulationInterface
+{
+	GENERATED_BODY()
+};
+
+/**
+ * Dataflow physics solver interface to send/receive (GT)
+ */
+class DATAFLOWSIMULATION_API IDataflowPhysicsSolverInterface : public IDataflowSimulationInterface
+{
+	GENERATED_BODY()
+	
+public:
+	IDataflowPhysicsSolverInterface();
+	
+	/** Get the simulation type */
+	virtual FString GetSimulationType() const override
+	{
+		return FDataflowPhysicsSolverProxy::StaticStruct()->GetName();
+	}
+};
+
+
+
+
diff --git a/Engine/Source/Runtime/Experimental/Dataflow/Engine/Private/Dataflow/DataflowContent.cpp b/Engine/Source/Runtime/Experimental/Dataflow/Engine/Private/Dataflow/DataflowContent.cpp
--- a/Engine/Source/Runtime/Experimental/Dataflow/Engine/Private/Dataflow/DataflowContent.cpp
+++ b/Engine/Source/Runtime/Experimental/Dataflow/Engine/Private/Dataflow/DataflowContent.cpp
@@ -276,9 +276,14 @@ TObjectPtr<UDataflowBaseContent> IDataflowContentOwner::BuildDataflowContent()
 // UDataflowBaseContent
 //
 
-void UDataflowBaseContent::SetIsDirty(bool InDirty) 
+void UDataflowBaseContent::SetConstructionDirty(bool InDirty) 
 { 
-	bIsDirty = InDirty;
+	bIsConstructionDirty = InDirty;
+}
+
+void UDataflowBaseContent::SetSimulationDirty(bool InDirty) 
+{ 
+	bIsSimulationDirty = InDirty;
 }
 
 UDataflowBaseContent::UDataflowBaseContent()
@@ -314,7 +319,8 @@ void UDataflowBaseContent::SetDataflowOwner(const TObjectPtr<UObject>& InOwner)
 		DataflowContext = MakeShared<Dataflow::FEngineContext>(nullptr, nullptr, Dataflow::FTimestamp::Invalid);
 	}
 	DataflowContext->Owner = InOwner;  
-	SetIsDirty(true);
+	SetConstructionDirty(true);
+	SetSimulationDirty(true);
 }
 
 TObjectPtr<UObject> UDataflowBaseContent::GetDataflowOwner() const 
@@ -329,7 +335,8 @@ void UDataflowBaseContent::SetDataflowAsset(const TObjectPtr<UDataflow>& Dataflo
 		DataflowContext = MakeShared<Dataflow::FEngineContext>(nullptr, nullptr, Dataflow::FTimestamp::Invalid);
 	}
 	DataflowContext->Graph = DataflowAsset;  
-	SetIsDirty(true);
+	SetConstructionDirty(true);
+	SetSimulationDirty(true);
 }
 
 TObjectPtr<UDataflow> UDataflowBaseContent::GetDataflowAsset() const 
@@ -344,7 +351,11 @@ void UDataflowBaseContent::SetLastModifiedTimestamp(Dataflow::FTimestamp InTimes
 		LastModifiedTimestamp = InTimestamp; 
 		if (bMakeDirty)
 		{
-			SetIsDirty(true);
+			if(GetDataflowAsset() && GetDataflowAsset()->Type == EDataflowType::Construction)
+			{
+				SetConstructionDirty(true);
+				SetSimulationDirty(true);
+			}
 			MarkPackageDirty(); 
 		}
 	}
@@ -353,7 +364,8 @@ void UDataflowBaseContent::SetLastModifiedTimestamp(Dataflow::FTimestamp InTimes
 void UDataflowBaseContent::SetDataflowContext(const TSharedPtr<Dataflow::FEngineContext>& InContext) 
 { 
 	DataflowContext = InContext;  
-	SetIsDirty(true); 
+	SetConstructionDirty(true);
+	SetSimulationDirty(true);
 	MarkPackageDirty();
 }
 
@@ -380,6 +392,31 @@ void UDataflowBaseContent::AddReferencedObjects(UObject* InThis, FReferenceColle
 	Super::AddReferencedObjects(InThis, Collector);
 }
 
+void OverrideActorProperty(TObjectPtr<AActor>& PreviewActor, const TObjectPtr<UObject>& PropertyValue, const FName& PropertyName)
+{
+	if(PreviewActor && PropertyValue)
+	{
+		if(const FProperty* DataflowProperty = PreviewActor->GetClass()->FindPropertyByName(PropertyName))
+		{
+			if (const FObjectProperty* ObjectProperty = CastField<FObjectProperty>(DataflowProperty))
+			{
+				if(ObjectProperty->PropertyClass == PropertyValue->GetClass())
+				{
+					if(UObject** PropertyObject = DataflowProperty->ContainerPtrToValuePtr<UObject*>(PreviewActor))
+					{
+						(*PropertyObject) = PropertyValue;
+					}
+				}
+			}
+		}
+	}
+}
+
+void UDataflowBaseContent::SetActorProperties(TObjectPtr<AActor>& PreviewActor) const
+{
+	OverrideActorProperty(PreviewActor, GetDataflowOwner(), TEXT("DataflowAsset"));
+}
+
 //
 // UDataflowSkeletalContent
 //
@@ -398,7 +435,8 @@ void UDataflowSkeletalContent::SetSkeletalMesh(const TObjectPtr<USkeletalMesh>&
 			SetSkeleton(SkeletalMesh->GetSkeleton());
 		}
 	}
-	SetIsDirty(true);
+	SetConstructionDirty(true);
+	SetSimulationDirty(true);
 }
 
 void UDataflowSkeletalContent::SetAnimationAsset(const TObjectPtr<UAnimationAsset>& SkeletalAnimationAsset)
@@ -408,7 +446,8 @@ void UDataflowSkeletalContent::SetAnimationAsset(const TObjectPtr<UAnimationAsse
 	{
 		SetSkeleton(AnimationAsset->GetSkeleton());
 	}
-	SetIsDirty(true);
+	SetConstructionDirty(true);
+	SetSimulationDirty(true);
 }
 
 void UDataflowSkeletalContent::SetSkeleton(const TObjectPtr<USkeleton>& SkeletonAsset)
@@ -422,7 +461,8 @@ void UDataflowSkeletalContent::SetSkeleton(const TObjectPtr<USkeleton>& Skeleton
 	{
 		SetAnimationAsset(nullptr);
 	}
-	SetIsDirty(true);
+	SetConstructionDirty(true);
+	SetSimulationDirty(true);
 }
 
 #if WITH_EDITOR
@@ -463,3 +503,11 @@ void UDataflowSkeletalContent::AddReferencedObjects(UObject* InThis, FReferenceC
 	Super::AddReferencedObjects(InThis, Collector);
 }
 
+void UDataflowSkeletalContent::SetActorProperties(TObjectPtr<AActor>& PreviewActor) const
+{
+	Super::SetActorProperties(PreviewActor);
+
+	OverrideActorProperty(PreviewActor, AnimationAsset, TEXT("AnimationAsset"));
+	OverrideActorProperty(PreviewActor, SkeletalMesh, TEXT("SkeletalMesh"));
+}
+
diff --git a/Engine/Source/Runtime/Experimental/Dataflow/Engine/Private/Dataflow/DataflowObject.cpp b/Engine/Source/Runtime/Experimental/Dataflow/Engine/Private/Dataflow/DataflowObject.cpp
--- a/Engine/Source/Runtime/Experimental/Dataflow/Engine/Private/Dataflow/DataflowObject.cpp
+++ b/Engine/Source/Runtime/Experimental/Dataflow/Engine/Private/Dataflow/DataflowObject.cpp
@@ -10,6 +10,17 @@
 
 #define LOCTEXT_NAMESPACE "UDataflow"
 
+
+namespace Dataflow
+{
+	namespace CVars
+	{
+		/** Enable the simulation dataflow (for now WIP) */
+        DATAFLOWENGINE_API bool bEnableSimulationDataflow = false;
+        FAutoConsoleVariableRef CVarEnableSimulationDataflow(TEXT("p.Dataflow.EnableSimulation"), bEnableSimulationDataflow, TEXT("If true enable the use of simulation dataflow (WIP)"));
+	}
+}
+
 FDataflowAssetEdit::FDataflowAssetEdit(UDataflow* InAsset, FPostEditFunctionCallback InCallback)
 	: PostEditCallback(InCallback)
 	, Asset(InAsset)
@@ -129,5 +140,24 @@ void UDataflow::Serialize(FArchive& Ar)
 	Dataflow->Serialize(Ar, this);
 }
 
+#if WITH_EDITOR
+bool UDataflow::CanEditChange(const FProperty* InProperty) const
+{
+	if (!Super::CanEditChange(InProperty))
+	{
+		return false;
+	}
+
+	const FName& Name = InProperty->GetFName();
+
+	if (Name == GET_MEMBER_NAME_CHECKED(ThisClass, Type))
+	{
+		return Dataflow::CVars::bEnableSimulationDataflow == true;
+	}
+
+	return true;
+}
+#endif
+
 #undef LOCTEXT_NAMESPACE
 
diff --git a/Engine/Source/Runtime/Experimental/Dataflow/Engine/Public/Dataflow/DataflowContent.h b/Engine/Source/Runtime/Experimental/Dataflow/Engine/Public/Dataflow/DataflowContent.h
--- a/Engine/Source/Runtime/Experimental/Dataflow/Engine/Public/Dataflow/DataflowContent.h
+++ b/Engine/Source/Runtime/Experimental/Dataflow/Engine/Public/Dataflow/DataflowContent.h
@@ -131,8 +131,11 @@ public:
 	*	Dirty - State Invalidation
 	*   Check if non-graph specific data has been changed, this usually requires a re-render 
 	*/
-	bool IsDirty() const { return bIsDirty; }
-	void SetIsDirty(bool InDirty);
+	bool IsConstructionDirty() const { return bIsConstructionDirty; }
+	void SetConstructionDirty(bool InDirty);
+	
+	bool IsSimulationDirty() const { return bIsSimulationDirty; }
+	void SetSimulationDirty(bool InDirty);
 
 	/** 
 	*	LastModifiedTimestamp - State Invalidation 
@@ -154,6 +157,9 @@ public:
 
 	/** Collect reference objects for GC */
 	virtual void AddContentObjects(FReferenceCollector& Collector) {}
+
+	/** Set all the preview actor exposed properties */
+	virtual void SetActorProperties(TObjectPtr<AActor>& PreviewActor) const;
 	
 	/** Data flow owner accessors (through the context) */
 	void SetDataflowOwner(const TObjectPtr<UObject>& InOwner);
@@ -164,13 +170,17 @@ public:
 	TObjectPtr<UDataflow> GetDataflowAsset() const;
 
 	/** Data flow terminal accessors */
-	void SetDataflowTerminal(const FString& InPath) { DataflowTerminal = InPath;  SetIsDirty(true);}
+	void SetDataflowTerminal(const FString& InPath) { DataflowTerminal = InPath;  SetConstructionDirty(true); SetSimulationDirty(true);}
 	const FString& GetDataflowTerminal() const { return DataflowTerminal; }
 
 	/** Terminal asset accessors */
-	void SetTerminalAsset(const TObjectPtr<UObject>& InAsset) { TerminalAsset = InAsset;  SetIsDirty(true);}
+	void SetTerminalAsset(const TObjectPtr<UObject>& InAsset) { TerminalAsset = InAsset;  SetConstructionDirty(true); SetSimulationDirty(true);}
 	const TObjectPtr<UObject>& GetTerminalAsset() const { return TerminalAsset;}
 
+	/** Preview class accessors */
+	void SetPreviewClass(const TSubclassOf<AActor>& InPreviewClass) { PreviewClass = InPreviewClass;  SetConstructionDirty(true); SetSimulationDirty(true);}
+	const TSubclassOf<AActor>& GetPreviewClass() const { return PreviewClass;}
+
 	/** Content Serialization */
 	virtual void Serialize(FArchive& Ar);
 
@@ -184,11 +194,11 @@ public:
 protected:
 	
 	/** Data flow terminal path for evaluation */
-	UPROPERTY(EditAnywhere, Category = "Dataflow", Transient, SkipSerialization)
+	UPROPERTY(Transient, SkipSerialization)
 	FString DataflowTerminal = "";
 
 	/** Data flow terminal path for evaluation */
-	UPROPERTY(EditAnywhere, Category = "Dataflow", Transient, SkipSerialization)
+	UPROPERTY(Transient, SkipSerialization)
 	TObjectPtr<UObject> TerminalAsset = nullptr;
 
 	/**  Engine context (data flow owner/asset) to be used for dataflow evaluation */
@@ -199,10 +209,17 @@ protected:
 
     /** Dirty flag to trigger rendering. Do we need that? since when accessing the member by non const ref we will not dirty it */
 	UPROPERTY()
-	bool bIsDirty = true;
+	bool bIsConstructionDirty = true;
+
+	/** Dirty flag to reset the simulation if necessary */
+	UPROPERTY()
+	bool bIsSimulationDirty = true;
 
 	/** Saved as a cached context. Will be automatically saved to a cache directory if true. Use the pvar p.Dataflow.Editor.ContextCaching to enable. [def:false] */
 	bool bIsSaved = false;
+
+	/** Preview actor class that could be used to visualize the result */
+	TSubclassOf<AActor> PreviewClass = nullptr;
 };
 
 /** 
@@ -239,10 +256,13 @@ public:
 	//~ UObject interface
 	static void AddReferencedObjects(UObject* InThis, FReferenceCollector& Collector);
 
+	/** Set all the preview actor exposed properties */
+	virtual void SetActorProperties(TObjectPtr<AActor>& PreviewActor) const override;
+
 protected:
 
 	/** Data flow skeletal mesh*/
-	UPROPERTY(Transient, SkipSerialization)
+	UPROPERTY(EditAnywhere, Category = "Preview", Transient, SkipSerialization)
 	TObjectPtr<USkeletalMesh> SkeletalMesh = nullptr;
 	
 	/** Data flow skeleton */
@@ -250,6 +270,6 @@ protected:
 	TObjectPtr<USkeleton> Skeleton = nullptr;
 	
 	/** Animation asset to be used to preview simulation */
-	UPROPERTY(Transient, SkipSerialization)
+	UPROPERTY(EditAnywhere, Category = "Preview", Transient, SkipSerialization)
 	TObjectPtr<UAnimationAsset> AnimationAsset = nullptr;
 };
diff --git a/Engine/Source/Runtime/Experimental/Dataflow/Engine/Public/Dataflow/DataflowObject.h b/Engine/Source/Runtime/Experimental/Dataflow/Engine/Public/Dataflow/DataflowObject.h
--- a/Engine/Source/Runtime/Experimental/Dataflow/Engine/Public/Dataflow/DataflowObject.h
+++ b/Engine/Source/Runtime/Experimental/Dataflow/Engine/Public/Dataflow/DataflowObject.h
@@ -48,6 +48,17 @@ private:
 	UDataflow* Asset;
 };
 
+/** Data flow types */
+UENUM()
+enum class EDataflowType : uint8
+{
+	/** the dataflow will be used to build assets */
+	Construction,
+
+	/** The dataflow will be used to define the simulation evolution */
+	Simulation
+};
+
 /**
 * UDataflow (UObject)
 *
@@ -76,6 +87,8 @@ public:
 
 	virtual bool IsEditorOnly() const { return true; }
 
+	/** Simulation tag to use in the node registry */
+	static const inline FString SimulationTag = TEXT("DataflowSimulationTag");
 
 public:
 	UPROPERTY(EditAnywhere, Category = "Evaluation")
@@ -87,12 +100,15 @@ public:
 	UPROPERTY(EditAnywhere, Category = "Render")
 	TObjectPtr<UMaterial> Material = nullptr;
 
+	UPROPERTY(EditAnywhere, Category = "Evaluation")
+	EDataflowType Type = EDataflowType::Construction;
 
 public:
 	/** UObject Interface */
 	static DATAFLOWENGINE_API void AddReferencedObjects(UObject* InThis, FReferenceCollector& Collector);
 
 #if WITH_EDITOR
+	virtual bool CanEditChange(const FProperty* InProperty) const override;
 	DATAFLOWENGINE_API virtual void PostEditChangeProperty(struct FPropertyChangedEvent& PropertyChangedEvent) override;
 #endif
 	DATAFLOWENGINE_API virtual void PostLoad() override;
diff --git a/Engine/Source/Runtime/Experimental/Dataflow/Engine/Public/Dataflow/DataflowObjectInterface.h b/Engine/Source/Runtime/Experimental/Dataflow/Engine/Public/Dataflow/DataflowObjectInterface.h
--- a/Engine/Source/Runtime/Experimental/Dataflow/Engine/Public/Dataflow/DataflowObjectInterface.h
+++ b/Engine/Source/Runtime/Experimental/Dataflow/Engine/Public/Dataflow/DataflowObjectInterface.h
@@ -39,101 +39,6 @@ namespace Dataflow
 
 	};
 
-	/** Simulation flag that could be used across simulation nodes.
-	 * The setup and clean modes are game thread only whereas the
-	 * update mode could run on game thread and physics thread
-	 */
-	enum class ESimulationFlags : uint8
-	{
-		SetupSimulation = 1 << 0,
-		CleanSimulation = 1 << 1,
-		UpdateSimulation = 1 << 2
-	};
-
-	ENUM_CLASS_FLAGS(ESimulationFlags);
-
-	/** Caching mode that could be used across simulation nodes */
-	enum class ECachingMode : uint8
-	{
-		PlayCache = 1 << 0,
-		RecordCache = 1 << 1
-	};
-
-	/** Simulation context that will be used by all the simulation/cache nodes*/
-	template<class Base = FContextSingle>
-	class TSimulationContext : public TEngineContext<Base>
-	{
-	public:
-		DATAFLOW_CONTEXT_INTERNAL(TEngineContext<Base>, TSimulationContext);
-
-		TSimulationContext(const TObjectPtr<UObject>& InOwner,const TObjectPtr<UDataflow>& InGraph, FTimestamp InTimestamp,
-			const TObjectPtr<AActor>& InRootActor, const ECachingMode InCachingMode)
-				: Super(InOwner, InGraph, InTimestamp), SimulationWorld(InRootActor->GetWorld()), RootActor(InRootActor),
-					SimulationFlags(ESimulationFlags::SetupSimulation), CachingMode(InCachingMode)
-		{
-			
-		}
-		
-		virtual ~TSimulationContext() override {};
-		
-		/** Set the timing infos */
-		void SetTimingInfos(const float DeltaSeconds, const float TimeSeconds) {DeltaTime = DeltaSeconds; SimulationTime = TimeSeconds;}
-
-		/** Get the delta time in seconds */
-		float GetDeltaTime() const {return DeltaTime;}
-		
-		/** Get the simulation time in seconds */
-		float GetSimulationTime() const {return SimulationTime;}
-		
-		/** Get the simulation world */
-		const TObjectPtr<UWorld>& GetSimulationWorld() const {return SimulationWorld;}
-		
-		/** Get the root actor */
-		const TObjectPtr<AActor>& GetRootActor() const {return RootActor;}
-
-		/** Check if the simulation mode is present in the simulation context */
-		bool HasSimulationFlag(const ESimulationFlags InSimulationFlag) const
-		{
-			return (SimulationFlags & InSimulationFlag) == InSimulationFlag;
-		}
-
-		/** Set the simulation mode */
-		void SetSimulationFlag(const ESimulationFlags InSimulationFlag) {SimulationFlags |= InSimulationFlag;}
-
-		/** Unset the simulation mode */
-		void UnsetSimulationFlag(const ESimulationFlags InSimulationFlag) {SimulationFlags &= ~InSimulationFlag;}
-
-		/** Get the caching mode */
-		ECachingMode GetCachingMode() const {return CachingMode;}
-
-		/** Set the caching mode */
-		void SetCachingMode(const ECachingMode InCachingMode) {CachingMode = InCachingMode;}
-		
-	private :
-		
-		/** Evolution world used to register components */
-		TObjectPtr<UWorld> SimulationWorld;
-
-		/** Root actor of the world */
-		TObjectPtr<AActor> RootActor;
-
-		/** Simulation time */
-		float SimulationTime = 0.0f;
-
-		/** Delta time */
-		float DeltaTime = 0.0f;
-
-		/** Simulation flags that could be used across simulation nodes */
-		ESimulationFlags SimulationFlags;
-
-		/** Caching mode that could be used across cache nodes */
-		ECachingMode CachingMode;
-	};
-
 	typedef TEngineContext<FContextSingle> FEngineContext;
 	typedef TEngineContext<FContextThreaded> FEngineContextThreaded;
-
-	typedef TSimulationContext<FContextSingle> FSimulationContext;
-	typedef TSimulationContext<FContextThreaded> FSimulationContextThreaded;
-
 }
