2 - Dataflow editor simulation support

diff --git a/Engine/Plugins/Experimental/ChaosCaching/Source/ChaosCaching/Private/Chaos/CacheManagerActor.cpp b/Engine/Plugins/Experimental/ChaosCaching/Source/ChaosCaching/Private/Chaos/CacheManagerActor.cpp
--- a/Engine/Plugins/Experimental/ChaosCaching/Source/ChaosCaching/Private/Chaos/CacheManagerActor.cpp
+++ b/Engine/Plugins/Experimental/ChaosCaching/Source/ChaosCaching/Private/Chaos/CacheManagerActor.cpp
@@ -820,7 +820,8 @@ FObservedComponent* AChaosCacheManager::FindObservedComponent(UPrimitiveComponen
 
 FObservedComponent& AChaosCacheManager::AddNewObservedComponent(UPrimitiveComponent* InComponent)
 {
-	check(InComponent->CreationMethod != EComponentCreationMethod::UserConstructionScript);
+	// this check is preventing components being spawned from CS
+	//check(InComponent->CreationMethod != EComponentCreationMethod::UserConstructionScript);
 	ObservedComponents.AddDefaulted();
 	FObservedComponent& NewEntry = ObservedComponents.Last();
 
@@ -851,13 +852,11 @@ void AChaosCacheManager::FindOrAddObservedComponent(UPrimitiveComponent* InCompo
 	if(InComponent != nullptr)
 	{
 		FObservedComponent* FoundComponent = FindObservedComponent(InComponent);
-		FObservedComponent* ObservedComponent = FoundComponent ? FoundComponent : &AddNewObservedComponent(InComponent);
-
-		if(ObservedComponent)
+		if(FObservedComponent* ObservedComponent = FoundComponent ? FoundComponent : &AddNewObservedComponent(InComponent))
 		{
 			ObservedComponent->bIsSimulating = bTransferSimulation ? InComponent->BodyInstance.bSimulatePhysics : false;
 
-			if(CacheName != TEXT(""))
+			if(CacheName != TEXT("")) 
 			{
 				ObservedComponent->CacheName = CacheName;
 			}
diff --git a/Engine/Plugins/Experimental/CharacterFXEditor/BaseCharacterFXEditor/Source/BaseCharacterFXEditor/Private/BaseCharacterFXEditor.cpp b/Engine/Plugins/Experimental/CharacterFXEditor/BaseCharacterFXEditor/Source/BaseCharacterFXEditor/Private/BaseCharacterFXEditor.cpp
--- a/Engine/Plugins/Experimental/CharacterFXEditor/BaseCharacterFXEditor/Source/BaseCharacterFXEditor/Private/BaseCharacterFXEditor.cpp
+++ b/Engine/Plugins/Experimental/CharacterFXEditor/BaseCharacterFXEditor/Source/BaseCharacterFXEditor/Private/BaseCharacterFXEditor.cpp
@@ -3,7 +3,7 @@
 #include "BaseCharacterFXEditor.h"
 #include "BaseCharacterFXEditorToolkit.h"
 
-void UBaseCharacterFXEditor::Initialize(const TArray<TObjectPtr<UObject>>& InObjects)
+void UBaseCharacterFXEditor::Initialize(const TArray<TObjectPtr<UObject>>& InObjects, const TSubclassOf<AActor>& InPreviewClass)
 {
 	OriginalObjectsToEdit = InObjects;
 
diff --git a/Engine/Plugins/Experimental/CharacterFXEditor/BaseCharacterFXEditor/Source/BaseCharacterFXEditor/Private/BaseCharacterFXEditorToolkit.cpp b/Engine/Plugins/Experimental/CharacterFXEditor/BaseCharacterFXEditor/Source/BaseCharacterFXEditor/Private/BaseCharacterFXEditorToolkit.cpp
--- a/Engine/Plugins/Experimental/CharacterFXEditor/BaseCharacterFXEditor/Source/BaseCharacterFXEditor/Private/BaseCharacterFXEditorToolkit.cpp
+++ b/Engine/Plugins/Experimental/CharacterFXEditor/BaseCharacterFXEditor/Source/BaseCharacterFXEditor/Private/BaseCharacterFXEditorToolkit.cpp
@@ -189,16 +189,19 @@ void FBaseCharacterFXEditorToolkit::PostInitAssetEditor()
 	// a viewport that will allow our mode to receive ticks.
 	// We don't need to invoke the tool palette tab anymore, since this is handled by
 	// underlying infrastructure.
-	if (!TabManager->FindExistingLiveTab(ViewportTabID))
+	if (!TabManager->FindExistingLiveTab(ViewportTabID) && bForceViewportTab)
 	{
 		TabManager->TryInvokeTab(ViewportTabID);
 	}
 
 	ViewportClient->FocusViewportOnBox(EdMode->SceneBoundingBox());
 
-	// We need the viewport client to start out focused, or else it won't get ticked until
-	// we click inside it.
-	ViewportClient->ReceivedFocus(ViewportClient->Viewport);
+	if(ViewportClient->Viewport)
+	{
+		// We need the viewport client to start out focused, or else it won't get ticked until
+		// we click inside it.
+		ViewportClient->ReceivedFocus(ViewportClient->Viewport);
+	}
 }
 
 void FBaseCharacterFXEditorToolkit::InitializeEdMode(UBaseCharacterFXEditorMode* EdMode)
diff --git a/Engine/Plugins/Experimental/CharacterFXEditor/BaseCharacterFXEditor/Source/BaseCharacterFXEditor/Public/BaseCharacterFXEditor.h b/Engine/Plugins/Experimental/CharacterFXEditor/BaseCharacterFXEditor/Source/BaseCharacterFXEditor/Public/BaseCharacterFXEditor.h
--- a/Engine/Plugins/Experimental/CharacterFXEditor/BaseCharacterFXEditor/Source/BaseCharacterFXEditor/Public/BaseCharacterFXEditor.h
+++ b/Engine/Plugins/Experimental/CharacterFXEditor/BaseCharacterFXEditor/Source/BaseCharacterFXEditor/Public/BaseCharacterFXEditor.h
@@ -33,7 +33,7 @@ public:
 	virtual TSharedPtr<FBaseAssetToolkit> CreateToolkit() PURE_VIRTUAL(UBaseCharacterFXEditor::CreateToolkit, return TSharedPtr<FBaseAssetToolkit>(););
 
 	/** Store objects to edit and call UAssetEditor::Initialize */
-	virtual void Initialize(const TArray<TObjectPtr<UObject>>& InObjects);
+	virtual void Initialize(const TArray<TObjectPtr<UObject>>& InObjects, const TSubclassOf<AActor>& InPreviewClass = nullptr);
 
 	/** Returns the asset editor instance interface, so that its window can be focused, for example. */
 	virtual IAssetEditorInstance* GetInstanceInterface();
diff --git a/Engine/Plugins/Experimental/CharacterFXEditor/BaseCharacterFXEditor/Source/BaseCharacterFXEditor/Public/BaseCharacterFXEditorToolkit.h b/Engine/Plugins/Experimental/CharacterFXEditor/BaseCharacterFXEditor/Source/BaseCharacterFXEditor/Public/BaseCharacterFXEditorToolkit.h
--- a/Engine/Plugins/Experimental/CharacterFXEditor/BaseCharacterFXEditor/Source/BaseCharacterFXEditor/Public/BaseCharacterFXEditorToolkit.h
+++ b/Engine/Plugins/Experimental/CharacterFXEditor/BaseCharacterFXEditor/Source/BaseCharacterFXEditor/Public/BaseCharacterFXEditorToolkit.h
@@ -63,5 +63,8 @@ protected:
 
 	// A default "CharacterFXEditor" category is created in RegisterTabSpawners. Override that function to set this to a concrete class type
 	TSharedPtr<FWorkspaceItem> EditorMenuCategory;
+
+	// Boolean to force the construction of the viewport tab even if the layout is not adding it 
+	bool bForceViewportTab = true;
 };
 
diff --git a/Engine/Plugins/Experimental/Dataflow/Source/DataflowEditor/Private/Dataflow/DataflowEditor.cpp b/Engine/Plugins/Experimental/Dataflow/Source/DataflowEditor/Private/Dataflow/DataflowEditor.cpp
--- a/Engine/Plugins/Experimental/Dataflow/Source/DataflowEditor/Private/Dataflow/DataflowEditor.cpp
+++ b/Engine/Plugins/Experimental/Dataflow/Source/DataflowEditor/Private/Dataflow/DataflowEditor.cpp
@@ -25,7 +25,7 @@ TSharedPtr<FBaseAssetToolkit> UDataflowEditor::CreateToolkit()
 	return DataflowToolkit;
 }
 
-void UDataflowEditor::Initialize(const TArray<TObjectPtr<UObject>>& InObjects)
+void UDataflowEditor::Initialize(const TArray<TObjectPtr<UObject>>& InObjects, const TSubclassOf<AActor>& InPreviewClass)
 {
 	if(!InObjects.IsEmpty())
 	{
@@ -59,6 +59,10 @@ void UDataflowEditor::Initialize(const TArray<TObjectPtr<UObject>>& InObjects)
 				}
 			}
 		}
+		if(EditorContent && InPreviewClass)
+		{
+			EditorContent->SetPreviewClass(InPreviewClass);
+		}
 		RequiredObjects.Add(EditorContent);
 
 		// Update the editor datas (skeleton information for the viewer)
@@ -124,16 +128,6 @@ void UDataflowEditor::RemoveTerminalContents(const TSharedPtr<Dataflow::FGraph>&
 
 void UDataflowEditor::AddTerminalContents(const TSharedPtr<Dataflow::FGraph>& DataflowGraph, ValidTerminalsType& ValidTerminals)
 {
-	auto BuildTerminalContent = [this](const TObjectPtr<UDataflowBaseContent>& TerminalContent,
-					const FString& TerminalName, const TObjectPtr<UObject>& TerminalAsset)
-	{
-		// Set the context (owner/asset) onto the terminal content
-		TerminalContent->SetDataflowContext(EditorContent->GetDataflowContext());
-				
-		TerminalContent->SetDataflowTerminal(TerminalName);
-		TerminalContent->SetTerminalAsset(TerminalAsset);
-	};
-			
 	for(const TSharedPtr<FDataflowNode>& DataflowNode : DataflowGraph->GetFilteredNodes(FDataflowTerminalNode::StaticType()))
 	{
 		if(const FDataflowTerminalNode* TerminalNode = DataflowNode->AsType<FDataflowTerminalNode>())
@@ -145,14 +139,15 @@ void UDataflowEditor::AddTerminalContents(const TSharedPtr<Dataflow::FGraph>& Da
 				{
 					TerminalContents.Add(TerminalOwner->BuildDataflowContent());
 					TerminalContent = &TerminalContents.Last();
-							
+					(*TerminalContent)->SetDataflowTerminal(TerminalNode->GetName().ToString());
+					(*TerminalContent)->SetDataflowContext(EditorContent->GetDataflowContext());
+					
 					(*TerminalContent)->SetLastModifiedTimestamp(EditorContent->GetLastModifiedTimestamp());
 					bHasTerminalsDirty = true;
 				}
 				if(TerminalNode->GetTerminalAsset() != (*TerminalContent)->GetTerminalAsset())
 				{
-					BuildTerminalContent(*TerminalContent,
-							TerminalNode->GetName().ToString(), TerminalNode->GetTerminalAsset());
+					(*TerminalContent)->SetTerminalAsset(TerminalNode->GetTerminalAsset());
 					bHasTerminalsDirty = true;
 				}
 			}
diff --git a/Engine/Plugins/Experimental/Dataflow/Source/DataflowEditor/Private/Dataflow/DataflowEditorCommands.cpp b/Engine/Plugins/Experimental/Dataflow/Source/DataflowEditor/Private/Dataflow/DataflowEditorCommands.cpp
--- a/Engine/Plugins/Experimental/Dataflow/Source/DataflowEditor/Private/Dataflow/DataflowEditorCommands.cpp
+++ b/Engine/Plugins/Experimental/Dataflow/Source/DataflowEditor/Private/Dataflow/DataflowEditorCommands.cpp
@@ -34,7 +34,10 @@
 
 const FString FDataflowEditorCommandsImpl::BeginWeightMapPaintToolIdentifier = TEXT("BeginWeightMapPaintTool");
 const FString FDataflowEditorCommandsImpl::AddWeightMapNodeIdentifier = TEXT("AddWeightMapNode");
-const FString FDataflowEditorCommandsImpl::UpdateSimulationCacheIdentifier = TEXT("UpdateSimulationCache");
+const FString FDataflowEditorCommandsImpl::RebuildSimulationSceneIdentifier = TEXT("RebuildSimulationScene");
+const FString FDataflowEditorCommandsImpl::PauseSimulationSceneIdentifier = TEXT("PauseSimulationScene");
+const FString FDataflowEditorCommandsImpl::StartSimulationSceneIdentifier = TEXT("StartSimulationScene");
+const FString FDataflowEditorCommandsImpl::StepSimulationSceneIdentifier = TEXT("StepSimulationScene");
 
 // @todo(brice) Remove Example Tools
 //const FString FDataflowEditorCommandsImpl::BeginAttributeEditorToolIdentifier = TEXT("BeginAttributeEditorTool");
@@ -64,7 +67,10 @@ void FDataflowEditorCommandsImpl::RegisterCommands()
 
 	UI_COMMAND(BeginWeightMapPaintTool, "Add Weight Map", "Paint weight maps on the mesh", EUserInterfaceActionType::None, FInputChord());
 	UI_COMMAND(AddWeightMapNode, "Add Weight Map", "Paint weight maps on the mesh", EUserInterfaceActionType::Button, FInputChord());
-	UI_COMMAND(UpdateSimulationCache, "UpdateSimulationCache", "Update Simulation Cache", EUserInterfaceActionType::ToggleButton, FInputChord());
+	UI_COMMAND(RebuildSimulationScene, "RebuildSimulationScene", "Rebuild the simulation scene", EUserInterfaceActionType::ToggleButton, FInputChord());
+	UI_COMMAND(PauseSimulationScene, "PauseSimulationScene", "Pause the simulation scene", EUserInterfaceActionType::ToggleButton, FInputChord());
+	UI_COMMAND(StartSimulationScene, "StartSimulationScene", "Start the simulation scene", EUserInterfaceActionType::ToggleButton, FInputChord());
+	UI_COMMAND(StepSimulationScene, "StepSimulationScene", "Step the simulation scene", EUserInterfaceActionType::ToggleButton, FInputChord());
 
 	for (const TPair<FName, TUniquePtr<Dataflow::IDataflowConstructionViewMode>>& NameAndMode : Dataflow::FRenderingViewModeFactory::GetInstance().GetViewModes())
 	{
diff --git a/Engine/Plugins/Experimental/Dataflow/Source/DataflowEditor/Private/Dataflow/DataflowEditorMode.cpp b/Engine/Plugins/Experimental/Dataflow/Source/DataflowEditor/Private/Dataflow/DataflowEditorMode.cpp
--- a/Engine/Plugins/Experimental/Dataflow/Source/DataflowEditor/Private/Dataflow/DataflowEditorMode.cpp
+++ b/Engine/Plugins/Experimental/Dataflow/Source/DataflowEditor/Private/Dataflow/DataflowEditorMode.cpp
@@ -294,12 +294,17 @@ void UDataflowEditorMode::Exit()
 {
 	UActorComponent::MarkRenderStateDirtyEvent.RemoveAll(this);
 
-	ConstructionScene->ResetConstructionScene();
-	ConstructionScene = nullptr;
-
-	SimulationScene->ResetSimulationScene();
-	SimulationScene = nullptr;
+	if(ConstructionScene)
+	{
+		ConstructionScene->ResetConstructionScene();
+		ConstructionScene = nullptr;
+	}
 
+	if(SimulationScene)
+	{
+		SimulationScene->ResetSimulationScene();
+		SimulationScene = nullptr;
+	}
 
 	Super::Exit();
 }
diff --git a/Engine/Plugins/Experimental/Dataflow/Source/DataflowEditor/Private/Dataflow/DataflowEditorPreviewSceneBase.cpp b/Engine/Plugins/Experimental/Dataflow/Source/DataflowEditor/Private/Dataflow/DataflowEditorPreviewSceneBase.cpp
--- a/Engine/Plugins/Experimental/Dataflow/Source/DataflowEditor/Private/Dataflow/DataflowEditorPreviewSceneBase.cpp
+++ b/Engine/Plugins/Experimental/Dataflow/Source/DataflowEditor/Private/Dataflow/DataflowEditorPreviewSceneBase.cpp
@@ -51,7 +51,6 @@ const TArray<TObjectPtr<UDataflowBaseContent>>& FDataflowPreviewSceneBase::GetTe
 void FDataflowPreviewSceneBase::AddReferencedObjects(FReferenceCollector& Collector)
 {
 	FAdvancedPreviewScene::AddReferencedObjects(Collector);
-	Collector.AddReferencedObject(RootSceneActor);
 	if (const TObjectPtr<UDataflowBaseContent> EditorContent = GetEditorContent())
 	{
 		EditorContent->AddContentObjects(Collector);
diff --git a/Engine/Plugins/Experimental/Dataflow/Source/DataflowEditor/Private/Dataflow/DataflowEditorStyle.h b/Engine/Plugins/Experimental/Dataflow/Source/DataflowEditor/Private/Dataflow/DataflowEditorStyle.h
--- a/Engine/Plugins/Experimental/Dataflow/Source/DataflowEditor/Private/Dataflow/DataflowEditorStyle.h
+++ b/Engine/Plugins/Experimental/Dataflow/Source/DataflowEditor/Private/Dataflow/DataflowEditorStyle.h
@@ -46,12 +46,9 @@ public:
 		//const FString MeshSelectionPropertyName = "DataflowEditor." + FDataflowEditorCommandsImpl::BeginMeshSelectionToolIdentifier;
 		//Set(*MeshSelectionPropertyName, new FSlateImageBrush(RootToContentDir(TEXT("Slate/Dataflow_SelectVertex_20x.png")), Icon20x20));
 		
-		FString PropertyNameString = "DataflowEditor." + FDataflowEditorCommandsImpl::AddWeightMapNodeIdentifier;
+		const FString PropertyNameString = "DataflowEditor." + FDataflowEditorCommandsImpl::AddWeightMapNodeIdentifier;
 		Set(*PropertyNameString, new IMAGE_BRUSH_SVG("Icons/PaintMaps", ToolbarIconSize));
 
-		PropertyNameString = "DataflowEditor." + FDataflowEditorCommandsImpl::UpdateSimulationCacheIdentifier;
-		Set(*PropertyNameString, new IMAGE_BRUSH_SVG("Icons/LV_Record", ToolbarIconSize));
-
 		DefaultMaterial = Cast<UMaterial>(StaticLoadObject(UMaterial::StaticClass(), NULL, TEXT("/Engine/BasicShapes/BasicShapeMaterial")));
 		VertexMaterial = Cast<UMaterial>(StaticLoadObject(UMaterial::StaticClass(), NULL, TEXT("/Engine/EditorMaterials/Dataflow/DataflowVertexMaterial")));
 		FSlateStyleRegistry::RegisterSlateStyle(*this);
diff --git a/Engine/Plugins/Experimental/Dataflow/Source/DataflowEditor/Private/Dataflow/DataflowEditorToolkit.cpp b/Engine/Plugins/Experimental/Dataflow/Source/DataflowEditor/Private/Dataflow/DataflowEditorToolkit.cpp
--- a/Engine/Plugins/Experimental/Dataflow/Source/DataflowEditor/Private/Dataflow/DataflowEditorToolkit.cpp
+++ b/Engine/Plugins/Experimental/Dataflow/Source/DataflowEditor/Private/Dataflow/DataflowEditorToolkit.cpp
@@ -82,7 +82,7 @@ FDataflowEditorToolkit::FDataflowEditorToolkit(UAssetEditor* InOwningAssetEditor
 {
 	check(DataflowEditor);
 
-	StandaloneDefaultLayout = FTabManager::NewLayout(FName("DataflowEditorLayout03"))
+	ConstructionDefaultLayout = FTabManager::NewLayout(FName("DataflowConstructionLayout03"))
 		->AddArea
 		(
 			FTabManager::NewPrimaryArea()->SetOrientation(Orient_Horizontal)
@@ -104,10 +104,18 @@ FDataflowEditorToolkit::FDataflowEditorToolkit(UAssetEditor* InOwningAssetEditor
 					->Split
 					(
 						FTabManager::NewStack()
-						->SetSizeCoefficient(0.9f)		// Relative width of (Construction Viewport) vs (Tools Panel, Preview Viewport)
+						->SetSizeCoefficient(0.45f)		// Relative width of (Construction Viewport) vs (Tools Panel, Preview Viewport)
 						->AddTab(ViewportTabID, ETabState::OpenedTab)
-						->SetExtensionId("RestSpaceViewportArea")
-						->SetHideTabWell(false)
+						->SetExtensionId("ViewportArea")
+						->SetHideTabWell(true)
+					)
+					->Split
+					(
+						FTabManager::NewStack()
+						->SetSizeCoefficient(0.45f)		// Relative width of (Construction Viewport) vs (Tools Panel, Preview Viewport)
+						->AddTab(SimulationViewportTabId, ETabState::OpenedTab)
+						->SetExtensionId("ViewportArea")
+						->SetHideTabWell(true)
 					)
 				)
 				->Split
@@ -142,8 +150,86 @@ FDataflowEditorToolkit::FDataflowEditorToolkit(UAssetEditor* InOwningAssetEditor
 					FTabManager::NewStack()
 					->SetSizeCoefficient(0.65f)	// Relative height of (Asset Details, Preview Scene Details) vs (Dataflow Node Details)
 					->AddTab(DetailsTabID, ETabState::OpenedTab)
+					->AddTab(PreviewSceneTabId, ETabState::OpenedTab)
+					->SetExtensionId("DetailsArea")
+					->SetHideTabWell(true)
+					->SetForegroundTab(DetailsTabID)
+				)
+				->Split
+				(
+					FTabManager::NewStack()
+					->SetSizeCoefficient(0.2f)	// Relative height of (Dataflow Node Details) vs (Asset Details, Preview Scene Details)
+					->AddTab(NodeDetailsTabId, ETabState::OpenedTab)
+					->SetExtensionId("NodeDetailsArea")
+					->SetHideTabWell(true)
+				)
+			)
+		);
+
+	SimulationDefaultLayout = FTabManager::NewLayout(FName("DataflowSimulationLayout02"))
+		->AddArea
+		(
+			FTabManager::NewPrimaryArea()->SetOrientation(Orient_Horizontal)
+			->Split
+			(
+				FTabManager::NewSplitter()->SetOrientation(Orient_Vertical)
+				->SetSizeCoefficient(0.8f)	// Relative width of (Tools Panel, Construction Viewport, Preview Viewport, Dataflow Graph Editor, Outliner) vs (Asset Details, Preview Scene Details, Dataflow Node Details)
+				->Split
+				(
+					FTabManager::NewSplitter()->SetOrientation(Orient_Horizontal)
+					->SetSizeCoefficient(0.60f)	// Relative height of (Tools Panel, Construction Viewport, Preview Viewport) vs (Dataflow Graph Editor, Outliner)
+					->Split
+					(
+						FTabManager::NewStack()
+						->SetSizeCoefficient(0.1f)		// Relative width of (Tools Panel) vs (Construction Viewport, Preview Viewport)
+						->SetExtensionId(UDataflowEditorUISubsystem::EditorSidePanelAreaName)
+						->SetHideTabWell(true)
+					)
+					->Split
+					(
+						FTabManager::NewStack()
+						->SetSizeCoefficient(0.9f)		// Relative width of (Construction Viewport) vs (Tools Panel, Preview Viewport)
+						->AddTab(ViewportTabID, ETabState::ClosedTab)
+						->AddTab(SimulationViewportTabId, ETabState::OpenedTab)
+						->SetExtensionId("ViewportArea")
+						->SetHideTabWell(false)
+					)
+				)
+				->Split
+				(
+					FTabManager::NewSplitter()->SetOrientation(Orient_Horizontal)
+					->SetSizeCoefficient(0.40f)	// Relative height of (Dataflow Node Details) vs (Asset Details, Preview Scene Details)
+					->Split
+					(
+						FTabManager::NewStack()
+						->SetSizeCoefficient(0.2f)
+						->AddTab(CollectionSpreadSheetTabId_1, ETabState::ClosedTab)
+						->SetExtensionId("CollectionSpreadSheetArea")
+						->SetHideTabWell(false)
+					)
+					->Split
+					(
+						FTabManager::NewStack()
+						->SetSizeCoefficient(0.8f)	// Relative height of (Dataflow Graph Editor, Outliner) vs (Tools Panel, Construction Viewport, Preview Viewport)
+						->AddTab(GraphCanvasTabId, ETabState::OpenedTab)
+						->SetExtensionId("GraphEditorArea")
+						->SetHideTabWell(false)
+						->SetForegroundTab(GraphCanvasTabId)
+					)
+				)
+			)
+			->Split
+			(
+				FTabManager::NewSplitter()->SetOrientation(Orient_Vertical)
+				->SetSizeCoefficient(0.2f)	// Relative width of (Asset Details, Preview Scene Details, Dataflow Node Details) vs (Tools Panel, Construction Viewport, Preview Viewport, Dataflow Graph Editor, Outliner)
+				->Split
+				(
+					FTabManager::NewStack()
+					->SetSizeCoefficient(0.65f)	// Relative height of (Asset Details, Preview Scene Details) vs (Dataflow Node Details)
+					->AddTab(DetailsTabID, ETabState::OpenedTab)
+					->AddTab(PreviewSceneTabId, ETabState::OpenedTab)
 					->SetExtensionId("DetailsArea")
-					->SetHideTabWell(false)
+					->SetHideTabWell(true)
 					->SetForegroundTab(DetailsTabID)
 				)
 				->Split
@@ -152,11 +238,25 @@ FDataflowEditorToolkit::FDataflowEditorToolkit(UAssetEditor* InOwningAssetEditor
 					->SetSizeCoefficient(0.2f)	// Relative height of (Dataflow Node Details) vs (Asset Details, Preview Scene Details)
 					->AddTab(NodeDetailsTabId, ETabState::OpenedTab)
 					->SetExtensionId("NodeDetailsArea")
-					->SetHideTabWell(false)
+					->SetHideTabWell(true)
 				)
 			)
 		);
 
+	if(TObjectPtr<UDataflowBaseContent>& EditorContent = DataflowEditor->GetEditorContent())
+	{
+		if(EditorContent->GetDataflowAsset() && EditorContent->GetDataflowAsset()->Type == EDataflowType::Simulation)
+		{
+			StandaloneDefaultLayout = SimulationDefaultLayout;
+			bForceViewportTab = false;
+		}
+		else
+		{
+			StandaloneDefaultLayout = ConstructionDefaultLayout;
+			bForceViewportTab = true;
+		}
+	}
+
 	// Add any extenders specified by the UISubsystem
 	// The extenders provide defined locations for FModeToolkit to attach
 	// tool palette tabs and detail panel tabs
@@ -427,6 +527,50 @@ void FDataflowEditorToolkit::GetSaveableObjects(TArray<UObject*>& OutObjects) co
 
 //~ End FBaseCharacterFXEditorToolkit overrides
 
+class FPreviewSceneDescriptionCustomization : public IDetailCustomization
+{
+public:
+	FPreviewSceneDescriptionCustomization(const TArray<UDataflowBaseContent*>& DataflowContents);
+
+	virtual ~FPreviewSceneDescriptionCustomization() {}
+
+	
+	/**Customize details for the description */
+	virtual void CustomizeDetails(IDetailLayoutBuilder& DetailBuilder) override;
+
+	private :
+		/** List of dataflow contents to preview */
+		TMap<FString,TArray<UObject*>> ContentTypesObjects;
+};
+
+FPreviewSceneDescriptionCustomization::FPreviewSceneDescriptionCustomization(const TArray<UDataflowBaseContent*>& DataflowContents) :
+	IDetailCustomization(), ContentTypesObjects()
+{
+	static const FString CategoryName = TEXT("Terminals");
+	for(UDataflowBaseContent* DataflowContent : DataflowContents)
+	{
+		TArray<UObject*>& ContentTypeObjects = ContentTypesObjects.FindOrAdd(CategoryName);
+		ContentTypeObjects.Add(DataflowContent); 
+	}
+}
+
+void FPreviewSceneDescriptionCustomization::CustomizeDetails(IDetailLayoutBuilder& DetailBuilder)
+{
+	FAddPropertyParams PropertyParams;
+	PropertyParams.AllowChildren(true);
+	PropertyParams.CreateCategoryNodes(false);
+	PropertyParams.HideRootObjectNode(true);
+	for(TPair<FString, TArray<UObject*>>& ContentTypeObjects : ContentTypesObjects)
+	{
+		DetailBuilder.EditCategory(*ContentTypeObjects.Key).AddExternalObjects(ContentTypeObjects.Value, EPropertyLocation::Common, PropertyParams);
+	}
+}
+
+TSharedRef<class IDetailCustomization> FDataflowEditorToolkit::CustomizePreviewSceneDescription() const
+{
+	const TArray<UDataflowBaseContent*> SimulationContents  = TArray<UDataflowBaseContent*>{SimulationScene->GetEditorContent()};
+	return MakeShareable(new FPreviewSceneDescriptionCustomization(SimulationContents));
+}
 
 //~ Begin FBaseAssetToolkit overrides
 
@@ -439,15 +583,26 @@ void FDataflowEditorToolkit::CreateWidgets()
 		if (UDataflow* DataflowAsset = EditorContent->GetDataflowAsset())
 		{
 			NodeDetailsEditor = CreateNodeDetailsEditorWidget(EditorContent->GetDataflowOwner());
-			AssetDetailsEditor = CreateAssetDetailsEditorWidget({EditorContent->GetDataflowOwner(),EditorContent->GetDataflowAsset()});
+			if(EditorContent->GetDataflowOwner() != EditorContent->GetDataflowAsset())
+			{
+				AssetDetailsEditor = CreateAssetDetailsEditorWidget({EditorContent->GetDataflowOwner(),EditorContent->GetDataflowAsset()});
+			}
+			else
+			{
+				AssetDetailsEditor = CreateAssetDetailsEditorWidget({EditorContent->GetDataflowAsset()});
+			}
 			GraphEditor = CreateGraphEditorWidget(DataflowAsset, NodeDetailsEditor);
 			CreateSimulationViewportClient();
 
 			FAdvancedPreviewSceneModule& AdvancedPreviewSceneModule = FModuleManager::LoadModuleChecked<FAdvancedPreviewSceneModule>("AdvancedPreviewScene");
 			
+			TArray<FAdvancedPreviewSceneModule::FDetailCustomizationInfo> DetailsCustomizations;
+			DetailsCustomizations.Add({ UDataflowSimulationSceneDescription::StaticClass(),
+				FOnGetDetailCustomizationInstance::CreateSP(const_cast<FDataflowEditorToolkit*>(this), &FDataflowEditorToolkit::CustomizePreviewSceneDescription) });
+			
 			AdvancedPreviewSettingsWidget = AdvancedPreviewSceneModule.CreateAdvancedPreviewSceneSettingsWidget(SimulationScene.ToSharedRef(), 
 				SimulationScene->GetPreviewSceneDescription(),
-				TArray<FAdvancedPreviewSceneModule::FDetailCustomizationInfo>(), 
+				DetailsCustomizations, 
 				TArray<FAdvancedPreviewSceneModule::FPropertyTypeCustomizationInfo>(),
 				TArray<FAdvancedPreviewSceneModule::FDetailDelegates>());
 		}
@@ -462,7 +617,7 @@ AssetEditorViewportFactoryFunction FDataflowEditorToolkit::GetViewportDelegate()
 	{
 		TSharedRef<SDataflowConstructionViewport> Viewport = SAssignNew(DataflowConstructionViewport, SDataflowConstructionViewport, InArgs)
 			.ViewportClient(StaticCastSharedPtr<FDataflowConstructionViewportClient>(ViewportClient));
-
+		
 		if (UDataflowEditorMode* const DataflowMode = Cast<UDataflowEditorMode>(EditorModeManager->GetActiveScriptableMode(UDataflowEditorMode::EM_DataflowEditorModeId)))
 		{
 			if (TSharedPtr<FModeToolkit> ModeToolkit = DataflowMode->GetToolkit().Pin())
@@ -865,7 +1020,7 @@ void FDataflowEditorToolkit::Tick(float DeltaTime)
 				TimeStamp = Dataflow::FTimestamp::Invalid;
 			}
 
-			// Update the list of dataflow teminal contents 
+			// Update the list of dataflow terminal contents 
 			DataflowEditor->UpdateTerminalContents(TimeStamp);
 			
 			// OnTick evaluation only pulls the terminal nodes. The other evaluations can be specific nodes.
diff --git a/Engine/Plugins/Experimental/Dataflow/Source/DataflowEditor/Private/Dataflow/DataflowSNode.cpp b/Engine/Plugins/Experimental/Dataflow/Source/DataflowEditor/Private/Dataflow/DataflowSNode.cpp
--- a/Engine/Plugins/Experimental/Dataflow/Source/DataflowEditor/Private/Dataflow/DataflowSNode.cpp
+++ b/Engine/Plugins/Experimental/Dataflow/Source/DataflowEditor/Private/Dataflow/DataflowSNode.cpp
@@ -285,24 +285,34 @@ EVisibility SDataflowEdNode::IsAddPinButtonVisible() const
 //
 // Add a menu option to create a graph node.
 //
-TSharedPtr<FAssetSchemaAction_Dataflow_CreateNode_DataflowEdNode> FAssetSchemaAction_Dataflow_CreateNode_DataflowEdNode::CreateAction(UEdGraph* ParentGraph, const FName & InNodeTypeName, const FName& InOverrideNodeName)
+TSharedPtr<FAssetSchemaAction_Dataflow_CreateNode_DataflowEdNode> FAssetSchemaAction_Dataflow_CreateNode_DataflowEdNode::CreateAction(const UEdGraph* ParentGraph, const FName & InNodeTypeName, const FName& InOverrideNodeName)
 {
-	if (Dataflow::FNodeFactory* Factory = Dataflow::FNodeFactory::GetInstance())
+	if (const UDataflow* Dataflow = Cast<UDataflow>(ParentGraph))
 	{
-		const Dataflow::FFactoryParameters& Param = Factory->GetParameters(InNodeTypeName);
-		if (Param.IsValid())
+		if (Dataflow::FNodeFactory* Factory = Dataflow::FNodeFactory::GetInstance())
 		{
-			const FText ToolTip = FText::FromString(Param.ToolTip.IsEmpty() ? FString("Add a Dataflow node.") : Param.ToolTip);
-			FText NodeName = FText::FromString(Param.DisplayName.ToString());
-			if (!InOverrideNodeName.IsNone())
+			const Dataflow::FFactoryParameters& Param = Factory->GetParameters(InNodeTypeName);
+			if (Param.IsValid())
 			{
-				NodeName = FText::FromName(InOverrideNodeName);
+				const bool bIsSimulationNode = Param.Tags.Contains(UDataflow::SimulationTag);
+				const bool bIsSimulationGraph = (Dataflow->Type == EDataflowType::Simulation);
+				
+				if((bIsSimulationGraph && bIsSimulationNode) || (!bIsSimulationGraph && !bIsSimulationNode))
+				{
+					const FText ToolTip = FText::FromString(Param.ToolTip.IsEmpty() ? FString("Add a Dataflow node.") : Param.ToolTip);
+					FText NodeName = FText::FromString(Param.DisplayName.ToString());
+					if (!InOverrideNodeName.IsNone())
+					{
+						NodeName = FText::FromName(InOverrideNodeName);
+					}
+				
+					const FText Category = FText::FromString(Param.Category.ToString().IsEmpty() ? FString("Dataflow") : Param.Category.ToString());
+					const FText Tags = FText::FromString(Param.Tags);
+					TSharedPtr<FAssetSchemaAction_Dataflow_CreateNode_DataflowEdNode> NewNodeAction(
+						new FAssetSchemaAction_Dataflow_CreateNode_DataflowEdNode(InNodeTypeName, Category, NodeName, ToolTip, Tags));
+					return NewNodeAction;
+				}
 			}
-			const FText Category = FText::FromString(Param.Category.ToString().IsEmpty() ? FString("Dataflow") : Param.Category.ToString());
-			const FText Tags = FText::FromString(Param.Tags);
-			TSharedPtr<FAssetSchemaAction_Dataflow_CreateNode_DataflowEdNode> NewNodeAction(
-				new FAssetSchemaAction_Dataflow_CreateNode_DataflowEdNode(InNodeTypeName, Category, NodeName, ToolTip, Tags));
-			return NewNodeAction;
 		}
 	}
 	return TSharedPtr<FAssetSchemaAction_Dataflow_CreateNode_DataflowEdNode>(nullptr);
diff --git a/Engine/Plugins/Experimental/Dataflow/Source/DataflowEditor/Private/Dataflow/DataflowSchema.cpp b/Engine/Plugins/Experimental/Dataflow/Source/DataflowEditor/Private/Dataflow/DataflowSchema.cpp
--- a/Engine/Plugins/Experimental/Dataflow/Source/DataflowEditor/Private/Dataflow/DataflowSchema.cpp
+++ b/Engine/Plugins/Experimental/Dataflow/Source/DataflowEditor/Private/Dataflow/DataflowSchema.cpp
@@ -96,7 +96,7 @@ void UDataflowSchema::GetGraphContextActions(FGraphContextMenuBuilder& ContextMe
 			if (FDataflowEditorCommands::Get().CreateNodesMap.Contains(NodeParameters.TypeName))
 			{
 				if (TSharedPtr<FAssetSchemaAction_Dataflow_CreateNode_DataflowEdNode> Action =
-					FAssetSchemaAction_Dataflow_CreateNode_DataflowEdNode::CreateAction(ContextMenuBuilder.OwnerOfTemporaries, NodeParameters.TypeName, NodeParameters.DisplayName))
+					FAssetSchemaAction_Dataflow_CreateNode_DataflowEdNode::CreateAction(ContextMenuBuilder.CurrentGraph, NodeParameters.TypeName, NodeParameters.DisplayName))
 				{
 					ContextMenuBuilder.AddAction(Action);
 				}
diff --git a/Engine/Plugins/Experimental/Dataflow/Source/DataflowEditor/Private/Dataflow/DataflowSimulationViewport.cpp b/Engine/Plugins/Experimental/Dataflow/Source/DataflowEditor/Private/Dataflow/DataflowSimulationViewport.cpp
--- a/Engine/Plugins/Experimental/Dataflow/Source/DataflowEditor/Private/Dataflow/DataflowSimulationViewport.cpp
+++ b/Engine/Plugins/Experimental/Dataflow/Source/DataflowEditor/Private/Dataflow/DataflowSimulationViewport.cpp
@@ -35,6 +35,15 @@ void SDataflowSimulationViewport::Construct(const FArguments& InArgs, const FAss
 	if(static_cast<FDataflowSimulationScene*>(Client->GetPreviewScene())->CanRunSimulation())
 	{
 		TWeakPtr<FDataflowSimulationScene> SimulationScene = GetSimulationScene();
+
+		auto HasCacheAsset = [SimulationScene]()
+		{
+			if(SimulationScene.Pin() && SimulationScene.Pin()->GetPreviewSceneDescription())
+			{
+				return (SimulationScene.Pin()->GetPreviewSceneDescription()->CacheAsset != nullptr) ? EVisibility::Visible : EVisibility::Collapsed;
+			}
+			return EVisibility::Collapsed;
+		};
             
 		ViewportOverlay->AddSlot()
 		[
@@ -47,7 +56,7 @@ void SDataflowSimulationViewport::Construct(const FArguments& InArgs, const FAss
 			[
 				SNew(SBorder)
 				.BorderImage(FAppStyle::Get().GetBrush("EditorViewport.OverlayBrush"))
-				.Visibility(EVisibility::Visible)
+				.Visibility(TAttribute<EVisibility>::Create(HasCacheAsset))
 				.Padding(10.0f, 2.0f)
 				[
 					SNew(SDataflowSimulationPanel, SimulationScene)
@@ -91,10 +100,49 @@ void SDataflowSimulationViewport::BindCommands()
 	const FDataflowEditorCommandsImpl& CommandInfos = FDataflowEditorCommands::Get();
 
 	CommandList->MapAction(
-		CommandInfos.UpdateSimulationCache,
+		CommandInfos.RebuildSimulationScene,
+		FExecuteAction::CreateLambda([this]()
+		{
+			if(FDataflowSimulationScene* SimulationScene = StaticCast<FDataflowSimulationScene*>(Client->GetPreviewScene()))
+			{
+				SimulationScene->RebuildSimulationScene(false);
+			}
+		}),
+		FCanExecuteAction::CreateLambda([this]() { return true; }),
+		FIsActionChecked::CreateLambda([this]() { return false; }));
+
+	CommandList->MapAction(
+		CommandInfos.PauseSimulationScene,
+		FExecuteAction::CreateLambda([this]()
+		{
+			if(FDataflowSimulationScene* SimulationScene = StaticCast<FDataflowSimulationScene*>(Client->GetPreviewScene()))
+			{
+				SimulationScene->PauseSimulationScene();
+			}
+		}),
+		FCanExecuteAction::CreateLambda([this]() { return true; }),
+		FIsActionChecked::CreateLambda([this]() { return false; }));
+
+	CommandList->MapAction(
+		CommandInfos.StartSimulationScene,
+		FExecuteAction::CreateLambda([this]()
+		{
+			if(FDataflowSimulationScene* SimulationScene = StaticCast<FDataflowSimulationScene*>(Client->GetPreviewScene()))
+			{
+				SimulationScene->StartSimulationScene();
+			}
+		}),
+		FCanExecuteAction::CreateLambda([this]() { return true; }),
+		FIsActionChecked::CreateLambda([this]() { return false; }));
+
+	CommandList->MapAction(
+		CommandInfos.StepSimulationScene,
 		FExecuteAction::CreateLambda([this]()
 		{
-			static_cast<FDataflowSimulationScene*>(Client->GetPreviewScene())->UpdateSimulationCache();
+			if(FDataflowSimulationScene* SimulationScene = StaticCast<FDataflowSimulationScene*>(Client->GetPreviewScene()))
+			{
+				SimulationScene->StepSimulationScene();
+			}
 		}),
 		FCanExecuteAction::CreateLambda([this]() { return true; }),
 		FIsActionChecked::CreateLambda([this]() { return false; }));
diff --git a/Engine/Plugins/Experimental/Dataflow/Source/DataflowEditor/Private/Dataflow/DataflowSimulationViewportToolbar.cpp b/Engine/Plugins/Experimental/Dataflow/Source/DataflowEditor/Private/Dataflow/DataflowSimulationViewportToolbar.cpp
--- a/Engine/Plugins/Experimental/Dataflow/Source/DataflowEditor/Private/Dataflow/DataflowSimulationViewportToolbar.cpp
+++ b/Engine/Plugins/Experimental/Dataflow/Source/DataflowEditor/Private/Dataflow/DataflowSimulationViewportToolbar.cpp
@@ -2,6 +2,7 @@
 #include "Dataflow/DataflowSimulationViewportToolbar.h"
 #include "Dataflow/DataflowSimulationViewport.h"
 #include "Dataflow/DataflowEditorCommands.h"
+#include "Dataflow/DataflowSimulationScene.h"
 #include "Styling/AppStyle.h"
 #include "Framework/MultiBox/MultiBoxBuilder.h"
 #include "SEditorViewportToolBarMenu.h"
@@ -21,11 +22,26 @@ void SDataflowSimulationViewportToolBar::Construct(const FArguments& InArgs, TSh
 
 void SDataflowSimulationViewportToolBar::ExtendLeftAlignedToolbarSlots(TSharedPtr<SHorizontalBox> MainBoxPtr, TSharedPtr<SViewportToolBar> ParentToolBarPtr) const
 {
+	const TSharedPtr<class FDataflowSimulationScene>& SimulationScene = EditorViewport.Pin()->GetSimulationScene();
+
+	auto HasCacheAsset = [SimulationScene]()
+	{
+		if(SimulationScene && SimulationScene->GetPreviewSceneDescription())
+		{
+			return (SimulationScene->GetPreviewSceneDescription()->CacheAsset == nullptr) ? EVisibility::Visible : EVisibility::Collapsed;
+		}
+		return EVisibility::Collapsed;
+	};
+	
 	const FMargin ToolbarSlotPadding(2.0f, 2.0f);
 	MainBoxPtr->AddSlot()
 		.Padding(ToolbarSlotPadding)
 		[
-			MakeToolBar(Extenders)
+			SNew(SBox)
+			.Visibility(TAttribute<EVisibility>::Create(HasCacheAsset))
+			[
+				MakeToolBar(Extenders)
+			]
 		];
 }
 
@@ -40,14 +56,33 @@ TSharedRef<SWidget> SDataflowSimulationViewportToolBar::MakeToolBar(const TShare
 	ToolbarBuilder.BeginSection("Sim Controls");
 	ToolbarBuilder.BeginBlockGroup();
 	{
-		// the simulation caching should be triggered from the simulation panel as well
-		// we keep it here as well to have a template to add future simulation controls
-		ToolbarBuilder.AddToolBarButton(FDataflowEditorCommands::Get().UpdateSimulationCache,
+		ToolbarBuilder.AddToolBarButton(FDataflowEditorCommands::Get().RebuildSimulationScene,
+			NAME_None,
+			TAttribute<FText>(),
+			TAttribute<FText>(),
+			FSlateIcon(FAppStyle::Get().GetStyleSetName(), "Animation.Backward_End"),
+			FName(*FDataflowEditorCommands::Get().RebuildSimulationSceneIdentifier));
+
+		ToolbarBuilder.AddToolBarButton(FDataflowEditorCommands::Get().PauseSimulationScene,
+			NAME_None,
+			TAttribute<FText>(),
+			TAttribute<FText>(),
+			FSlateIcon(FAppStyle::Get().GetStyleSetName(), "Animation.Pause"),
+			FName(*FDataflowEditorCommands::Get().PauseSimulationSceneIdentifier));
+
+		ToolbarBuilder.AddToolBarButton(FDataflowEditorCommands::Get().StartSimulationScene,
+			NAME_None,
+			TAttribute<FText>(),
+			TAttribute<FText>(),
+			FSlateIcon(FAppStyle::Get().GetStyleSetName(), "Icons.Play"),
+			FName(*FDataflowEditorCommands::Get().StartSimulationSceneIdentifier));
+
+		ToolbarBuilder.AddToolBarButton(FDataflowEditorCommands::Get().StepSimulationScene,
 			NAME_None,
 			TAttribute<FText>(),
 			TAttribute<FText>(),
-			FSlateIcon(FAppStyle::Get().GetStyleSetName(), "Animation.Record"),
-			FName(*FDataflowEditorCommands::Get().UpdateSimulationCacheIdentifier));
+			FSlateIcon(FAppStyle::Get().GetStyleSetName(), "Animation.Forward_Step"),
+			FName(*FDataflowEditorCommands::Get().StepSimulationSceneIdentifier));
 	}
 	ToolbarBuilder.EndBlockGroup();
 	ToolbarBuilder.EndSection();
diff --git a/Engine/Plugins/Experimental/Dataflow/Source/DataflowEditor/Public/Dataflow/DataflowEditor.h b/Engine/Plugins/Experimental/Dataflow/Source/DataflowEditor/Public/Dataflow/DataflowEditor.h
--- a/Engine/Plugins/Experimental/Dataflow/Source/DataflowEditor/Public/Dataflow/DataflowEditor.h
+++ b/Engine/Plugins/Experimental/Dataflow/Source/DataflowEditor/Public/Dataflow/DataflowEditor.h
@@ -44,7 +44,7 @@ public:
 	virtual TSharedPtr<FBaseAssetToolkit> CreateToolkit() override;
 
 	/** Initialize editor contents given a list of objects */
-	virtual void Initialize(const TArray<TObjectPtr<UObject>>& InObjects) override;
+	virtual void Initialize(const TArray<TObjectPtr<UObject>>& InObjects, const TSubclassOf<AActor>& InPreviewClass = nullptr) override;
 
 	/** Update the terminal contents */
 	void UpdateTerminalContents(const Dataflow::FTimestamp TimeStamp);
diff --git a/Engine/Plugins/Experimental/Dataflow/Source/DataflowEditor/Public/Dataflow/DataflowEditorCommands.h b/Engine/Plugins/Experimental/Dataflow/Source/DataflowEditor/Public/Dataflow/DataflowEditorCommands.h
--- a/Engine/Plugins/Experimental/Dataflow/Source/DataflowEditor/Public/Dataflow/DataflowEditorCommands.h
+++ b/Engine/Plugins/Experimental/Dataflow/Source/DataflowEditor/Public/Dataflow/DataflowEditorCommands.h
@@ -63,9 +63,18 @@ public:
 	const static FString AddWeightMapNodeIdentifier;
 	TSharedPtr<FUICommandInfo> AddWeightMapNode;
 
-	const static FString UpdateSimulationCacheIdentifier;
-	TSharedPtr<FUICommandInfo> UpdateSimulationCache;
+	const static FString RebuildSimulationSceneIdentifier;
+	TSharedPtr<FUICommandInfo> RebuildSimulationScene;
+	
+	const static FString PauseSimulationSceneIdentifier;
+	TSharedPtr<FUICommandInfo> PauseSimulationScene;
 
+	const static FString StartSimulationSceneIdentifier;
+	TSharedPtr<FUICommandInfo> StartSimulationScene;
+	
+	const static FString StepSimulationSceneIdentifier;
+	TSharedPtr<FUICommandInfo> StepSimulationScene;
+	
 	// @todo(brice) Remove Example Tools
 	//const static FString BeginAttributeEditorToolIdentifier;
 	//TSharedPtr<FUICommandInfo> BeginAttributeEditorTool;
diff --git a/Engine/Plugins/Experimental/Dataflow/Source/DataflowEditor/Public/Dataflow/DataflowEditorToolkit.h b/Engine/Plugins/Experimental/Dataflow/Source/DataflowEditor/Public/Dataflow/DataflowEditorToolkit.h
--- a/Engine/Plugins/Experimental/Dataflow/Source/DataflowEditor/Public/Dataflow/DataflowEditorToolkit.h
+++ b/Engine/Plugins/Experimental/Dataflow/Source/DataflowEditor/Public/Dataflow/DataflowEditorToolkit.h
@@ -44,6 +44,8 @@ namespace Dataflow
 	};
 }
 
+
+
 class DATAFLOWEDITOR_API FDataflowEditorToolkit final : public FBaseCharacterFXEditorToolkit, public FTickableEditorObject, public FNotifyHook, public FGCObject
 {
 	using FBaseCharacterFXEditorToolkit::ObjectScene;
@@ -170,6 +172,9 @@ private:
 	TSharedPtr<FDataflowCollectionSpreadSheet> DataflowCollectionSpreadSheet_4;
 	TSharedPtr<SWidget> AdvancedPreviewSettingsWidget;
 
+	/** Customize preview scene with editor/terminal contents */
+	TSharedRef<class IDetailCustomization> CustomizePreviewSceneDescription() const;
+
 	// Utility factory functions to build the widgets
 	TSharedRef<SDataflowGraphEditor> CreateGraphEditorWidget(UDataflow* ObjectToEdit, TSharedPtr<IStructureDetailsView> PropertiesEditor);
     TSharedPtr<IDetailsView> CreateAssetDetailsEditorWidget(const TArray<UObject*>& ObjectsToEdit);
@@ -213,4 +218,10 @@ private:
 
 	/** Simulation Viewport client */
 	TSharedPtr<FDataflowSimulationViewportClient> SimulationViewportClient;
+
+	/** Simulation default layout */
+	TSharedPtr<FTabManager::FLayout> SimulationDefaultLayout;
+	
+	/** Simulation default layout */
+	TSharedPtr<FTabManager::FLayout> ConstructionDefaultLayout;
 };
diff --git a/Engine/Plugins/Experimental/Dataflow/Source/DataflowEditor/Public/Dataflow/DataflowSNode.h b/Engine/Plugins/Experimental/Dataflow/Source/DataflowEditor/Public/Dataflow/DataflowSNode.h
--- a/Engine/Plugins/Experimental/Dataflow/Source/DataflowEditor/Public/Dataflow/DataflowSNode.h
+++ b/Engine/Plugins/Experimental/Dataflow/Source/DataflowEditor/Public/Dataflow/DataflowSNode.h
@@ -86,7 +86,7 @@ public:
 	FAssetSchemaAction_Dataflow_CreateNode_DataflowEdNode(const FName& InType, FText InNodeCategory, FText InMenuDesc, FText InToolTip, FText InKeywords)
 		: FEdGraphSchemaAction(InNodeCategory, InMenuDesc, InToolTip, 0, InKeywords), NodeTypeName(InType) {}
 
-	static TSharedPtr<FAssetSchemaAction_Dataflow_CreateNode_DataflowEdNode> CreateAction(UEdGraph* ParentGraph, const FName & NodeTypeName, const FName& InOverrideNodeName = NAME_None);
+	static TSharedPtr<FAssetSchemaAction_Dataflow_CreateNode_DataflowEdNode> CreateAction(const UEdGraph* ParentGraph, const FName & NodeTypeName, const FName& InOverrideNodeName = NAME_None);
 
 	virtual UEdGraphNode* PerformAction(class UEdGraph* ParentGraph, UEdGraphPin* FromPin, const FVector2D Location, bool bSelectNewNode = true) override;
 	//virtual void AddReferencedObjects(FReferenceCollector& Collector) override;
