Dataflow - Fixed the Undo on BreakPinLinks that was missing a transaction.

diff --git a/Engine/Plugins/Experimental/Dataflow/Source/DataflowEditor/Private/Dataflow/DataflowSchema.cpp b/Engine/Plugins/Experimental/Dataflow/Source/DataflowEditor/Private/Dataflow/DataflowSchema.cpp
--- a/Engine/Plugins/Experimental/Dataflow/Source/DataflowEditor/Private/Dataflow/DataflowSchema.cpp
+++ b/Engine/Plugins/Experimental/Dataflow/Source/DataflowEditor/Private/Dataflow/DataflowSchema.cpp
@@ -16,6 +16,7 @@
 #include "ToolMenu.h"
 #include "GraphEditorActions.h"
 #include "Dataflow/DataflowSettings.h"
+#include "ScopedTransaction.h"
 
 #include UE_INLINE_GENERATED_CPP_BY_NAME(DataflowSchema)
 
@@ -298,6 +299,12 @@ void UDataflowSchema::OnPinConnectionDoubleCicked(UEdGraphPin* PinA, UEdGraphPin
 	CreateAndConnectNewReRouteNode(PinA, PinB, GraphPosition);
 }
 
+void UDataflowSchema::BreakPinLinks(UEdGraphPin& TargetPin, bool bSendsNodeNotifcation) const
+{
+	const FScopedTransaction Transaction(LOCTEXT("BreakPinLinks", "Break Pin Links"));
+	Super::BreakPinLinks(TargetPin, bSendsNodeNotifcation);
+}
+
 FConnectionDrawingPolicy* UDataflowSchema::CreateConnectionDrawingPolicy(int32 InBackLayerID, int32 InFrontLayerID, float InZoomFactor, const FSlateRect& InClippingRect, class FSlateWindowElementList& InDrawElements, class UEdGraph* InGraphObj) const
 {
 	return new FDataflowConnectionDrawingPolicy(InBackLayerID, InFrontLayerID, InZoomFactor, InClippingRect, InDrawElements, InGraphObj);
diff --git a/Engine/Plugins/Experimental/Dataflow/Source/DataflowEditor/Public/Dataflow/DataflowSchema.h b/Engine/Plugins/Experimental/Dataflow/Source/DataflowEditor/Public/Dataflow/DataflowSchema.h
--- a/Engine/Plugins/Experimental/Dataflow/Source/DataflowEditor/Public/Dataflow/DataflowSchema.h
+++ b/Engine/Plugins/Experimental/Dataflow/Source/DataflowEditor/Public/Dataflow/DataflowSchema.h
@@ -25,7 +25,7 @@ public:
 	virtual FLinearColor GetPinTypeColor(const FEdGraphPinType& PinType) const override;
 	virtual FConnectionDrawingPolicy* CreateConnectionDrawingPolicy(int32 InBackLayerID, int32 InFrontLayerID, float InZoomFactor, const FSlateRect& InClippingRect, class FSlateWindowElementList& InDrawElements, class UEdGraph* InGraphObj) const override;
 	virtual void OnPinConnectionDoubleCicked(UEdGraphPin* PinA, UEdGraphPin* PinB, const FVector2D& GraphPosition) const override;
-
+	virtual void BreakPinLinks(UEdGraphPin& TargetPin, bool bSendsNodeNotifcation) const override;
 	//~ End EdGraphSchema Interface
 
 	static FLinearColor GetTypeColor(const FName& Type);
