Fix dataflow preview scene customisation defined multiple times

diff --git a/Engine/Plugins/Experimental/Dataflow/Source/DataflowEditor/Private/Dataflow/DataflowEditorToolkit.cpp b/Engine/Plugins/Experimental/Dataflow/Source/DataflowEditor/Private/Dataflow/DataflowEditorToolkit.cpp
--- a/Engine/Plugins/Experimental/Dataflow/Source/DataflowEditor/Private/Dataflow/DataflowEditorToolkit.cpp
+++ b/Engine/Plugins/Experimental/Dataflow/Source/DataflowEditor/Private/Dataflow/DataflowEditorToolkit.cpp
@@ -527,12 +527,12 @@ void FDataflowEditorToolkit::GetSaveableObjects(TArray<UObject*>& OutObjects) co
 
 //~ End FBaseCharacterFXEditorToolkit overrides
 
-class FPreviewSceneDescriptionCustomization : public IDetailCustomization
+class FDataflowPreviewSceneDescriptionCustomization : public IDetailCustomization
 {
 public:
-	FPreviewSceneDescriptionCustomization(const TArray<UDataflowBaseContent*>& DataflowContents);
+	FDataflowPreviewSceneDescriptionCustomization(const TArray<UDataflowBaseContent*>& DataflowContents);
 
-	virtual ~FPreviewSceneDescriptionCustomization() {}
+	virtual ~FDataflowPreviewSceneDescriptionCustomization() {}
 
 	
 	/**Customize details for the description */
@@ -543,7 +543,7 @@ public:
 		TMap<FString,TArray<UObject*>> ContentTypesObjects;
 };
 
-FPreviewSceneDescriptionCustomization::FPreviewSceneDescriptionCustomization(const TArray<UDataflowBaseContent*>& DataflowContents) :
+FDataflowPreviewSceneDescriptionCustomization::FDataflowPreviewSceneDescriptionCustomization(const TArray<UDataflowBaseContent*>& DataflowContents) :
 	IDetailCustomization(), ContentTypesObjects()
 {
 	static const FString CategoryName = TEXT("Terminals");
@@ -554,7 +554,7 @@ FPreviewSceneDescriptionCustomization::FPreviewSceneDescriptionCustomization(con
 	}
 }
 
-void FPreviewSceneDescriptionCustomization::CustomizeDetails(IDetailLayoutBuilder& DetailBuilder)
+void FDataflowPreviewSceneDescriptionCustomization::CustomizeDetails(IDetailLayoutBuilder& DetailBuilder)
 {
 	FAddPropertyParams PropertyParams;
 	PropertyParams.AllowChildren(true);
@@ -569,7 +569,7 @@ void FPreviewSceneDescriptionCustomization::CustomizeDetails(IDetailLayoutBuilde
 TSharedRef<class IDetailCustomization> FDataflowEditorToolkit::CustomizePreviewSceneDescription() const
 {
 	const TArray<UDataflowBaseContent*> SimulationContents  = TArray<UDataflowBaseContent*>{SimulationScene->GetEditorContent()};
-	return MakeShareable(new FPreviewSceneDescriptionCustomization(SimulationContents));
+	return MakeShareable(new FDataflowPreviewSceneDescriptionCustomization(SimulationContents));
 }
 
 //~ Begin FBaseAssetToolkit overrides
