Dataflow : GC collection component scope fix.
Reset the selection when the render is reset, avoiding a GC crash.

diff --git a/Engine/Plugins/Experimental/Dataflow/Source/DataflowEditor/Private/Dataflow/DataflowConstructionScene.cpp b/Engine/Plugins/Experimental/Dataflow/Source/DataflowEditor/Private/Dataflow/DataflowConstructionScene.cpp
--- a/Engine/Plugins/Experimental/Dataflow/Source/DataflowEditor/Private/Dataflow/DataflowConstructionScene.cpp
+++ b/Engine/Plugins/Experimental/Dataflow/Source/DataflowEditor/Private/Dataflow/DataflowConstructionScene.cpp
@@ -347,6 +347,19 @@ bool FDataflowConstructionScene::HasRenderableGeometry()
 
 void FDataflowConstructionScene::ResetConstructionScene()
 {
+	// The ModeManagerss::USelection will hold references to Components, but 
+	// does not report them to the garbage collector. We need to clear the
+	// saved selection when the scene is rebuilt. @todo(Dataflow) If that 
+	// selection needs to persist across render resets, we will also need to
+	// buffer the names of the selected objects so they can be reselected.
+	if (GetDataflowModeManager())
+	{
+		if (USelection* SelectedComponents = GetDataflowModeManager()->GetSelectedComponents())
+		{
+			SelectedComponents->DeselectAll();
+		}
+	}
+
 	// Some objects, like the UMeshElementsVisualizer and Settings Objects
 	// are not part of a tool, so they won't get ticked.This member holds
 	// ticked objects that get rebuilt on Update
