Dataflow: Move rendering code to Editor module

diff --git a/Engine/Plugins/Experimental/ChaosFlesh/Source/ChaosFleshDeprecatedNodes/Private/Dataflow/ChaosFleshConstructTetGridNode.cpp b/Engine/Plugins/Experimental/ChaosFlesh/Source/ChaosFleshDeprecatedNodes/Private/Dataflow/ChaosFleshConstructTetGridNode.cpp
--- a/Engine/Plugins/Experimental/ChaosFlesh/Source/ChaosFleshDeprecatedNodes/Private/Dataflow/ChaosFleshConstructTetGridNode.cpp
+++ b/Engine/Plugins/Experimental/ChaosFlesh/Source/ChaosFleshDeprecatedNodes/Private/Dataflow/ChaosFleshConstructTetGridNode.cpp
@@ -2,6 +2,7 @@
 
 #include "Dataflow/ChaosFleshConstructTetGridNode.h"
 
+#include "Chaos/Utilities.h"
 #include "ChaosFlesh/FleshCollection.h"
 #include "ChaosFlesh/ChaosFlesh.h"
 #include "ChaosLog.h"
diff --git a/Engine/Plugins/Experimental/ChaosFlesh/Source/ChaosFleshDeprecatedNodes/Private/Dataflow/ChaosFleshCreateTetrahedralCollectionNode.cpp b/Engine/Plugins/Experimental/ChaosFlesh/Source/ChaosFleshDeprecatedNodes/Private/Dataflow/ChaosFleshCreateTetrahedralCollectionNode.cpp
--- a/Engine/Plugins/Experimental/ChaosFlesh/Source/ChaosFleshDeprecatedNodes/Private/Dataflow/ChaosFleshCreateTetrahedralCollectionNode.cpp
+++ b/Engine/Plugins/Experimental/ChaosFlesh/Source/ChaosFleshDeprecatedNodes/Private/Dataflow/ChaosFleshCreateTetrahedralCollectionNode.cpp
@@ -2,6 +2,7 @@
 
 #include "Dataflow/ChaosFleshCreateTetrahedralCollectionNode.h"
 
+#include "Chaos/Utilities.h"
 #include "ChaosFlesh/ChaosFlesh.h"
 #include "ChaosFlesh/FleshCollection.h"
 #include "ChaosFlesh/FleshCollectionUtility.h"
diff --git a/Engine/Plugins/Experimental/Dataflow/Source/DataflowEditor/Private/Dataflow/DataflowConstructionScene.cpp b/Engine/Plugins/Experimental/Dataflow/Source/DataflowEditor/Private/Dataflow/DataflowConstructionScene.cpp
--- a/Engine/Plugins/Experimental/Dataflow/Source/DataflowEditor/Private/Dataflow/DataflowConstructionScene.cpp
+++ b/Engine/Plugins/Experimental/Dataflow/Source/DataflowEditor/Private/Dataflow/DataflowConstructionScene.cpp
@@ -8,6 +8,7 @@
 #include "Dataflow/DataflowEditor.h"
 #include "Dataflow/DataflowObject.h"
 #include "Dataflow/DataflowEditorStyle.h"
+#include "Dataflow/DataflowRenderingFactory.h"
 #include "Drawing/MeshElementsVisualizer.h"
 #include "Elements/Framework/EngineElementsLibrary.h"
 #include "Selection.h"
@@ -141,8 +142,7 @@ void FDataflowConstructionScene::UpdateDynamicMeshComponents()
 					GeometryCollection::Facades::FRenderingFacade Facade(*RenderCollection);
 					Facade.DefineSchema();
 
-					check(EditorContent->GetConstructionViewMode());
-					Target->Render(Facade, DataflowContext.ToSharedRef(), *EditorContent->GetConstructionViewMode());
+					Dataflow::RenderNodeOutput(Facade, *Target, *EditorContent);
 
 					const int32 NumGeometry = Facade.NumGeometry();
 					for (int32 MeshIndex = 0; MeshIndex < NumGeometry; ++MeshIndex)
@@ -180,8 +180,7 @@ void FDataflowConstructionScene::UpdateDynamicMeshComponents()
 					GeometryCollection::Facades::FRenderingFacade Facade(*RenderCollection);
 					Facade.DefineSchema();
 
-					check(EditorContent->GetConstructionViewMode());
-					Target->Render(Facade, DataflowContext.ToSharedRef(), *EditorContent->GetConstructionViewMode());
+					Dataflow::RenderNodeOutput(Facade, *Target, *EditorContent);
 
 					const int32 NumGeometry = Facade.NumGeometry();
 					for (int32 MeshIndex = 0; MeshIndex < NumGeometry; ++MeshIndex)
diff --git a/Engine/Plugins/Experimental/Dataflow/Source/DataflowEditor/Private/Dataflow/DataflowConstructionViewportClient.cpp b/Engine/Plugins/Experimental/Dataflow/Source/DataflowEditor/Private/Dataflow/DataflowConstructionViewportClient.cpp
--- a/Engine/Plugins/Experimental/Dataflow/Source/DataflowEditor/Private/Dataflow/DataflowConstructionViewportClient.cpp
+++ b/Engine/Plugins/Experimental/Dataflow/Source/DataflowEditor/Private/Dataflow/DataflowConstructionViewportClient.cpp
@@ -19,19 +19,6 @@
 #include "SGraphPanel.h"
 #include "SNodePanel.h"
 
-
-static_assert(std::is_same<std::underlying_type_t<ELevelViewportType>, int>::value);
-static_assert(std::is_same<std::underlying_type_t<Dataflow::EDataflowViewportType>, int>::value);
-static_assert(static_cast<int>(ELevelViewportType::LVT_OrthoXY) == static_cast<int>(Dataflow::EDataflowViewportType::DVT_OrthoXY));
-static_assert(static_cast<int>(ELevelViewportType::LVT_OrthoXZ) == static_cast<int>(Dataflow::EDataflowViewportType::DVT_OrthoXZ));
-static_assert(static_cast<int>(ELevelViewportType::LVT_OrthoYZ) == static_cast<int>(Dataflow::EDataflowViewportType::DVT_OrthoYZ));
-static_assert(static_cast<int>(ELevelViewportType::LVT_Perspective) == static_cast<int>(Dataflow::EDataflowViewportType::DVT_Perspective));
-static_assert(static_cast<int>(ELevelViewportType::LVT_OrthoFreelook) == static_cast<int>(Dataflow::EDataflowViewportType::DVT_OrthoFreelook));
-static_assert(static_cast<int>(ELevelViewportType::LVT_OrthoNegativeXY) == static_cast<int>(Dataflow::EDataflowViewportType::DVT_OrthoNegativeXY));
-static_assert(static_cast<int>(ELevelViewportType::LVT_OrthoNegativeXZ) == static_cast<int>(Dataflow::EDataflowViewportType::DVT_OrthoNegativeXZ));
-static_assert(static_cast<int>(ELevelViewportType::LVT_OrthoNegativeYZ) == static_cast<int>(Dataflow::EDataflowViewportType::DVT_OrthoNegativeYZ));
-static_assert(static_cast<int>(ELevelViewportType::LVT_MAX) == static_cast<int>(Dataflow::EDataflowViewportType::DVT_MAX));
-
 #define LOCTEXT_NAMESPACE "DataflowConstructionViewportClient"
 
 FDataflowConstructionViewportClient::FDataflowConstructionViewportClient(FEditorModeTools* InModeTools,
diff --git a/Engine/Plugins/Experimental/Dataflow/Source/DataflowEditor/Private/Dataflow/DataflowEditorMode.cpp b/Engine/Plugins/Experimental/Dataflow/Source/DataflowEditor/Private/Dataflow/DataflowEditorMode.cpp
--- a/Engine/Plugins/Experimental/Dataflow/Source/DataflowEditor/Private/Dataflow/DataflowEditorMode.cpp
+++ b/Engine/Plugins/Experimental/Dataflow/Source/DataflowEditor/Private/Dataflow/DataflowEditorMode.cpp
@@ -15,6 +15,7 @@
 #include "Dataflow/DataflowContent.h"
 #include "Dataflow/DataflowEditorCommands.h"
 #include "Dataflow/DataflowEditorModeToolkit.h"
+#include "Dataflow/DataflowEditorUtil.h"
 #include "Dataflow/DataflowGraphEditor.h"
 #include "Dataflow/DataflowEditorPreviewSceneBase.h"
 #include "Dataflow/DataflowConstructionScene.h"
@@ -589,24 +590,19 @@ bool UDataflowEditorMode::CanChangeConstructionViewModeTo(const FName& NewViewMo
 {
 	if (!GetToolManager()->HasActiveTool(EToolSide::Left))
 	{
-		if (const TSharedPtr<const SDataflowGraphEditor> PinnedDataflowGraphEditor = DataflowGraphEditor.Pin())
+		if (const TObjectPtr<UDataflowBaseContent>& EditorContent = ConstructionScene->GetEditorContent())
 		{
-			if (const UEdGraphNode* const SelectedNode = PinnedDataflowGraphEditor->GetSingleSelectedNode())
+			if (const TSharedPtr<const SDataflowGraphEditor> PinnedDataflowGraphEditor = DataflowGraphEditor.Pin())
 			{
-				if (const UDataflowEdNode* const SelectedDataflowNode = Cast<UDataflowEdNode>(SelectedNode))
+				if (const UEdGraphNode* const SelectedNode = PinnedDataflowGraphEditor->GetSingleSelectedNode())
 				{
-					const Dataflow::FRenderingViewModeFactory& ViewModeFactory = Dataflow::FRenderingViewModeFactory::GetInstance();
-
-					if (const Dataflow::IDataflowConstructionViewMode* const ViewMode = ViewModeFactory.GetViewMode(NewViewModeName))
+					if (const UDataflowEdNode* const SelectedDataflowEdNode = Cast<UDataflowEdNode>(SelectedNode))
 					{
-						if (const TObjectPtr<UDataflowBaseContent> EditorContent = ConstructionScene->GetEditorContent())
+						if (const Dataflow::IDataflowConstructionViewMode* const ViewMode = Dataflow::FRenderingViewModeFactory::GetInstance().GetViewMode(NewViewModeName))
 						{
-							if (const TSharedPtr<Dataflow::FEngineContext> Context = EditorContent->GetDataflowContext())
+							if (Dataflow::CanRenderNodeOutput(*SelectedDataflowEdNode, *EditorContent, *ViewMode))
 							{
-								if (SelectedDataflowNode->CanRender(Context.ToSharedRef(), *ViewMode))
-								{
-									return true;
-								}
+								return true;
 							}
 						}
 					}
diff --git a/Engine/Plugins/Experimental/Dataflow/Source/DataflowEditor/Private/Dataflow/DataflowEditorModule.cpp b/Engine/Plugins/Experimental/Dataflow/Source/DataflowEditor/Private/Dataflow/DataflowEditorModule.cpp
--- a/Engine/Plugins/Experimental/Dataflow/Source/DataflowEditor/Private/Dataflow/DataflowEditorModule.cpp
+++ b/Engine/Plugins/Experimental/Dataflow/Source/DataflowEditor/Private/Dataflow/DataflowEditorModule.cpp
@@ -7,6 +7,7 @@
 #include "Dataflow/DataflowEditorMode.h"
 #include "Dataflow/DataflowEditorToolkit.h"
 #include "Dataflow/DataflowEditorCommands.h"
+#include "Dataflow/DataflowEngineRendering.h"
 #include "Dataflow/DataflowSNodeFactories.h"
 
 #include "CoreMinimal.h"
@@ -16,12 +17,16 @@
 
 #define LOCTEXT_NAMESPACE "DataflowEditor"
 
+const FColor FDataflowEditorModule::SurfaceColor = FLinearColor(0.6, 0.6, 0.6).ToRGBE();
+
 void FDataflowEditorModule::StartupModule()
 {
 	FDataflowEditorStyle::Get();
 	
 	FPropertyEditorModule& PropertyEditorModule = FModuleManager::LoadModuleChecked<FPropertyEditorModule>(TEXT("PropertyEditor"));
-	}
+
+	Dataflow::RenderingCallbacks();
+}
 
 void FDataflowEditorModule::ShutdownModule()
 {	
diff --git a/Engine/Plugins/Experimental/Dataflow/Source/DataflowEditor/Private/Dataflow/DataflowEditorToolkit.cpp b/Engine/Plugins/Experimental/Dataflow/Source/DataflowEditor/Private/Dataflow/DataflowEditorToolkit.cpp
--- a/Engine/Plugins/Experimental/Dataflow/Source/DataflowEditor/Private/Dataflow/DataflowEditorToolkit.cpp
+++ b/Engine/Plugins/Experimental/Dataflow/Source/DataflowEditor/Private/Dataflow/DataflowEditorToolkit.cpp
@@ -19,6 +19,7 @@
 #include "Dataflow/DataflowEditorModeToolkit.h"
 #include "Dataflow/DataflowEditorModule.h"
 #include "Dataflow/DataflowEditorModeUILayer.h"
+#include "Dataflow/DataflowEditorUtil.h"
 #include "Dataflow/DataflowConstructionViewport.h"
 #include "Dataflow/DataflowEdNode.h"
 #include "Dataflow/DataflowGraphEditor.h"
@@ -26,7 +27,7 @@
 #include "Dataflow/DataflowObject.h"
 #include "Dataflow/DataflowObjectInterface.h"
 #include "Dataflow/DataflowEditorPreviewSceneBase.h"
-#include "Dataflow/DataflowConstructionScene.h"
+#include "Dataflow/DataflowRenderingFactory.h"
 #include "Dataflow/DataflowSimulationScene.h"
 #include "Dataflow/DataflowSchema.h"
 #include "Dataflow/DataflowSkeletonView.h"
@@ -721,6 +722,10 @@ void FDataflowEditorToolkit::OnNodeTitleCommitted(const FText& InNewText, ETextC
 
 void FDataflowEditorToolkit::OnNodeSelectionChanged(const TSet<UObject*>& InNewSelection)
 {
+	//
+	// Local helper lambdas
+	//
+
 	auto FindDataflowNodesInSet = [](const TSet<TObjectPtr<UObject>>& InSet) {
 		TSet<TObjectPtr<UObject>> Results;
 		for (UObject* Item : InSet)
@@ -798,8 +803,10 @@ void FDataflowEditorToolkit::OnNodeSelectionChanged(const TSet<UObject*>& InNewS
 		}
 	};
 
-
-
+	//
+	// Actual function
+	// 
+	
 	// Despite this function's name, we might not have actually changed which node is selected
 	bool bPrimarySelectionChanged = false;
 
@@ -924,43 +931,37 @@ void FDataflowEditorToolkit::OnNodeSelectionChanged(const TSet<UObject*>& InNewS
 	{
 		bool bFoundViewMode = true;
 
-		if (const TObjectPtr<UDataflowBaseContent>& EditorContent = GetEditorContent())
+		if (PrimarySelection && GetEditorContent())
 		{
-			if (const TSharedPtr<Dataflow::FEngineContext> Context = EditorContent->GetDataflowContext())
+			if (!Dataflow::CanRenderNodeOutput(*PrimarySelection, *GetEditorContent(), *DataflowMode->GetConstructionViewMode()))
 			{
-				if (PrimarySelection)
+				// Selected node can't render with the current view mode. Check through available view modes and see if it can render with any of them
+
+				bFoundViewMode = false;
+
+				TArray<Dataflow::FRenderingParameter> RenderingParameters = PrimarySelection->GetRenderParameters();
+				for (const Dataflow::FRenderingParameter& Param : RenderingParameters)
 				{
-					if (!PrimarySelection->CanRender(Context.ToSharedRef(), *DataflowMode->GetConstructionViewMode()))
+					const FName NodeOutputTypeName = Param.Type;
+
+					for (const TPair<FName, TUniquePtr<Dataflow::IDataflowConstructionViewMode>>& ViewMode : Dataflow::FRenderingViewModeFactory::GetInstance().GetViewModes())
 					{
-						// Selected node can't render with the current view mode. Check through available view modes and see if it can render with any of them
+						check(ViewMode.Value.IsValid());
 
-						bFoundViewMode = false;
+						const bool bCanRender = Dataflow::CanRenderNodeOutput(*PrimarySelection, *GetEditorContent(), *ViewMode.Value);
 
-						TArray<Dataflow::FRenderingParameter> RenderingParameters = PrimarySelection->GetRenderParameters();
-						for (const Dataflow::FRenderingParameter& Param : RenderingParameters)
+						if (bCanRender)
 						{
-							const FName NodeOutputTypeName = Param.Type;
-
-							for (const TPair<FName, TUniquePtr<Dataflow::IDataflowConstructionViewMode>>& ViewMode : Dataflow::FRenderingViewModeFactory::GetInstance().GetViewModes())
-							{
-								check(ViewMode.Value.IsValid());
-
-								const bool bCanRender = PrimarySelection->CanRender(Context.ToSharedRef(), *ViewMode.Value);
-
-								if (bCanRender)
-								{
-									DataflowMode->SetConstructionViewMode(ViewMode.Key);
-									bFoundViewMode = true;
-									break;
-								}
-							}
-
-							if (bFoundViewMode)
-							{
-								break;
-							}
+							DataflowMode->SetConstructionViewMode(ViewMode.Key);
+							bFoundViewMode = true;
+							break;
 						}
 					}
+
+					if (bFoundViewMode)
+					{
+						break;
+					}
 				}
 			}
 		}
diff --git a/Engine/Plugins/Experimental/Dataflow/Source/DataflowEditor/Private/Dataflow/DataflowEditorUtil.cpp b/Engine/Plugins/Experimental/Dataflow/Source/DataflowEditor/Private/Dataflow/DataflowEditorUtil.cpp
--- a/Engine/Plugins/Experimental/Dataflow/Source/DataflowEditor/Private/Dataflow/DataflowEditorUtil.cpp
+++ b/Engine/Plugins/Experimental/Dataflow/Source/DataflowEditor/Private/Dataflow/DataflowEditorUtil.cpp
@@ -8,6 +8,7 @@
 #include "Dataflow/DataflowEditorToolkit.h"
 #include "Dataflow/DataflowEdNode.h"
 #include "Dataflow/DataflowObject.h"
+#include "Dataflow/DataflowRenderingFactory.h"
 #include "DynamicMesh/MeshNormals.h"
 #include "Engine/SkeletalMesh.h"
 #include "Animation/Skeleton.h"
@@ -130,4 +131,51 @@ namespace Dataflow
 		ensure(false);
 		return MakeShared<FEngineContext>(nullptr, nullptr, FTimestamp::Invalid);
 	}
+
+	bool CanRenderNodeOutput(const UDataflowEdNode& EdNode, const UDataflowBaseContent& EditorContent, const IDataflowConstructionViewMode& ViewMode)
+	{
+		if (const TSharedPtr<FEngineContext> Context = EditorContent.GetDataflowContext())
+		{
+			if (TSharedPtr<const FDataflowNode> NodeTarget = EdNode.GetDataflowGraph()->FindBaseNode(FName(EdNode.GetName())))
+			{
+				if (const FRenderingFactory* const Factory = FRenderingFactory::GetInstance())
+				{
+					for (const FRenderingParameter& Parameter : EdNode.GetRenderParameters())
+					{
+						if (Factory->CanRenderNodeOutput(FGraphRenderingState{ EdNode.GetDataflowNodeGuid(), NodeTarget.Get(), Parameter, *Context.Get(), ViewMode }))
+						{
+							return true;
+						}
+					}
+				}
+			}
+		}
+
+		return false;
+	}
+
+	void RenderNodeOutput(GeometryCollection::Facades::FRenderingFacade& Facade, const UDataflowEdNode& Node, const UDataflowBaseContent& EditorContent)
+	{
+		const TObjectPtr<UDataflow>& DataflowAsset = EditorContent.GetDataflowAsset();
+		const IDataflowConstructionViewMode* ConstructionViewMode = EditorContent.GetConstructionViewMode();
+		const TSharedPtr<FEngineContext>& DataflowContext = EditorContent.GetDataflowContext();
+
+		if (DataflowAsset && DataflowContext && ConstructionViewMode)
+		{
+			if (FRenderingFactory* const Factory = FRenderingFactory::GetInstance())
+			{
+				for (const FRenderingParameter& Parameter : Node.GetRenderParameters())
+				{
+					if (const TSharedPtr<FGraph> Graph = DataflowAsset->GetDataflow())
+					{
+						if (TSharedPtr<const FDataflowNode> NodeTarget = Graph->FindBaseNode(FName(Node.GetName())))
+						{
+							Factory->RenderNodeOutput(Facade, FGraphRenderingState{ Node.GetDataflowNodeGuid(), NodeTarget.Get(), Parameter, *DataflowContext, *ConstructionViewMode });
+						}
+					}
+				}
+			}
+		}
+	}
+
 }
\ No newline at end of file
diff --git a/Engine/Plugins/Experimental/Dataflow/Source/DataflowEditor/Private/Dataflow/DataflowEditorUtil.h b/Engine/Plugins/Experimental/Dataflow/Source/DataflowEditor/Private/Dataflow/DataflowEditorUtil.h
--- a/Engine/Plugins/Experimental/Dataflow/Source/DataflowEditor/Private/Dataflow/DataflowEditorUtil.h
+++ b/Engine/Plugins/Experimental/Dataflow/Source/DataflowEditor/Private/Dataflow/DataflowEditorUtil.h
@@ -16,6 +16,8 @@ class USkeletalMesh;
 class USkeleton;
 class UAnimationAsset;
 class UMaterial;
+class UDataflowEdNode;
+namespace Dataflow { class IDataflowConstructionViewMode; }
 
 namespace UE::Dataflow::Private
 {
@@ -45,4 +47,7 @@ namespace UE
 namespace Dataflow
 {
 	TSharedPtr<::Dataflow::FEngineContext> GetContext(TObjectPtr<UDataflowBaseContent> Content);
+
+	bool CanRenderNodeOutput(const UDataflowEdNode& EdNode, const UDataflowBaseContent& EditorContent, const ::Dataflow::IDataflowConstructionViewMode& ViewMode);
+	void RenderNodeOutput(GeometryCollection::Facades::FRenderingFacade& Facade, const UDataflowEdNode& Node, const UDataflowBaseContent& EditorContent);
 }
\ No newline at end of file
diff --git a/Engine/Plugins/Experimental/Dataflow/Source/DataflowEditor/Private/Dataflow/DataflowEngineRendering.cpp b/Engine/Plugins/Experimental/Dataflow/Source/DataflowEditor/Private/Dataflow/DataflowEngineRendering.cpp
new file mode 100644
--- /dev/null
+++ b/Engine/Plugins/Experimental/Dataflow/Source/DataflowEditor/Private/Dataflow/DataflowEngineRendering.cpp
@@ -0,0 +1,606 @@
+// Copyright Epic Games, Inc. All Rights Reserved.
+
+#include "Dataflow/DataflowEngineRendering.h"
+
+#include "Dataflow/DataflowConnectionTypes.h"
+#include "Dataflow/DataflowEditorModule.h"
+#include "Dataflow/DataflowRenderingFactory.h"
+#include "DynamicMesh/DynamicMesh3.h"
+#include "Field/FieldSystemTypes.h"
+#include "GeometryCollection/GeometryCollectionUtility.h"
+#include "GeometryCollection/Facades/CollectionRenderingFacade.h"
+#include "GeometryCollection/Facades/CollectionExplodedVectorFacade.h"
+#include "GeometryCollection/ManagedArrayCollection.h"
+#include "GeometryCollection/GeometryCollectionAlgo.h"
+#include "UDynamicMesh.h"
+
+namespace Dataflow
+{
+	void RenderBasicGeometryCollection(GeometryCollection::Facades::FRenderingFacade& RenderCollection, const Dataflow::FGraphRenderingState& State, TArray<FLinearColor>* VertexColorOverride = nullptr)
+	{
+		FManagedArrayCollection Default;
+		FName PrimaryOutput = State.GetRenderOutputs()[0]; // "Collection"
+		const FManagedArrayCollection& Collection = State.GetValue<FManagedArrayCollection>(PrimaryOutput, Default);
+
+		const TManagedArray<int32>& BoneIndex = Collection.GetAttribute<int32>("BoneMap", FGeometryCollection::VerticesGroup);
+		const TManagedArray<int32>& Parents = Collection.GetAttribute<int32>(FTransformCollection::ParentAttribute, FTransformCollection::TransformGroup);
+		const TManagedArray<FTransform3f>& Transforms = Collection.GetAttribute<FTransform3f>(FTransformCollection::TransformAttribute, FTransformCollection::TransformGroup);
+
+		TArray<FTransform> M;
+		GeometryCollectionAlgo::GlobalMatrices(Transforms, Parents, M);
+
+		// If Collection has "ExplodedVector" attribute then use it to modify the global matrices (ExplodedView node creates it)
+		GeometryCollection::Facades::FCollectionExplodedVectorFacade ExplodedViewFacade(Collection);
+		ExplodedViewFacade.UpdateGlobalMatricesWithExplodedVectors(M);
+
+		auto ToD = [](FVector3f V) { return FVector3d(V.X, V.Y, V.Z); };
+		auto ToF = [](FVector3d V) { return FVector3f(V.X, V.Y, V.Z); };
+
+
+		const TManagedArray<FVector3f>& Vertex = Collection.GetAttribute<FVector3f>("Vertex", FGeometryCollection::VerticesGroup);
+		const TManagedArray<FIntVector>& Faces = Collection.GetAttribute<FIntVector>("Indices", FGeometryCollection::FacesGroup);
+		const TManagedArray<bool>* FaceVisible = Collection.FindAttribute<bool>("Visible", FGeometryCollection::FacesGroup);
+
+		TArray<FVector3f> Vertices; Vertices.AddUninitialized(Vertex.Num());
+		TArray<FIntVector> Tris; Tris.AddUninitialized(Faces.Num());
+		TArray<bool> Visited; Visited.Init(false, Vertices.Num());
+
+		int32 Tdx = 0;
+		for (int32 FaceIdx = 0; FaceIdx < Faces.Num(); ++FaceIdx)
+		{
+			if (FaceVisible && !(*FaceVisible)[FaceIdx]) continue;
+
+			const FIntVector& Face = Faces[FaceIdx];
+
+			FIntVector Tri = FIntVector(Face[0], Face[1], Face[2]);
+			FTransform Ms[3] = { M[BoneIndex[Tri[0]]], M[BoneIndex[Tri[1]]], M[BoneIndex[Tri[2]]] };
+
+			Tris[Tdx++] = Tri;
+			if (!Visited[Tri[0]]) Vertices[Tri[0]] = ToF(Ms[0].TransformPosition(ToD(Vertex[Tri[0]])));
+			if (!Visited[Tri[1]]) Vertices[Tri[1]] = ToF(Ms[1].TransformPosition(ToD(Vertex[Tri[1]])));
+			if (!Visited[Tri[2]]) Vertices[Tri[2]] = ToF(Ms[2].TransformPosition(ToD(Vertex[Tri[2]])));
+
+			Visited[Tri[0]] = true; Visited[Tri[1]] = true; Visited[Tri[2]] = true;
+		}
+
+		Tris.SetNum(Tdx);
+
+		// Maybe these buffers should be shrunk, but there are unused vertices in the buffer. 
+		for (int i = 0; i < Visited.Num(); i++) if (!Visited[i]) Vertices[i] = FVector3f(0);
+
+		// Copy VertexNormals from the Collection if exists otherwise compute and set it
+		TArray<FVector3f> VertexNormals; VertexNormals.AddUninitialized(Vertex.Num());
+		if (const TManagedArray<FVector3f>* VertexNormal = Collection.FindAttribute<FVector3f>("Normal", FGeometryCollection::VerticesGroup))
+		{
+			for (int32 VertexIdx = 0; VertexIdx < VertexNormals.Num(); ++VertexIdx)
+			{
+				VertexNormals[VertexIdx] = (*VertexNormal)[VertexIdx];
+			}
+		}
+		else
+		{
+			for (int32 VertexIdx = 0; VertexIdx < VertexNormals.Num(); ++VertexIdx)
+			{
+				// TODO: Compute the normal
+				VertexNormals[VertexIdx] = FVector3f(0.f);
+			}
+		}
+
+		// Copy VertexColors from the Collection if exists otherwise set it to FDataflowEditorModule::SurfaceColor
+		TArray<FLinearColor> VertexColors; VertexColors.AddUninitialized(Vertex.Num());
+		if (VertexColorOverride && VertexColorOverride->Num() == Vertex.Num())
+		{
+			for (int32 VertexIdx = 0; VertexIdx < VertexColors.Num(); ++VertexIdx)
+			{
+				VertexColors[VertexIdx] = (*VertexColorOverride)[VertexIdx];
+			}
+		}
+		else
+		{
+			if (const TManagedArray<FLinearColor>* VertexColorManagedArray = Collection.FindAttribute<FLinearColor>("Color", FGeometryCollection::VerticesGroup))
+			{
+				for (int32 VertexIdx = 0; VertexIdx < VertexColors.Num(); ++VertexIdx)
+				{
+					VertexColors[VertexIdx] = (*VertexColorManagedArray)[VertexIdx];
+				}
+			}
+			else
+			{
+				for (int32 VertexIdx = 0; VertexIdx < VertexColors.Num(); ++VertexIdx)
+				{
+					VertexColors[VertexIdx] = FLinearColor(FDataflowEditorModule::SurfaceColor);
+				}
+			}
+		}
+
+		// Set the data on the RenderCollection
+		int32 GeometryIndex = RenderCollection.StartGeometryGroup(State.GetGuid().ToString());
+		RenderCollection.AddSurface(MoveTemp(Vertices), MoveTemp(Tris), MoveTemp(VertexNormals), MoveTemp(VertexColors));
+		RenderCollection.EndGeometryGroup(GeometryIndex);
+
+	}
+
+	void RenderMeshIndexedGeometryCollection(GeometryCollection::Facades::FRenderingFacade& RenderCollection, const Dataflow::FGraphRenderingState& State, TArray<FLinearColor>* VertexColorOverride = nullptr )
+	{
+		auto ToD = [](FVector3f V) { return FVector3d(V.X, V.Y, V.Z); };
+		auto ToF = [](FVector3d V) { return FVector3f(V.X, V.Y, V.Z); };
+
+		FManagedArrayCollection Default;
+		FName PrimaryOutput = State.GetRenderOutputs()[0]; // "Collection"
+		const FManagedArrayCollection& Collection = State.GetValue<FManagedArrayCollection>(PrimaryOutput, Default);
+
+		const TManagedArray<int32>& BoneIndex = Collection.GetAttribute<int32>("BoneMap", FGeometryCollection::VerticesGroup);
+		const TManagedArray<int32>& Parents = Collection.GetAttribute<int32>(FTransformCollection::ParentAttribute, FTransformCollection::TransformGroup);
+		const TManagedArray<FTransform3f>& Transforms = Collection.GetAttribute<FTransform3f>(FTransformCollection::TransformAttribute, FTransformCollection::TransformGroup);
+		const TManagedArray<FString>& BoneNames = Collection.GetAttribute<FString>("BoneName", FGeometryCollection::TransformGroup);
+		const TManagedArray<FVector3f>& Vertex = Collection.GetAttribute<FVector3f>("Vertex", FGeometryCollection::VerticesGroup);
+		const TManagedArray<FIntVector>& Faces = Collection.GetAttribute<FIntVector>("Indices", FGeometryCollection::FacesGroup);
+		const TManagedArray<bool>* FaceVisible = Collection.FindAttribute<bool>("Visible", FGeometryCollection::FacesGroup);
+
+		const TManagedArray<int32>& VertexStart = Collection.GetAttribute<int32>("VertexStart", FGeometryCollection::GeometryGroup);
+		const TManagedArray<int32>& VertexCount = Collection.GetAttribute<int32>("VertexCount", FGeometryCollection::GeometryGroup);
+		const TManagedArray<int32>& FacesStart = Collection.GetAttribute<int32>("FaceStart", FGeometryCollection::GeometryGroup);
+		const TManagedArray<int32>& FacesCount = Collection.GetAttribute<int32>("FaceCount", FGeometryCollection::GeometryGroup);
+		int32 TotalVertices = Collection.NumElements(FGeometryCollection::VerticesGroup);
+
+		TArray<FTransform> M;
+		GeometryCollectionAlgo::GlobalMatrices(Transforms, Parents, M);
+		GeometryCollection::Facades::FCollectionExplodedVectorFacade ExplodedViewFacade(Collection);
+		ExplodedViewFacade.UpdateGlobalMatricesWithExplodedVectors(M);
+
+		for (int Gdx = 0; Gdx < Collection.NumElements(FGeometryCollection::GeometryGroup); Gdx++)
+		{
+			TArray<FVector3f> Vertices; Vertices.AddUninitialized(VertexCount[Gdx]);
+			TArray<FIntVector> Tris; Tris.AddUninitialized(FacesCount[Gdx]);
+			TArray<bool> Visited; Visited.Init(false, VertexCount[Gdx]);
+
+			int32 Tdx = 0;
+			int32 LastFaceIndex = FacesStart[Gdx] + FacesCount[Gdx];
+			for (int32 FaceIdx = FacesStart[Gdx]; FaceIdx < LastFaceIndex; ++FaceIdx)
+			{
+				if (FaceVisible && !(*FaceVisible)[FaceIdx]) continue;
+
+				const FIntVector& Face = Faces[FaceIdx];
+
+				FIntVector Tri = FIntVector(Face[0], Face[1], Face[2]);
+				FTransform Ms[3] = { M[BoneIndex[Tri[0]]], M[BoneIndex[Tri[1]]], M[BoneIndex[Tri[2]]] };
+				FIntVector MovedTri = FIntVector(Face[0] - VertexStart[Gdx], Face[1] - VertexStart[Gdx], Face[2] - VertexStart[Gdx]);
+
+				Tris[Tdx++] = MovedTri;
+				if (!Visited[MovedTri[0]]) Vertices[Tri[0] - VertexStart[Gdx]] = ToF(Ms[0].TransformPosition(ToD(Vertex[Tri[0]])));
+				if (!Visited[MovedTri[1]]) Vertices[Tri[1] - VertexStart[Gdx]] = ToF(Ms[1].TransformPosition(ToD(Vertex[Tri[1]])));
+				if (!Visited[MovedTri[2]]) Vertices[Tri[2] - VertexStart[Gdx]] = ToF(Ms[2].TransformPosition(ToD(Vertex[Tri[2]])));
+
+				Visited[MovedTri[0]] = true; Visited[MovedTri[1]] = true; Visited[MovedTri[2]] = true;
+			}
+
+			Tris.SetNum(Tdx);
+
+			// move the unused points too. Need to keep them for vertex alignment with ediior tools. 
+			for (int i = 0; i < Visited.Num(); i++)
+			{
+				if (!Visited[i])
+				{
+					Vertices[i] = ToF(M[BoneIndex[i + VertexStart[Gdx]]].TransformPosition(ToD(Vertex[i + VertexStart[Gdx]])));
+				}
+			}
+
+			// Copy VertexNormals from the Collection if exists otherwise compute and set it
+			ensure(VertexCount[Gdx] == Vertices.Num());
+			TArray<FVector3f> VertexNormals; VertexNormals.AddUninitialized(Vertices.Num());
+			if (const TManagedArray<FVector3f>* VertexNormal = Collection.FindAttribute<FVector3f>("Normal", FGeometryCollection::VerticesGroup))
+			{
+				int32 LastVertIndex = VertexStart[Gdx] + VertexCount[Gdx];
+				for (int32 VertexIdx = VertexStart[Gdx], SrcVertexIdx = 0; VertexIdx < LastVertIndex; ++VertexIdx, ++SrcVertexIdx)
+				{
+					VertexNormals[SrcVertexIdx] = (*VertexNormal)[VertexIdx];
+				}
+			}
+			else
+			{
+				for (int32 VertexIdx = 0; VertexIdx < Vertices.Num(); ++VertexIdx)
+				{
+					// TODO: Compute the normal
+					VertexNormals[VertexIdx] = FVector3f(0.f);
+				}
+			}
+
+			// Copy VertexColors from the Collection if exists otherwise set it to FDataflowEditorModule::SurfaceColor
+			TArray<FLinearColor> VertexColors; VertexColors.AddUninitialized(Vertices.Num());
+			if (VertexColorOverride && VertexColorOverride->Num() == TotalVertices)
+			{
+				int32 LastVertIndex = VertexStart[Gdx] + VertexCount[Gdx];
+				for (int32 VertexIdx = VertexStart[Gdx], SrcVertexIdx = 0; VertexIdx < LastVertIndex; ++VertexIdx, ++SrcVertexIdx)
+				{
+					VertexColors[SrcVertexIdx] = (*VertexColorOverride)[VertexIdx];
+				}
+			}
+			else
+			{
+				if (const TManagedArray<FLinearColor>* VertexColorManagedArray = Collection.FindAttribute<FLinearColor>("Color", FGeometryCollection::VerticesGroup))
+				{
+					int32 LastVertIndex = VertexStart[Gdx] + VertexCount[Gdx];
+					for (int32 VertexIdx = VertexStart[Gdx], SrcVertexIdx = 0; VertexIdx < LastVertIndex; ++VertexIdx, ++SrcVertexIdx)
+					{
+						VertexColors[SrcVertexIdx] = (*VertexColorManagedArray)[VertexIdx];
+					}
+				}
+				else
+				{
+					for (int32 VertexIdx = 0; VertexIdx < VertexColors.Num(); ++VertexIdx)
+					{
+						VertexColors[VertexIdx] = FLinearColor(FDataflowEditorModule::SurfaceColor);
+					}
+				}
+			}
+
+			// Set the data on the RenderCollection
+			if (Vertices.Num() && Tris.Num())
+			{
+				FString GeometryName = State.GetGuid().ToString(); GeometryName.AppendChar('.').AppendInt(Gdx);
+				if (BoneIndex[VertexStart[Gdx]] != INDEX_NONE)
+				{
+					GeometryName = BoneNames[BoneIndex[VertexStart[Gdx]]];
+				}
+				int32 GeometryIndex = RenderCollection.StartGeometryGroup(GeometryName);
+				RenderCollection.AddSurface(MoveTemp(Vertices), MoveTemp(Tris), MoveTemp(VertexNormals), MoveTemp(VertexColors));
+				RenderCollection.EndGeometryGroup(GeometryIndex);
+			}
+		}
+	}
+
+	
+
+	class FGeometryCollectionSurfaceRenderCallbacks : public FRenderingFactory::ICallbackInterface
+	{
+		virtual Dataflow::FRenderKey GetRenderKey() const override
+		{
+			return { "SurfaceRender", FGeometryCollection::StaticType() };
+		}
+
+		virtual bool CanRender(const Dataflow::IDataflowConstructionViewMode& ViewMode) const override
+		{
+			return (ViewMode.GetName() == FDataflowConstruction3DViewMode::Name);
+		}
+
+		virtual void Render(GeometryCollection::Facades::FRenderingFacade& RenderCollection, const FGraphRenderingState& State) override
+		{
+			if (State.GetRenderOutputs().Num())
+			{
+				FManagedArrayCollection Default;
+				FName PrimaryOutput = State.GetRenderOutputs()[0]; // "Collection"
+				const FManagedArrayCollection& Collection = State.GetValue<FManagedArrayCollection>(PrimaryOutput, Default);
+
+				const bool bFoundIndices = Collection.FindAttributeTyped<FIntVector>("Indices", FGeometryCollection::FacesGroup) != nullptr;
+				const bool bFoundVertices = Collection.FindAttributeTyped<FVector3f>("Vertex", FGeometryCollection::VerticesGroup) != nullptr;
+				const bool bFoundTransforms = Collection.FindAttributeTyped<FTransform3f>(FTransformCollection::TransformAttribute, FTransformCollection::TransformGroup) != nullptr;
+				const bool bFoundBoneMap = Collection.FindAttributeTyped<int32>("BoneMap", FGeometryCollection::VerticesGroup) != nullptr;
+				const bool bFoundParents = Collection.FindAttributeTyped<int32>(FTransformCollection::ParentAttribute, FTransformCollection::TransformGroup) != nullptr;
+				UE_LOG(LogTemp, Warning, TEXT("Render GC with found params = %d %d %d %d %d"), bFoundIndices, bFoundVertices, bFoundTransforms, bFoundBoneMap, bFoundParents);
+				bool bFoundRenderData = bFoundIndices && bFoundVertices && bFoundTransforms && bFoundBoneMap && bFoundParents
+					&& Collection.NumElements(FTransformCollection::TransformGroup) > 0;
+
+				const bool bFoundVertexStart = Collection.FindAttributeTyped<int32>("VertexStart", FGeometryCollection::GeometryGroup) != nullptr;
+				const bool bFoundVertexCount = Collection.FindAttributeTyped<int32>("VertexCount", FGeometryCollection::GeometryGroup) != nullptr;
+				const bool bFoundFaceStart = Collection.FindAttributeTyped<int32>("FaceStart", FGeometryCollection::GeometryGroup) != nullptr;
+				const bool bFoundFaceCount = Collection.FindAttributeTyped<int32>("FaceCount", FGeometryCollection::GeometryGroup) != nullptr;
+				UE_LOG(LogTemp, Warning, TEXT("Render GC with found mesh group params = %d %d %d %d"), bFoundVertexStart, bFoundVertexCount, bFoundFaceStart, bFoundFaceCount);
+				bool bFoundGeometryAttributes = bFoundVertexStart && bFoundVertexCount && bFoundFaceStart && bFoundFaceCount
+					&& Collection.NumElements(FGeometryCollection::GeometryGroup) > 0;
+
+				if (bFoundRenderData && bFoundGeometryAttributes)
+				{
+					RenderMeshIndexedGeometryCollection(RenderCollection, State);
+				}
+				else if (bFoundRenderData)
+				{
+					RenderBasicGeometryCollection(RenderCollection, State);
+				}
+			}
+		}
+	};
+
+
+	class FGeometryCollectionSurfaceWeightsRenderCallbacks : public FRenderingFactory::ICallbackInterface
+	{
+		virtual Dataflow::FRenderKey GetRenderKey() const override
+		{
+			return { "SurfaceWeightsRender", FGeometryCollection::StaticType() };
+		}
+
+		virtual bool CanRender(const Dataflow::IDataflowConstructionViewMode& ViewMode) const override
+		{
+			return (ViewMode.GetName() == FDataflowConstruction3DViewMode::Name);
+		}
+
+		virtual void Render(GeometryCollection::Facades::FRenderingFacade& RenderCollection, const FGraphRenderingState& State) override
+		{
+			if (State.GetRenderOutputs().Num() >= 2)
+			{
+				FManagedArrayCollection Default;
+				FName PrimaryOutput = State.GetRenderOutputs()[0]; // "Collection"
+				const FManagedArrayCollection& Collection = State.GetValue<FManagedArrayCollection>(PrimaryOutput, Default);
+
+				const bool bFoundIndices = Collection.FindAttributeTyped<FIntVector>("Indices", FGeometryCollection::FacesGroup) != nullptr;
+				const bool bFoundVertices = Collection.FindAttributeTyped<FVector3f>("Vertex", FGeometryCollection::VerticesGroup) != nullptr;
+				const bool bFoundTransforms = Collection.FindAttributeTyped<FTransform3f>(FTransformCollection::TransformAttribute, FTransformCollection::TransformGroup) != nullptr;
+				const bool bFoundBoneMap = Collection.FindAttributeTyped<int32>("BoneMap", FGeometryCollection::VerticesGroup) != nullptr;
+				const bool bFoundParents = Collection.FindAttributeTyped<int32>(FTransformCollection::ParentAttribute, FTransformCollection::TransformGroup) != nullptr;
+				UE_LOG(LogTemp, Warning, TEXT("Render GC with found params = %d %d %d %d %d"), bFoundIndices, bFoundVertices, bFoundTransforms, bFoundBoneMap, bFoundParents);
+				bool bFoundRenderData = bFoundIndices && bFoundVertices && bFoundTransforms && bFoundBoneMap && bFoundParents
+					&& Collection.NumElements(FTransformCollection::TransformGroup) > 0;
+
+				const bool bFoundVertexStart = Collection.FindAttributeTyped<int32>("VertexStart", FGeometryCollection::GeometryGroup) != nullptr;
+				const bool bFoundVertexCount = Collection.FindAttributeTyped<int32>("VertexCount", FGeometryCollection::GeometryGroup) != nullptr;
+				const bool bFoundFaceStart = Collection.FindAttributeTyped<int32>("FaceStart", FGeometryCollection::GeometryGroup) != nullptr;
+				const bool bFoundFaceCount = Collection.FindAttributeTyped<int32>("FaceCount", FGeometryCollection::GeometryGroup) != nullptr;
+				UE_LOG(LogTemp, Warning, TEXT("Render GC with found mesh group params = %d %d %d %d"), bFoundVertexStart, bFoundVertexCount, bFoundFaceStart, bFoundFaceCount);
+				bool bFoundGeometryAttributes = bFoundVertexStart && bFoundVertexCount && bFoundFaceStart && bFoundFaceCount
+					&& Collection.NumElements(FGeometryCollection::GeometryGroup) > 0;
+
+				FCollectionAttributeKey DefaultKey;
+				FName SecondaryOutput = State.GetRenderOutputs()[1]; // "AttributeKey"
+				const FCollectionAttributeKey& AttributeKey = State.GetValue<FCollectionAttributeKey>(SecondaryOutput, DefaultKey);
+
+				const bool bFoundVertexColor = Collection.FindAttributeTyped<FLinearColor>("Color", FGeometryCollection::VerticesGroup) != nullptr;
+				const bool bFoundFloatScalar = Collection.FindAttributeTyped<float>(FName(AttributeKey.Attribute), FName(AttributeKey.Group)) != nullptr;
+				bool bFoundVertexScalarAndColors = bFoundVertexColor && bFoundFloatScalar && AttributeKey.Group.Equals(FGeometryCollection::VerticesGroup.ToString());
+
+				TArray<FLinearColor>* Colors = nullptr;
+				if (bFoundVertexScalarAndColors)
+				{
+					auto RangeValue = [](const TManagedArray<float>* FloatArray)
+						{
+							float Min = FLT_MAX;
+							float Max = -FLT_MAX;
+							for (int i = 0; i < FloatArray->Num(); i++) {
+								Min = FMath::Min(Min, (*FloatArray)[i]);
+								Max = FMath::Max(Max, (*FloatArray)[i]);
+							}
+							return TPair<float, float>(Min, Max);
+						};
+
+					const TManagedArray<float>* FloatArray = Collection.FindAttributeTyped<float>(FName(AttributeKey.Attribute), FName(AttributeKey.Group));
+					if (FloatArray && FloatArray->Num())
+					{
+						Colors = new TArray<FLinearColor>();
+						Colors->AddUninitialized(FloatArray->Num());
+
+						TPair<float, float> Range = RangeValue(FloatArray);
+						float Delta = FMath::Abs(Range.Get<1>() - Range.Get<0>());
+						if (Delta > FLT_EPSILON)
+						{
+							for (int32 VertexIdx = 0; VertexIdx < FloatArray->Num(); ++VertexIdx)
+							{
+								(*Colors)[VertexIdx] = FLinearColor::White * ((*FloatArray)[VertexIdx] - Range.Get<0>()) / Delta;
+							}
+						}
+						else
+						{
+							for (int32 VertexIdx = 0; VertexIdx < FloatArray->Num(); ++VertexIdx)
+							{
+								(*Colors)[VertexIdx] = FLinearColor::Black;
+							}
+						}
+					}
+				}
+
+				if (bFoundRenderData && bFoundGeometryAttributes)
+				{
+					RenderMeshIndexedGeometryCollection(RenderCollection, State, Colors);
+				}
+				else if (bFoundRenderData)
+				{
+					RenderBasicGeometryCollection(RenderCollection, State, Colors);
+				}
+
+				if (Colors)
+				{
+					delete Colors;
+				}
+			}
+		}
+	};
+
+	class FDynamicMesh3SurfaceRenderCallbacks : public FRenderingFactory::ICallbackInterface
+	{
+		virtual Dataflow::FRenderKey GetRenderKey() const override
+		{
+			return { "SurfaceRender", FName("FDynamicMesh3") };
+		}
+
+		virtual bool CanRender(const Dataflow::IDataflowConstructionViewMode& ViewMode) const override
+		{
+			return (ViewMode.GetName() == FDataflowConstruction3DViewMode::Name);
+		}
+
+		virtual void Render(GeometryCollection::Facades::FRenderingFacade& RenderCollection, const FGraphRenderingState& State) override
+		{
+			if (State.GetRenderOutputs().Num())
+			{
+				FName PrimaryOutput = State.GetRenderOutputs()[0]; // "Mesh"
+
+				TObjectPtr<UDynamicMesh> Default;
+				if (const TObjectPtr<UDynamicMesh> Mesh = State.GetValue<TObjectPtr<UDynamicMesh>>(PrimaryOutput, Default))
+				{
+					const UE::Geometry::FDynamicMesh3& DynamicMesh = Mesh->GetMeshRef();
+
+					const int32 NumVertices = DynamicMesh.VertexCount();
+					const int32 NumTriangles = DynamicMesh.TriangleCount();
+
+					if (NumVertices > 0 && NumTriangles > 0)
+					{
+						// This will contain the valid triangles only
+						TArray<FIntVector> Tris; Tris.Reserve(DynamicMesh.TriangleCount());
+
+						// DynamicMesh.TrianglesItr() returns the valid triangles only
+						for (UE::Geometry::FIndex3i Tri : DynamicMesh.TrianglesItr())
+						{
+							Tris.Add(FIntVector(Tri.A, Tri.B, Tri.C));
+						}
+
+						// This will contain all the vertices (invalid ones too)
+						// Otherwise the IDs need to be remaped
+						TArray<FVector3f> Vertices; Vertices.AddZeroed(DynamicMesh.MaxVertexID());
+
+						// DynamicMesh.VertexIndicesItr() returns the valid vertices only
+						for (int32 VertexID : DynamicMesh.VertexIndicesItr())
+						{
+							Vertices[VertexID] = (FVector3f)DynamicMesh.GetVertex(VertexID);
+						}
+
+						// Add VertexNormal and VertexColor
+						TArray<FVector3f> VertexNormals; VertexNormals.AddUninitialized(Vertices.Num());
+						TArray<FLinearColor> VertexColors; VertexColors.AddUninitialized(Vertices.Num());
+						for (int32 VertexIdx = 0; VertexIdx < VertexNormals.Num(); ++VertexIdx)
+						{
+							// TODO: Get the normal from FDynamicMesh3
+							VertexNormals[VertexIdx] = FVector3f(0.f);
+							VertexColors[VertexIdx] = FLinearColor(FDataflowEditorModule::SurfaceColor);
+						}
+
+						int32 GeometryIndex = RenderCollection.StartGeometryGroup(State.GetGuid().ToString());
+						RenderCollection.AddSurface(MoveTemp(Vertices), MoveTemp(Tris), MoveTemp(VertexNormals), MoveTemp(VertexColors));
+						RenderCollection.EndGeometryGroup(GeometryIndex);
+					}
+				}
+			}
+		}
+	};
+
+
+	class FBoxSurfaceRenderCallbacks : public FRenderingFactory::ICallbackInterface
+	{
+		virtual Dataflow::FRenderKey GetRenderKey() const override
+		{
+			return { "SurfaceRender", FName("FBox") };
+		}
+
+		virtual bool CanRender(const Dataflow::IDataflowConstructionViewMode& ViewMode) const override
+		{
+			return (ViewMode.GetName() == FDataflowConstruction3DViewMode::Name);
+		}
+
+		virtual void Render(GeometryCollection::Facades::FRenderingFacade& RenderCollection, const FGraphRenderingState& State) override
+		{
+			if (State.GetRenderOutputs().Num())
+			{
+				FName PrimaryOutput = State.GetRenderOutputs()[0]; // "Box"
+
+				FBox Default(ForceInit);
+				const FBox& Box = State.GetValue<FBox>(PrimaryOutput, Default);
+
+				const int32 NumVertices = 8;
+				const int32 NumTriangles = 12;
+
+				TArray<FVector3f> Vertices; Vertices.AddUninitialized(NumVertices);
+				TArray<FIntVector> Tris; Tris.AddUninitialized(NumTriangles);
+
+				FVector Min = Box.Min;
+				FVector Max = Box.Max;
+
+				// Add vertices
+				Vertices[0] = FVector3f(Min);
+				Vertices[1] = FVector3f(Max.X, Min.Y, Min.Z);
+				Vertices[2] = FVector3f(Max.X, Max.Y, Min.Z);
+				Vertices[3] = FVector3f(Min.X, Max.Y, Min.Z);
+				Vertices[4] = FVector3f(Min.X, Min.Y, Max.Z);
+				Vertices[5] = FVector3f(Max.X, Min.Y, Max.Z);
+				Vertices[6] = FVector3f(Max);
+				Vertices[7] = FVector3f(Min.X, Max.Y, Max.Z);
+
+				// Add triangles
+				Tris[0] = FIntVector(0, 1, 3); Tris[1] = FIntVector(1, 2, 3);
+				Tris[2] = FIntVector(0, 4, 1); Tris[3] = FIntVector(4, 5, 1);
+				Tris[4] = FIntVector(5, 2, 1); Tris[5] = FIntVector(5, 6, 2);
+				Tris[6] = FIntVector(3, 2, 6); Tris[7] = FIntVector(7, 3, 6);
+				Tris[8] = FIntVector(0, 3, 7); Tris[9] = FIntVector(4, 0, 7);
+				Tris[10] = FIntVector(5, 4, 7); Tris[11] = FIntVector(5, 7, 6);
+
+				TArray<FVector3f> VertexNormals; VertexNormals.AddUninitialized(NumVertices);
+				// TODO: Compute vertex normals
+
+				// Add VertexNormal and VertexColor
+				TArray<FLinearColor> VertexColors; VertexColors.AddUninitialized(Vertices.Num());
+				for (int32 VertexIdx = 0; VertexIdx < VertexNormals.Num(); ++VertexIdx)
+				{
+					VertexNormals[VertexIdx] = FVector3f(0.f);
+					VertexColors[VertexIdx] = FLinearColor(FDataflowEditorModule::SurfaceColor);
+				}
+
+				int32 GeometryIndex = RenderCollection.StartGeometryGroup(State.GetGuid().ToString());
+				RenderCollection.AddSurface(MoveTemp(Vertices), MoveTemp(Tris), MoveTemp(VertexNormals), MoveTemp(VertexColors));
+				RenderCollection.EndGeometryGroup(GeometryIndex);
+			}
+		}
+	};
+
+
+	class FFieldVolumeRenderCallbacks : public FRenderingFactory::ICallbackInterface
+	{
+		virtual Dataflow::FRenderKey GetRenderKey() const override
+		{
+			return { "VolumeRender", FFieldCollection::StaticType() };
+		}
+
+		virtual bool CanRender(const Dataflow::IDataflowConstructionViewMode& ViewMode) const override
+		{
+			return (ViewMode.GetName() == FDataflowConstruction3DViewMode::Name);
+		}
+
+		virtual void Render(GeometryCollection::Facades::FRenderingFacade& RenderCollection, const FGraphRenderingState& State) override
+		{
+			if (State.GetRenderOutputs().Num())
+			{
+				FName PrimaryOutput = State.GetRenderOutputs()[0]; // "VectorField"
+				if (PrimaryOutput.IsEqual(FName("VectorField")))
+				{
+					FFieldCollection Default;
+					const FFieldCollection& Collection = State.GetValue<FFieldCollection>(PrimaryOutput, Default);
+					TArray<TPair<FVector3f, FVector3f>> VectorField = Collection.GetVectorField();
+					const int32 NumVertices = 3 * VectorField.Num();
+					const int32 NumTriangles = VectorField.Num();
+
+					TArray<FVector3f> Vertices; Vertices.AddUninitialized(NumVertices);
+					TArray<FIntVector> Tris; Tris.AddUninitialized(NumTriangles);
+					TArray<FVector3f> VertexNormals; VertexNormals.AddUninitialized(NumVertices);
+					TArray<FLinearColor> VertexColors; VertexColors.AddUninitialized(NumVertices);
+
+					for (int32 i = 0; i < VectorField.Num(); i++)
+					{
+
+						FVector3f Dir = VectorField[i].Value - VectorField[i].Key;
+						FVector3f DirAdd = Dir;
+						DirAdd.X += 1.f;
+						FVector3f OrthogonalDir = (Dir ^ DirAdd).GetSafeNormal();
+						Tris[i] = FIntVector(3 * i, 3 * i + 1, 3 * i + 2);
+						Vertices[3 * i] = VectorField[i].Key;
+						Vertices[3 * i + 1] = VectorField[i].Value;
+						Vertices[3 * i + 2] = VectorField[i].Key + float(0.1) * Dir.Size() * OrthogonalDir;
+						FVector3f TriangleNormal = (OrthogonalDir ^ Dir).GetSafeNormal();
+						VertexNormals[3 * i] = TriangleNormal;
+						VertexNormals[3 * i + 1] = TriangleNormal;
+						VertexNormals[3 * i + 2] = TriangleNormal;
+						VertexColors[i] = FLinearColor(FDataflowEditorModule::SurfaceColor);
+					}
+
+					int32 GeometryIndex = RenderCollection.StartGeometryGroup(State.GetGuid().ToString());
+					RenderCollection.AddSurface(MoveTemp(Vertices), MoveTemp(Tris), MoveTemp(VertexNormals), MoveTemp(VertexColors));
+					RenderCollection.EndGeometryGroup(GeometryIndex);
+				}
+			}
+		}
+	};
+
+	void RenderingCallbacks()
+	{
+		using namespace Dataflow;
+
+		FRenderingFactory::GetInstance()->RegisterCallbacks(MakeUnique<FGeometryCollectionSurfaceRenderCallbacks>());
+		FRenderingFactory::GetInstance()->RegisterCallbacks(MakeUnique<FGeometryCollectionSurfaceWeightsRenderCallbacks>());
+		FRenderingFactory::GetInstance()->RegisterCallbacks(MakeUnique<FDynamicMesh3SurfaceRenderCallbacks>());
+		FRenderingFactory::GetInstance()->RegisterCallbacks(MakeUnique<FBoxSurfaceRenderCallbacks>());
+		FRenderingFactory::GetInstance()->RegisterCallbacks(MakeUnique<FFieldVolumeRenderCallbacks>());
+	}
+
+}
diff --git a/Engine/Plugins/Experimental/Dataflow/Source/DataflowEditor/Private/Dataflow/DataflowRenderingFactory.cpp b/Engine/Plugins/Experimental/Dataflow/Source/DataflowEditor/Private/Dataflow/DataflowRenderingFactory.cpp
new file mode 100644
--- /dev/null
+++ b/Engine/Plugins/Experimental/Dataflow/Source/DataflowEditor/Private/Dataflow/DataflowRenderingFactory.cpp
@@ -0,0 +1,81 @@
+// Copyright Epic Games, Inc. All Rights Reserved.
+
+#include "Dataflow/DataflowRenderingFactory.h"
+
+#include "Dataflow/DataflowNode.h"
+#include "Dataflow/DataflowInputOutput.h"
+#include "Misc/MessageDialog.h"
+
+namespace Dataflow
+{
+	FRenderingFactory* FRenderingFactory::Instance = nullptr;
+
+	void FRenderingFactory::RenderNodeOutput(GeometryCollection::Facades::FRenderingFacade& RenderData, const FGraphRenderingState& State)
+	{
+		if (CallbackMap.Contains(State.GetRenderKey()))
+		{
+			CallbackMap[State.GetRenderKey()]->Render(RenderData, State);
+		}
+		else
+		{
+			UE_LOG(LogChaos, Warning,
+				TEXT("Warning : Dataflow missing output renderer <%s,%s> for node %s"), 
+				*State.GetRenderKey().Get<0>(),
+				*State.GetRenderKey().Get<1>().ToString(),
+				*State.GetNodeName().ToString());
+		}
+	}
+
+	bool FRenderingFactory::CanRenderNodeOutput(const FGraphRenderingState& State) const
+	{
+		if (CallbackMap.Contains(State.GetRenderKey()))
+		{
+			return CallbackMap[State.GetRenderKey()]->CanRender(State.GetViewMode());
+		}
+
+		return false;
+	}
+
+	FRenderingFactory* FRenderingFactory::GetInstance()
+	{
+		if (!Instance)
+		{
+			Instance = new FRenderingFactory();
+		}
+		return Instance;
+	}
+
+	void FRenderingFactory::RegisterCallbacks(TUniquePtr<ICallbackInterface> InCallbacks)
+	{
+		const FRenderKey Key = InCallbacks->GetRenderKey();
+		if (CallbackMap.Contains(Key))
+		{
+			UE_LOG(LogChaos, Warning,
+				TEXT("Warning : Dataflow output rendering callback registration conflicts with existing rendering callback (<%s,%s>)"),
+				*Key.Get<0>(),
+				*Key.Get<1>().ToString());
+		}
+		else
+		{
+			CallbackMap.Add(Key, MoveTemp(InCallbacks));
+		}
+	}
+
+	void FRenderingFactory::DeregisterCallbacks(const FRenderKey& Key)
+	{
+		if (!CallbackMap.Contains(Key))
+		{
+			UE_LOG(LogChaos, Warning,
+				TEXT("Warning : Dataflow output rendering callback deregistration. Rendering callback not registered: (<%s,%s>)"),
+				*Key.Get<0>(),
+				*Key.Get<1>().ToString());
+		}
+		else
+		{
+			CallbackMap.Remove(Key);
+		}
+	}
+
+
+}
+
diff --git a/Engine/Plugins/Experimental/Dataflow/Source/DataflowEditor/Private/Dataflow/DataflowRenderingViewMode.cpp b/Engine/Plugins/Experimental/Dataflow/Source/DataflowEditor/Private/Dataflow/DataflowRenderingViewMode.cpp
new file mode 100644
--- /dev/null
+++ b/Engine/Plugins/Experimental/Dataflow/Source/DataflowEditor/Private/Dataflow/DataflowRenderingViewMode.cpp
@@ -0,0 +1,124 @@
+// Copyright Epic Games, Inc. All Rights Reserved.
+
+#include "Dataflow/DataflowRenderingViewMode.h"
+#include "Misc/LazySingleton.h"
+
+#define LOCTEXT_NAMESPACE "DataflowRenderingViewMode"
+
+namespace Dataflow
+{
+	bool IDataflowConstructionViewMode::IsPerspective() const
+	{
+		return GetViewportType() == ELevelViewportType::LVT_Perspective;
+	}
+
+	ELevelViewportType FDataflowConstruction2DViewModeBase::GetViewportType() const
+	{
+		return ELevelViewportType::LVT_OrthoNegativeXY;
+	}
+
+	ELevelViewportType FDataflowConstruction3DViewModeBase::GetViewportType() const
+	{
+		return ELevelViewportType::LVT_Perspective;
+	}
+
+	FName FDataflowConstruction2DViewMode::Name = FName("2DView");
+
+	FName FDataflowConstruction2DViewMode::GetName() const
+	{
+		return Name;
+	}
+
+	FText FDataflowConstruction2DViewMode::GetButtonText() const
+	{
+		return FText(LOCTEXT("2DViewButtonText", "2DView"));
+	}
+
+	FText FDataflowConstruction2DViewMode::GetTooltipText() const
+	{
+		return FText(LOCTEXT("2DViewTooltipText", "Default 2D View"));
+	}
+	
+	FName FDataflowConstruction3DViewMode::Name = FName("3DView");
+
+	FName FDataflowConstruction3DViewMode::GetName() const
+	{
+		return Name;
+	}
+
+	FText FDataflowConstruction3DViewMode::GetButtonText() const
+	{
+		return FText(LOCTEXT("3DViewButtonText", "3DView"));
+	}
+
+	FText FDataflowConstruction3DViewMode::GetTooltipText() const
+	{
+		return FText(LOCTEXT("3DViewTooltipText", "Default 3D View"));
+	}
+
+	TUniquePtr<FRenderingViewModeFactory> FRenderingViewModeFactory::Instance = nullptr;
+
+	//
+	// Factory
+	//
+
+	FRenderingViewModeFactory::FRenderingViewModeFactory()
+	{
+		// Default modes
+		ViewModeMap.Add(FDataflowConstruction2DViewMode::Name, MakeUnique<FDataflowConstruction2DViewMode>());
+		ViewModeMap.Add(FDataflowConstruction3DViewMode::Name, MakeUnique<FDataflowConstruction3DViewMode>());
+	}
+
+	FRenderingViewModeFactory& FRenderingViewModeFactory::GetInstance()
+	{
+		return TLazySingleton<FRenderingViewModeFactory>::Get();
+	}
+
+	void FRenderingViewModeFactory::TearDown()
+	{
+		return TLazySingleton<FRenderingViewModeFactory>::TearDown();
+	}
+
+	void FRenderingViewModeFactory::RegisterViewMode(TUniquePtr<IDataflowConstructionViewMode>&& ViewMode)
+	{
+		const FName NewViewModeName = ViewMode->GetName();
+		if (ViewModeMap.Contains(NewViewModeName))
+		{
+			UE_LOG(LogChaos, Warning, TEXT("Dataflow rendering view mode registration conflicts with existing view mode : %s"), *NewViewModeName.ToString());
+		}
+		else
+		{
+			ViewModeMap.Add(NewViewModeName, MoveTemp(ViewMode));
+		}
+	}
+
+	void FRenderingViewModeFactory::DeregisterViewMode(const FName& ViewModeName)
+	{
+		if (!ViewModeMap.Contains(ViewModeName))
+		{
+			UE_LOG(LogChaos, Warning, TEXT("Dataflow rendering view mode deregistration -- view mode not registered : %s"), *ViewModeName.ToString());
+		}
+		else
+		{
+			ViewModeMap.Remove(ViewModeName);
+		}
+	}
+
+	const IDataflowConstructionViewMode* FRenderingViewModeFactory::GetViewMode(const FName& ViewModeName) const
+	{
+		if (ViewModeMap.Contains(ViewModeName))
+		{
+			return ViewModeMap[ViewModeName].Get();
+		}
+		return nullptr;
+	}
+
+	const TMap<FName, TUniquePtr<IDataflowConstructionViewMode>>& FRenderingViewModeFactory::GetViewModes() const
+	{
+		return ViewModeMap;
+	}
+
+}
+
+#undef LOCTEXT_NAMESPACE 
+
diff --git a/Engine/Plugins/Experimental/Dataflow/Source/DataflowEditor/Public/Dataflow/DataflowEditorModule.h b/Engine/Plugins/Experimental/Dataflow/Source/DataflowEditor/Public/Dataflow/DataflowEditorModule.h
--- a/Engine/Plugins/Experimental/Dataflow/Source/DataflowEditor/Public/Dataflow/DataflowEditorModule.h
+++ b/Engine/Plugins/Experimental/Dataflow/Source/DataflowEditor/Public/Dataflow/DataflowEditorModule.h
@@ -15,6 +15,8 @@ class DATAFLOWEDITOR_API FDataflowEditorModule : public FBaseCharacterFXEditorMo
 {
 public:
 
+	static const FColor SurfaceColor;
+
 	/** IModuleInterface implementation */
 	virtual void StartupModule();
 	virtual void ShutdownModule();
diff --git a/Engine/Plugins/Experimental/Dataflow/Source/DataflowEditor/Public/Dataflow/DataflowEngineRendering.h b/Engine/Plugins/Experimental/Dataflow/Source/DataflowEditor/Public/Dataflow/DataflowEngineRendering.h
new file mode 100644
--- /dev/null
+++ b/Engine/Plugins/Experimental/Dataflow/Source/DataflowEditor/Public/Dataflow/DataflowEngineRendering.h
@@ -0,0 +1,15 @@
+// Copyright Epic Games, Inc. All Rights Reserved.
+
+#pragma once
+
+#include "CoreMinimal.h"
+
+
+namespace Dataflow
+{
+	/**
+	 * Rendering callbacks for the Dataflow Rendering Factory
+	 */
+	void RenderingCallbacks();
+
+}
\ No newline at end of file
diff --git a/Engine/Plugins/Experimental/Dataflow/Source/DataflowEditor/Public/Dataflow/DataflowRenderingFactory.h b/Engine/Plugins/Experimental/Dataflow/Source/DataflowEditor/Public/Dataflow/DataflowRenderingFactory.h
new file mode 100644
--- /dev/null
+++ b/Engine/Plugins/Experimental/Dataflow/Source/DataflowEditor/Public/Dataflow/DataflowRenderingFactory.h
@@ -0,0 +1,99 @@
+// Copyright Epic Games, Inc. All Rights Reserved.
+
+#pragma once
+
+#include "ChaosLog.h"
+#include "CoreMinimal.h"
+#include "Dataflow/DataflowNodeParameters.h"
+#include "Dataflow/DataflowInputOutput.h"
+#include "Dataflow/DataflowNode.h"
+#include "Dataflow/DataflowRenderingViewMode.h"
+#include "GeometryCollection/Facades/CollectionRenderingFacade.h"
+
+namespace Dataflow
+{
+	class FContext;
+	typedef TPair<FString, FName> FRenderKey;
+
+	struct FGraphRenderingState 
+	{
+		FGraphRenderingState(const FGuid InGuid, const FDataflowNode* InNode, const FRenderingParameter& InParameters, Dataflow::FContext& InContext, const Dataflow::IDataflowConstructionViewMode& ViewMode)
+			: NodeGuid(InGuid)
+			, Node(InNode)
+			, RenderName(InParameters.Name)
+			, RenderType(InParameters.Type)
+			, RenderOutputs(InParameters.Outputs)
+			, Context(InContext)
+			, ViewMode(ViewMode)
+		{}
+
+		const FGuid& GetGuid() const { return NodeGuid; }
+		FName GetNodeName() const { return Node?Node->GetName():FName(); }
+		FRenderKey GetRenderKey() const { return { RenderName,RenderType }; }
+		const TArray<FName>& GetRenderOutputs() const { return RenderOutputs; }
+
+		template<class T>
+		const T& GetValue(FName OutputName, const T& Default) const
+		{
+			if (Node)
+			{
+				if (const FDataflowOutput* Output = Node->FindOutput(OutputName))
+				{
+					return Output->GetValue<T>(Context, Default);
+				}
+			}
+			return Default;
+		}
+
+		const Dataflow::IDataflowConstructionViewMode& GetViewMode() const { return ViewMode; }
+
+	private:
+		const FGuid NodeGuid;
+		const FDataflowNode* Node = nullptr;
+
+		FString RenderName;
+		FName RenderType;
+		TArray<FName> RenderOutputs;
+
+		Dataflow::FContext& Context;
+
+		const Dataflow::IDataflowConstructionViewMode& ViewMode;
+	};
+
+
+	class FRenderingFactory
+	{
+	public:
+
+		class ICallbackInterface
+		{
+		public:
+			virtual ~ICallbackInterface() = default;
+			virtual FRenderKey GetRenderKey() const = 0;
+			virtual bool CanRender(const Dataflow::IDataflowConstructionViewMode& ViewMode) const = 0;
+			virtual void Render(GeometryCollection::Facades::FRenderingFacade& RenderData, const FGraphRenderingState& State) = 0;
+		};
+
+		~FRenderingFactory() { delete Instance; }
+
+		DATAFLOWEDITOR_API static FRenderingFactory* GetInstance();
+
+		DATAFLOWEDITOR_API void RegisterCallbacks(TUniquePtr<ICallbackInterface> InCallbacks);
+		DATAFLOWEDITOR_API void DeregisterCallbacks(const FRenderKey& Key);
+
+		DATAFLOWEDITOR_API bool Contains(const FRenderKey& InKey) const { return CallbackMap.Contains(InKey); }
+		DATAFLOWEDITOR_API void RenderNodeOutput(GeometryCollection::Facades::FRenderingFacade& RenderData, const FGraphRenderingState& State);
+		DATAFLOWEDITOR_API bool CanRenderNodeOutput(const FGraphRenderingState& State) const;
+
+	private:
+
+		FRenderingFactory() {}
+
+		static FRenderingFactory* Instance;
+
+		TMap<FRenderKey, TUniquePtr<ICallbackInterface>> CallbackMap;
+	};
+
+}
+
+
diff --git a/Engine/Plugins/Experimental/Dataflow/Source/DataflowEditor/Public/Dataflow/DataflowRenderingViewMode.h b/Engine/Plugins/Experimental/Dataflow/Source/DataflowEditor/Public/Dataflow/DataflowRenderingViewMode.h
new file mode 100644
--- /dev/null
+++ b/Engine/Plugins/Experimental/Dataflow/Source/DataflowEditor/Public/Dataflow/DataflowRenderingViewMode.h
@@ -0,0 +1,108 @@
+// Copyright Epic Games, Inc. All Rights Reserved.
+
+#pragma once
+
+#include "UObject/NameTypes.h"
+#include "Editor/UnrealEdTypes.h"
+#include "ChaosLog.h"
+
+namespace Dataflow
+{
+	//
+	// Dataflow Construction View Modes allow the user to choose different views on the same node output. The View Mode system serves two main purposes: 
+	// 
+	// 1) In the Dataflow Editor's Construction Viewport, the currently active IDataflowConstructionViewMode will determine which ELevelViewportType should be used for rendering.
+	//    In other words, the the ViewMode will determine whether we have a 3D perspective camera or a 2D orthographic camera.
+	// 2) In the Datflow Rendering Factory (Dataflow::FRenderingFactory), the rendering callbacks that are registered will also be able to return different 
+	//    render data based on what the currently active IDataflowConstructionViewMode in the viewport is.
+	//
+	// For example, suppose you have created a new type FMyMeshType that contains both a 3D Mesh representation and a 2D Mesh representation (for, say, UVs.) If you want to
+	// write custom ViewModes for it, here is the process:
+	// 
+	// 1) Define two new classes, e.g. FMyMeshType2DViewMode (inheriting from FDataflowConstruction2DViewModeBase) and FMyMeshType3DViewMode (inheriting from FDataflowConstruction3DViewModeBase)
+	// 2) In your module initialization, register both view modes with the Dataflow::FRenderingViewModeFactory
+	// 3) Create a render function callback class and register it with Dataflow::FRenderingFactory. 
+	//    - The CanRender function should return true if the view mode is either of your two new custom view modes.
+	//    - The Render function should return different RenderCollections depending on what the current view mode is (which can be determined from the FGraphRenderingState parameter.)
+	//
+
+	class IDataflowConstructionViewMode
+	{
+	public:
+		virtual ~IDataflowConstructionViewMode() = default;
+
+		DATAFLOWEDITOR_API virtual FName GetName() const = 0;
+		DATAFLOWEDITOR_API virtual FText GetButtonText() const = 0;
+		DATAFLOWEDITOR_API virtual FText GetTooltipText() const = 0;
+		DATAFLOWEDITOR_API virtual ELevelViewportType GetViewportType() const = 0;
+		
+		DATAFLOWEDITOR_API bool IsPerspective() const;
+	};
+
+	// Base 2D and 3D types
+
+	class FDataflowConstruction2DViewModeBase : public IDataflowConstructionViewMode
+	{
+	public:
+		DATAFLOWEDITOR_API virtual ELevelViewportType GetViewportType() const override;
+	};
+
+	class FDataflowConstruction3DViewModeBase : public IDataflowConstructionViewMode
+	{
+	public:
+		DATAFLOWEDITOR_API virtual ELevelViewportType GetViewportType() const override;
+	};
+
+	// Concrete default 2D and 3D types
+
+	class FDataflowConstruction2DViewMode : public FDataflowConstruction2DViewModeBase
+	{
+	public:
+		DATAFLOWEDITOR_API static FName Name;
+		virtual ~FDataflowConstruction2DViewMode() = default;
+	private:
+		DATAFLOWEDITOR_API virtual FName GetName() const override;
+		DATAFLOWEDITOR_API virtual FText GetButtonText() const override;
+		DATAFLOWEDITOR_API virtual FText GetTooltipText() const override;
+	};
+
+	class FDataflowConstruction3DViewMode : public FDataflowConstruction3DViewModeBase
+	{
+	public:
+		DATAFLOWEDITOR_API static FName Name;
+		virtual ~FDataflowConstruction3DViewMode() = default;
+	private:
+		DATAFLOWEDITOR_API virtual FName GetName() const override;
+		DATAFLOWEDITOR_API virtual FText GetButtonText() const override;
+		DATAFLOWEDITOR_API virtual FText GetTooltipText() const override;
+	};
+
+	//
+	// ViewMode registry/factory
+	//
+
+	class FRenderingViewModeFactory
+	{
+	public:
+
+		FRenderingViewModeFactory();
+
+		// FLazySingleton
+		static DATAFLOWEDITOR_API FRenderingViewModeFactory& GetInstance();
+		static DATAFLOWEDITOR_API void TearDown();
+
+		DATAFLOWEDITOR_API void RegisterViewMode(TUniquePtr<IDataflowConstructionViewMode>&& ViewMode);
+		DATAFLOWEDITOR_API void DeregisterViewMode(const FName& ViewModeName);
+
+		DATAFLOWEDITOR_API const IDataflowConstructionViewMode* GetViewMode(const FName& ViewModeName) const;
+
+		DATAFLOWEDITOR_API const TMap<FName, TUniquePtr<IDataflowConstructionViewMode>>& GetViewModes() const;
+
+	private:
+
+		DATAFLOWEDITOR_API static TUniquePtr<FRenderingViewModeFactory> Instance;
+
+		TMap<FName, TUniquePtr<IDataflowConstructionViewMode>> ViewModeMap;
+	};
+}
+
diff --git a/Engine/Plugins/Experimental/Dataflow/Source/DataflowEnginePlugin/Private/Dataflow/DataflowComponent.cpp b/Engine/Plugins/Experimental/Dataflow/Source/DataflowEnginePlugin/Private/Dataflow/DataflowComponent.cpp
--- a/Engine/Plugins/Experimental/Dataflow/Source/DataflowEnginePlugin/Private/Dataflow/DataflowComponent.cpp
+++ b/Engine/Plugins/Experimental/Dataflow/Source/DataflowEnginePlugin/Private/Dataflow/DataflowComponent.cpp
@@ -4,7 +4,6 @@
 #include "DataflowEngineSceneProxy.h"
 
 #include "Dataflow/DataflowEdNode.h"
-#include "Dataflow/DataflowRenderingFactory.h"
 #include "GeometryCollection/Facades/CollectionRenderingFacade.h"
 #include "GeometryCollection/Facades/CollectionBoundsFacade.h"
 #include "GeometryCollection/GeometryCollection.h"
@@ -48,68 +47,6 @@ void UDataflowComponent::AddRenderTarget(const UDataflowEdNode* InTarget)
 	Invalidate();
 }
 
-void UDataflowComponent::BuildRenderCollection()
-{
-	RenderCollection = FManagedArrayCollection();
-	GeometryCollection::Facades::FRenderingFacade Facade(RenderCollection);
-	Facade.DefineSchema();
-
-	if (Context && Dataflow)
-	{
-		for (const UDataflowEdNode* Target : RenderTargets)
-		{
-			if (Target && ViewMode)
-			{
-				Target->Render(Facade, Context.ToSharedRef(), *ViewMode);
-			}
-		}
-	}
-	MarkRenderStateDirty();
-}
-
-void UDataflowComponent::TickComponent(float DeltaTime, enum ELevelTick TickType, FActorComponentTickFunction *ThisTickFunction)
-{
-	bool bNeedsSceneProxyUpdate = false;
-
-	Super::TickComponent(DeltaTime, TickType, ThisTickFunction);
-
-	if (bUpdateRender)
-	{
-		SelectionState.Vertices.Empty();
-		RenderCollection = FManagedArrayCollection();
-		GeometryCollection::Facades::FRenderingFacade Facade(RenderCollection);
-		Facade.DefineSchema();
-
-		if (Context && Dataflow)
-		{
-			for (const UDataflowEdNode* Target : RenderTargets)
-			{
-				if (Target && ViewMode)
-				{
-					Target->Render(Facade, Context.ToSharedRef(), *ViewMode);
-				}
-			}
-		}
-
-		bBoundsNeedsUpdate = true;
-		UpdateLocalBounds();
-
-		bUpdateRender = false;
-		bNeedsSceneProxyUpdate = true;
-	}
-
-	if (bUpdateSelection)
-	{
-		SelectionState.UpdateSelection(this);
-		bUpdateSelection = false;
-		bNeedsSceneProxyUpdate = true;
-	}
-
-	if (bNeedsSceneProxyUpdate)
-	{
-		MarkRenderStateDirty();
-	}
-}
 
 void UDataflowComponent::UpdateLocalBounds()
 {
diff --git a/Engine/Plugins/Experimental/Dataflow/Source/DataflowEnginePlugin/Private/Dataflow/DataflowEnginePlugin.cpp b/Engine/Plugins/Experimental/Dataflow/Source/DataflowEnginePlugin/Private/Dataflow/DataflowEnginePlugin.cpp
--- a/Engine/Plugins/Experimental/Dataflow/Source/DataflowEnginePlugin/Private/Dataflow/DataflowEnginePlugin.cpp
+++ b/Engine/Plugins/Experimental/Dataflow/Source/DataflowEnginePlugin/Private/Dataflow/DataflowEnginePlugin.cpp
@@ -5,15 +5,12 @@
 #include "CoreMinimal.h"
 #include "Modules/ModuleManager.h"
 #include "Dataflow/DataflowEngineContextCaching.h"
-#include "Dataflow/DataflowEngineRendering.h"
-#include "Dataflow/DataflowRenderingFactory.h"
 #include "GeometryCollection/GeometryCollectionUtility.h"
 #include "GeometryCollection/Facades/CollectionRenderingFacade.h"
 #include "GeometryCollection/ManagedArrayCollection.h"
 #include "GeometryCollection/GeometryCollectionAlgo.h"
 
 
-FColor IDataflowEnginePlugin::SurfaceColor = FLinearColor(0.6, 0.6, 0.6).ToRGBE();
 FColor IDataflowEnginePlugin::VertexColor = FLinearColor(0.0,0.0,0.0).ToRGBE();
 FColor IDataflowEnginePlugin::SelectionPrimaryColor = FLinearColor(0.8, 0.4, 0.0).ToRGBE();
 FColor IDataflowEnginePlugin::SelectionLockedPrimaryColor = FLinearColor(0.8, 0.4, 0.2).ToRGBE();
@@ -31,7 +28,6 @@ IMPLEMENT_MODULE( FDataflowEnginePlugin, DataflowEnginePlugin)
 
 void FDataflowEnginePlugin::StartupModule()
 {
-	Dataflow::RenderingCallbacks();
 	Dataflow::ContextCachingCallbacks();
 }
 
diff --git a/Engine/Plugins/Experimental/Dataflow/Source/DataflowEnginePlugin/Private/Dataflow/DataflowEngineRendering.cpp b/Engine/Plugins/Experimental/Dataflow/Source/DataflowEnginePlugin/Private/Dataflow/DataflowEngineRendering.cpp
deleted file mode 100644
diff --git a/Engine/Plugins/Experimental/Dataflow/Source/DataflowEnginePlugin/Public/Dataflow/DataflowComponent.h b/Engine/Plugins/Experimental/Dataflow/Source/DataflowEnginePlugin/Public/Dataflow/DataflowComponent.h
--- a/Engine/Plugins/Experimental/Dataflow/Source/DataflowEnginePlugin/Public/Dataflow/DataflowComponent.h
+++ b/Engine/Plugins/Experimental/Dataflow/Source/DataflowEnginePlugin/Public/Dataflow/DataflowComponent.h
@@ -29,7 +29,6 @@ public:
 
 
 	//~ USceneComponent Interface.
-	virtual void TickComponent(float DeltaTime, enum ELevelTick TickType, FActorComponentTickFunction* ThisTickFunction) override;
 	virtual FBoxSphereBounds CalcBounds(const FTransform& LocalToWorld) const override;
 
 	//~ UPrimitiveComponent Interface.
@@ -52,8 +51,6 @@ public:
 	void SetRenderingCollection(FManagedArrayCollection&& InCollection);
 	const FManagedArrayCollection& GetRenderingCollection() const;
 	      FManagedArrayCollection& ModifyRenderingCollection();
-	void BuildRenderCollection();
-
 
 	/** Dataflow */
 	void SetDataflow(const UDataflow* InDataflow) { Dataflow = InDataflow; }
diff --git a/Engine/Plugins/Experimental/Dataflow/Source/DataflowEnginePlugin/Public/Dataflow/DataflowEnginePlugin.h b/Engine/Plugins/Experimental/Dataflow/Source/DataflowEnginePlugin/Public/Dataflow/DataflowEnginePlugin.h
--- a/Engine/Plugins/Experimental/Dataflow/Source/DataflowEnginePlugin/Public/Dataflow/DataflowEnginePlugin.h
+++ b/Engine/Plugins/Experimental/Dataflow/Source/DataflowEnginePlugin/Public/Dataflow/DataflowEnginePlugin.h
@@ -15,7 +15,6 @@ class IDataflowEnginePlugin : public IModuleInterface
 
 public:
 
-	static FColor SurfaceColor;
 	static FColor VertexColor;
 	static FColor SelectionPrimaryColor;
 	static FColor SelectionLockedPrimaryColor;
diff --git a/Engine/Plugins/Experimental/Dataflow/Source/DataflowEnginePlugin/Public/Dataflow/DataflowEngineRendering.h b/Engine/Plugins/Experimental/Dataflow/Source/DataflowEnginePlugin/Public/Dataflow/DataflowEngineRendering.h
deleted file mode 100644
diff --git a/Engine/Source/Runtime/Experimental/Dataflow/Engine/Private/Dataflow/DataflowEdNode.cpp b/Engine/Source/Runtime/Experimental/Dataflow/Engine/Private/Dataflow/DataflowEdNode.cpp
--- a/Engine/Source/Runtime/Experimental/Dataflow/Engine/Private/Dataflow/DataflowEdNode.cpp
+++ b/Engine/Source/Runtime/Experimental/Dataflow/Engine/Private/Dataflow/DataflowEdNode.cpp
@@ -6,7 +6,6 @@
 #include "Dataflow/DataflowCore.h"
 #include "Dataflow/DataflowCoreNodes.h"
 #include "Dataflow/DataflowObject.h"
-#include "Dataflow/DataflowRenderingFactory.h"
 #include "GeometryCollection/Facades/CollectionRenderingFacade.h"
 #include "Logging/LogMacros.h"
 #include "Textures/SlateIcon.h"
@@ -760,50 +759,5 @@ TArray<Dataflow::FRenderingParameter> UDataflowEdNode::GetRenderParameters() con
 }
 
 
-bool UDataflowEdNode::Render(GeometryCollection::Facades::FRenderingFacade& RenderData, const TSharedRef<Dataflow::FContext> Context, const Dataflow::IDataflowConstructionViewMode& ViewMode) const
-{
-	bool bNeedsRefresh = false;
-	if (DataflowGraph)
-	{
-		if (TSharedPtr<const FDataflowNode> NodeTarget = DataflowGraph->FindBaseNode(FName(GetName())))
-		{
-			if (Dataflow::FRenderingFactory* Factory = Dataflow::FRenderingFactory::GetInstance())
-			{
-				if (GetRenderParameters().Num())
-				{
-					for (Dataflow::FRenderingParameter& Parameter : GetRenderParameters())
-					{
-						Factory->RenderNodeOutput(RenderData, Dataflow::FGraphRenderingState{ GetDataflowNodeGuid(), NodeTarget.Get(), Parameter, Context.Get(), ViewMode});
-						bNeedsRefresh = true;
-					}
-				}
-			}
-		}
-	}
-	return bNeedsRefresh;
-}
-
-
-bool UDataflowEdNode::CanRender(const TSharedRef<Dataflow::FContext> Context, const Dataflow::IDataflowConstructionViewMode& ViewMode) const
-{
-	if (DataflowGraph)
-	{
-		if (TSharedPtr<const FDataflowNode> NodeTarget = DataflowGraph->FindBaseNode(FName(GetName())))
-		{
-			if (const Dataflow::FRenderingFactory* const Factory = Dataflow::FRenderingFactory::GetInstance())
-			{
-				for (const Dataflow::FRenderingParameter& Parameter : GetRenderParameters())
-				{
-					if (Factory->CanRenderNodeOutput(Dataflow::FGraphRenderingState{ GetDataflowNodeGuid(), NodeTarget.Get(), Parameter, Context.Get(), ViewMode }))
-					{
-						return true;
-					}
-				}
-			}
-		}
-	}
-	return false;
-}
-
 #undef LOCTEXT_NAMESPACE
 
diff --git a/Engine/Source/Runtime/Experimental/Dataflow/Engine/Private/Dataflow/DataflowRenderingFactory.cpp b/Engine/Source/Runtime/Experimental/Dataflow/Engine/Private/Dataflow/DataflowRenderingFactory.cpp
deleted file mode 100644
diff --git a/Engine/Source/Runtime/Experimental/Dataflow/Engine/Private/Dataflow/DataflowRenderingViewMode.cpp b/Engine/Source/Runtime/Experimental/Dataflow/Engine/Private/Dataflow/DataflowRenderingViewMode.cpp
deleted file mode 100644
diff --git a/Engine/Source/Runtime/Experimental/Dataflow/Engine/Public/Dataflow/DataflowEdNode.h b/Engine/Source/Runtime/Experimental/Dataflow/Engine/Public/Dataflow/DataflowEdNode.h
--- a/Engine/Source/Runtime/Experimental/Dataflow/Engine/Public/Dataflow/DataflowEdNode.h
+++ b/Engine/Source/Runtime/Experimental/Dataflow/Engine/Public/Dataflow/DataflowEdNode.h
@@ -4,7 +4,6 @@
 
 #include "CoreMinimal.h"
 #include "Dataflow/DataflowCore.h"
-#include "Dataflow/DataflowRenderingFactory.h"
 #include "EdGraph/EdGraph.h"
 #include "EdGraph/EdGraphNode.h"
 #include "GeometryCollection/GeometryCollection.h"
@@ -101,9 +100,6 @@ public:
 	DATAFLOWENGINE_API bool CanEnableWireframeRenderNode() const;
 
 	DATAFLOWENGINE_API TArray<Dataflow::FRenderingParameter> GetRenderParameters() const;
-	DATAFLOWENGINE_API virtual bool Render(GeometryCollection::Facades::FRenderingFacade& RenderData, const TSharedRef<Dataflow::FContext> Context, const Dataflow::IDataflowConstructionViewMode& ViewMode) const;
-
-	DATAFLOWENGINE_API virtual bool CanRender(const TSharedRef<Dataflow::FContext> Context, const Dataflow::IDataflowConstructionViewMode& ViewMode) const;
 
 private:
 
diff --git a/Engine/Source/Runtime/Experimental/Dataflow/Engine/Public/Dataflow/DataflowRenderingFactory.h b/Engine/Source/Runtime/Experimental/Dataflow/Engine/Public/Dataflow/DataflowRenderingFactory.h
deleted file mode 100644
diff --git a/Engine/Source/Runtime/Experimental/Dataflow/Engine/Public/Dataflow/DataflowRenderingViewMode.h b/Engine/Source/Runtime/Experimental/Dataflow/Engine/Public/Dataflow/DataflowRenderingViewMode.h
deleted file mode 100644
