Dataflow:
- Fixed Transform data display in DataflowCollectionSpreadsheet window

diff --git a/Engine/Plugins/Experimental/Dataflow/Source/DataflowEditor/Private/Dataflow/CollectionSpreadSheetWidget.cpp b/Engine/Plugins/Experimental/Dataflow/Source/DataflowEditor/Private/Dataflow/CollectionSpreadSheetWidget.cpp
--- a/Engine/Plugins/Experimental/Dataflow/Source/DataflowEditor/Private/Dataflow/CollectionSpreadSheetWidget.cpp
+++ b/Engine/Plugins/Experimental/Dataflow/Source/DataflowEditor/Private/Dataflow/CollectionSpreadSheetWidget.cpp
@@ -87,6 +87,15 @@ namespace CollectionSpreadSheetHelpers
 		return Out;
 	}
 
+	FString AttributeValueToString(FTransform3f Value)
+	{
+		const FVector3f Translation = Value.GetTranslation();
+		const FVector3f Rotation = Value.GetRotation().Euler();
+		const FVector3f Scale = Value.GetScale3D();
+
+		return FString::Printf(TEXT("T:(%s) R:(%s) S:(%s)"), *Translation.ToString(), *Rotation.ToString(), *Scale.ToString());
+	}
+
 	FString AttributeValueToString(FTransform Value)
 	{
 		const FVector Translation = Value.GetTranslation();
@@ -191,6 +200,10 @@ namespace CollectionSpreadSheetHelpers
 			ValueAsString = AttributeValueToString<FBox>(InCollection, InAttributeName, InGroupName, InIdxColumn);
 			break;
 
+		case FManagedArrayCollection::EArrayType::FTransform3fType:
+			ValueAsString = AttributeValueToString<FTransform3f>(InCollection, InAttributeName, InGroupName, InIdxColumn);
+			break;
+
 		default:
 			//ensure(false);
 			ValueAsString = "<Unknown Data Type>";
@@ -277,6 +290,7 @@ void SCollectionSpreadSheet::Construct(const FArguments& InArgs)
 	];
 
 	AttrTypeWidthMap.Add("Transform", 600);
+	AttrTypeWidthMap.Add("Transform3f", 600);
 	AttrTypeWidthMap.Add("String", 200);
 	AttrTypeWidthMap.Add("LinearColor", 250);
 	AttrTypeWidthMap.Add("int32", 100);
