Dataflow:
- Fixed evaluation bug

diff --git a/Engine/Plugins/Experimental/Dataflow/Source/DataflowEditor/Private/Dataflow/DataflowEditorCommands.cpp b/Engine/Plugins/Experimental/Dataflow/Source/DataflowEditor/Private/Dataflow/DataflowEditorCommands.cpp
--- a/Engine/Plugins/Experimental/Dataflow/Source/DataflowEditor/Private/Dataflow/DataflowEditorCommands.cpp
+++ b/Engine/Plugins/Experimental/Dataflow/Source/DataflowEditor/Private/Dataflow/DataflowEditorCommands.cpp
@@ -201,22 +201,12 @@ void FDataflowEditorCommands::EvaluateNode(Dataflow::FContext& Context, Dataflow
 				}
 			}
 
-			if (Node != nullptr)
+			if (Node)
 			{
-				if (Output == nullptr)
+				if (Context.GetTimestamp() < Node->GetTimestamp())
 				{
-					if (Node->GetTimestamp() >= OutLastNodeTimestamp)
-					{
-						Context.Evaluate(Node, nullptr);
-						OutLastNodeTimestamp = Context.GetTimestamp();
-					}
-				}
-				else // Output != nullptr
-				{
-					if (!Context.HasData(Output->CacheKey(), Context.GetTimestamp()))
-					{
-						Context.Evaluate(Node, Output);
-					}
+					Context.Evaluate(Node, Output);
+					OutLastNodeTimestamp = Context.GetTimestamp();
 				}
 			}
 		}
@@ -243,38 +233,23 @@ void FDataflowEditorCommands::EvaluateTerminalNode(Dataflow::FContext& Context,
 				}
 			}
 
-			if (Node != nullptr)
+			if (Node)
 			{
-				if (Output == nullptr)
+				if (Context.GetTimestamp() < Node->GetTimestamp())
 				{
-					if (Node->GetTimestamp() >= OutLastNodeTimestamp)
+					if (const FDataflowTerminalNode* TerminalNode = Node->AsType<const FDataflowTerminalNode>())
 					{
-						if (const FDataflowTerminalNode* TerminalNode = Node->AsType<const FDataflowTerminalNode>())
+						if (InAsset)
 						{
-							if (InAsset)
-							{
-								TerminalNode->SetAssetValue(InAsset, Context);  // Must set asset value before call to Evaluate
-							}
+							TerminalNode->SetAssetValue(InAsset, Context);  // Kriss.Gossart: Must set asset value before call to Evaluate
+																			// This dependency is neccessary for now for the cloth editor code.
+																			// It is order dependant, with some nodes relying on the non const 
+																			// SetAssetValue to be called first in order to have a valid const Evaluation.
 						}
-
-						Context.Evaluate(Node, nullptr);
-						OutLastNodeTimestamp = Context.GetTimestamp();
 					}
-				}
-				else // Output != nullptr
-				{
-					if (!Context.HasData(Output->CacheKey(), Context.GetTimestamp()))
-					{
-						if (const FDataflowTerminalNode* TerminalNode = Node->AsType<const FDataflowTerminalNode>())
-						{
-							if (InAsset)
-							{
-								TerminalNode->SetAssetValue(InAsset, Context);  // Must set asset value before call to Evaluate
-							}
-						}
 
-						Context.Evaluate(Node, Output);
-					}
+					Context.Evaluate(Node, Output);
+					OutLastNodeTimestamp = Context.GetTimestamp();
 				}
 			}
 		}
@@ -1031,7 +1006,7 @@ void FDataflowEditorCommands::PasteNodes(UDataflow* Graph, const TSharedPtr<SDat
 			ShowNotificationMessage(FText::Format(MessageFormat, NumPastedNodes), SNotificationItem::CS_Success);
 		}
 
-		// Display message stating that comment boxe(s) were pasted to clipboard
+		// Display message stating that comment box(es) were pasted to clipboard
 		const int32 NumPastedComments = CopyPasteContent.CommentNodeData.Num();
 		if (NumPastedComments > 0)
 		{
