Dataflow - Fixed delete node undo.

[RN] Dataflow - Fixed delete node undo.
#tests Undid deleting a bunch of nodes, made sure the connections were still working, and the graph was still evaluating correctly.

diff --git a/Engine/Plugins/Experimental/Dataflow/Source/DataflowEditor/Private/Dataflow/DataflowGraphEditor.cpp b/Engine/Plugins/Experimental/Dataflow/Source/DataflowEditor/Private/Dataflow/DataflowGraphEditor.cpp
--- a/Engine/Plugins/Experimental/Dataflow/Source/DataflowEditor/Private/Dataflow/DataflowGraphEditor.cpp
+++ b/Engine/Plugins/Experimental/Dataflow/Source/DataflowEditor/Private/Dataflow/DataflowGraphEditor.cpp
@@ -202,9 +202,21 @@ void SDataflowGraphEditor::DeleteNode()
 		const FGraphPanelSelectionSet& SelectedNodes = GetSelectedNodes();
 		if (SelectedNodes.Num() > 0)
 		{
-			FDataflowEditorCommands::DeleteNodes(DataflowAsset.Get(), SelectedNodes);
+			if (UDataflow* const Graph = DataflowAsset.Get())
+			{
+				const FScopedTransaction Transaction(LOCTEXT("DeleteSelectedNodes", "Delete selected nodes"));
+
+				Graph->Modify();
+
+				for (FGraphPanelSelectionSet::TConstIterator It(SelectedNodes); It; ++It)
+				{
+					(*It)->Modify();
+				}
 
-			OnNodeDeletedMulticast.Broadcast(SelectedNodes);
+				FDataflowEditorCommands::DeleteNodes(Graph, SelectedNodes);
+
+				OnNodeDeletedMulticast.Broadcast(SelectedNodes);
+			}
 		}
 	}
 }
