Fix multiple proxy creation in dataflow + pre/post tick events

diff --git a/Engine/Plugins/Experimental/ChaosFlesh/Source/ChaosFleshEngine/Private/ChaosFlesh/ChaosDeformableSolverComponent.cpp b/Engine/Plugins/Experimental/ChaosFlesh/Source/ChaosFleshEngine/Private/ChaosFlesh/ChaosDeformableSolverComponent.cpp
--- a/Engine/Plugins/Experimental/ChaosFlesh/Source/ChaosFleshEngine/Private/ChaosFlesh/ChaosDeformableSolverComponent.cpp
+++ b/Engine/Plugins/Experimental/ChaosFlesh/Source/ChaosFleshEngine/Private/ChaosFlesh/ChaosDeformableSolverComponent.cpp
@@ -155,6 +155,10 @@ void UDeformableSolverComponent::UpdateDeformableEndTickState(bool bRegister)
 void UDeformableSolverComponent::BeginPlay()
 {
 	Super::BeginPlay();
+	if(!SimulationAsset.DataflowAsset)
+	{
+		BuildSimulationProxy();
+	}
 }
 
 void UDeformableSolverComponent::TickComponent(float DeltaTime, enum ELevelTick TickType, FActorComponentTickFunction* ThisTickFunction)
@@ -197,6 +201,10 @@ void UDeformableSolverComponent::TickComponent(float DeltaTime, enum ELevelTick
 void UDeformableSolverComponent::EndPlay(const EEndPlayReason::Type EndPlayReason)
 {
 	Super::EndPlay(EndPlayReason);
+	if(!SimulationAsset.DataflowAsset)
+	{
+		ResetSimulationProxy();
+	}
 }
 
 void UDeformableSolverComponent::BuildSimulationProxy()
diff --git a/Engine/Plugins/Experimental/ChaosFlesh/Source/ChaosFleshEngine/Public/ChaosFlesh/ChaosDeformableSolverComponent.h b/Engine/Plugins/Experimental/ChaosFlesh/Source/ChaosFleshEngine/Public/ChaosFlesh/ChaosDeformableSolverComponent.h
--- a/Engine/Plugins/Experimental/ChaosFlesh/Source/ChaosFleshEngine/Public/ChaosFlesh/ChaosDeformableSolverComponent.h
+++ b/Engine/Plugins/Experimental/ChaosFlesh/Source/ChaosFleshEngine/Public/ChaosFlesh/ChaosDeformableSolverComponent.h
@@ -273,6 +273,7 @@ struct FDataflowFleshSolverProxy : public FDataflowPhysicsSolverProxy
 		const Chaos::Softs::FDeformableSolver::FPhysicsThreadAccess PhysicsThreadAccess(Solver.Get(), Chaos::Softs::FPhysicsThreadAccessor());
 		return PhysicsThreadAccess.GetProperties().TimeStepSize;
 	}
+	virtual bool IsValid() const override { return Solver.IsValid();}
 	// End FPhysicsSolverInterface overrides
 
 	/** Chaos deformable solver that will be used in the component */
diff --git a/Engine/Plugins/Experimental/Dataflow/Source/DataflowSimulation/Private/Dataflow/DataflowSimulationManager.cpp b/Engine/Plugins/Experimental/Dataflow/Source/DataflowSimulation/Private/Dataflow/DataflowSimulationManager.cpp
--- a/Engine/Plugins/Experimental/Dataflow/Source/DataflowSimulation/Private/Dataflow/DataflowSimulationManager.cpp
+++ b/Engine/Plugins/Experimental/Dataflow/Source/DataflowSimulation/Private/Dataflow/DataflowSimulationManager.cpp
@@ -5,8 +5,10 @@
 #include "Dataflow/DataflowSimulationInterface.h"
 #include "Dataflow/DataflowSimulationProxy.h"
 #include "Components/ActorComponent.h"
+#include "Async/TaskGraphInterfaces.h"
 #include "Containers/Map.h"
 #include "Engine/World.h"
+#include "Kismet/GameplayStatics.h"
 
 static FAutoConsoleTaskPriority DataflowSimulationTaskPriority(
 		TEXT("TaskGraph.TaskPriorities.DataflowSimulationTask"),
@@ -88,6 +90,34 @@ FDelegateHandle UDataflowSimulationManager::OnDestroyPhysicsStateHandle;
 UDataflowSimulationManager::UDataflowSimulationManager()
 {}
 
+inline void PreSimulationTick(const TObjectPtr<UObject>& SimulationWorld, const float SimulationTime, const float DeltaTime)
+{
+	if(SimulationWorld)
+	{
+		TArray<AActor*> Actors;
+		UGameplayStatics::GetAllActorsWithInterface(SimulationWorld, UDataflowSimulationActor::StaticClass(), Actors);
+	
+		for (AActor* CurrentActor : Actors)
+		{
+			IDataflowSimulationActor::Execute_PreDataflowSimulationTick(CurrentActor, SimulationTime, DeltaTime);
+		}
+	}
+}
+
+inline void PostSimulationTick(const TObjectPtr<UObject>& SimulationWorld, const float SimulationTime, const float DeltaTime)
+{
+	if(SimulationWorld)
+	{
+		TArray<AActor*> Actors;
+		UGameplayStatics::GetAllActorsWithInterface(SimulationWorld, UDataflowSimulationActor::StaticClass(), Actors);
+	
+		for (AActor* CurrentActor : Actors)
+		{
+			IDataflowSimulationActor::Execute_PostDataflowSimulationTick(CurrentActor, SimulationTime, DeltaTime);
+		}
+	}
+}
+
 void UDataflowSimulationManager::Tick(float DeltaTime)
 {
 	Super::Tick(DeltaTime);
@@ -154,11 +184,19 @@ void UDataflowSimulationManager::OnStartup()
 	{
 		if(IDataflowSimulationInterface* SimulationInterface = Cast<IDataflowSimulationInterface>(ActorComponent))
 		{
-			// Build the simulation proxy
-			SimulationInterface->BuildSimulationProxy();
+			if(SimulationInterface->GetSimulationAsset().DataflowAsset)
+			{
+				FDataflowSimulationProxy* SimulationProxy = SimulationInterface->GetSimulationProxy();
+				
+				if(!SimulationProxy || (SimulationProxy && !SimulationProxy->IsValid()))
+				{
+					// Build the simulation proxy
+					SimulationInterface->BuildSimulationProxy();
+				}
 
-			// Register the simulation interface to the manager
-			SimulationInterface->RegisterManagerInterface(ActorComponent->GetWorld());
+				// Register the simulation interface to the manager
+				SimulationInterface->RegisterManagerInterface(ActorComponent->GetWorld());
+			}
 		}
 	});
 
@@ -166,11 +204,19 @@ void UDataflowSimulationManager::OnStartup()
 	{
 		if(IDataflowSimulationInterface* SimulationInterface = Cast<IDataflowSimulationInterface>(ActorComponent))
 		{
-			// Reset the simulation proxy
-			SimulationInterface->ResetSimulationProxy();
+			if(SimulationInterface->GetSimulationAsset().DataflowAsset)
+			{
+				FDataflowSimulationProxy* SimulationProxy = SimulationInterface->GetSimulationProxy();
+				
+				if(SimulationProxy && SimulationProxy->IsValid())
+				{
+					// Reset the simulation proxy
+					SimulationInterface->ResetSimulationProxy();
+				}
 
-			// Unregister the simulation interface from the manager
-			SimulationInterface->UnregisterManagerInterface(ActorComponent->GetWorld());
+				// Unregister the simulation interface from the manager
+				SimulationInterface->UnregisterManagerInterface(ActorComponent->GetWorld());
+			}
 		}
 	});
 
@@ -204,8 +250,15 @@ void UDataflowSimulationManager::OnShutdown()
 
 void UDataflowSimulationManager::WriteSimulationData(const float DeltaTime)
 {
+	// Pre-simulation callback that could be used in BP before the simulation
+	PreSimulationTick(GetWorld(), GetWorld()->GetTimeSeconds(), DeltaTime);
+	
 	for(const TPair<TObjectPtr<UDataflow>, Dataflow::FDataflowSimulationData>& DataflowData : SimulationData)
 	{
+		if(DataflowData.Value.SimulationContext.IsValid())
+		{
+			DataflowData.Value.SimulationContext->ResetSimulationProxies();
+		}
 		for(const TPair<FString,TSet<IDataflowSimulationInterface*>>& SimulationInterfaces : DataflowData.Value.SimulationInterfaces)
 		{
 			for(IDataflowSimulationInterface* SimulationInterface : SimulationInterfaces.Value)
@@ -219,14 +272,20 @@ void UDataflowSimulationManager::WriteSimulationData(const float DeltaTime)
 					if(SimulationInterface->GetSimulationProxy())
 					{
 						SimulationInterface->GetSimulationProxy()->SetSimulationGroups(SimulationInterface->GetSimulationAsset().SimulationGroups);
-						DataflowData.Value.SimulationContext->AddSimulationProxy(SimulationInterfaces.Key, SimulationInterface->GetSimulationProxy());
+						if(DataflowData.Value.SimulationContext.IsValid())
+						{
+							DataflowData.Value.SimulationContext->AddSimulationProxy(SimulationInterfaces.Key, SimulationInterface->GetSimulationProxy());
+						}
 					}
 					
 					SimulationInterface->WriteToSimulation(DeltaTime);
 				}
 			}
 		}
-		DataflowData.Value.SimulationContext->RegisterProxyGroups();
+		if(DataflowData.Value.SimulationContext.IsValid())
+		{
+			DataflowData.Value.SimulationContext->RegisterProxyGroups();
+		}
 	}
 }
 
@@ -244,13 +303,18 @@ void UDataflowSimulationManager::ReadSimulationData(const float DeltaTime)
 				}
 			}
 		}
-		DataflowData.Value.SimulationContext->ResetSimulationProxies();
+		if(DataflowData.Value.SimulationContext.IsValid())
+		{
+			DataflowData.Value.SimulationContext->ResetSimulationProxies();
+		}
 	}
 	if(bStepSimulationScene)
 	{
 		bIsSimulationEnabled = false;
 		bStepSimulationScene = false;
 	}
+	// Post-simulation callback that could be used in BP after the simulation
+	PostSimulationTick(GetWorld(), GetWorld()->GetTimeSeconds(), DeltaTime);
 }
 
 void UDataflowSimulationManager::AdvanceSimulationData(const float DeltaTime, const float SimulationTime)
@@ -268,9 +332,12 @@ void UDataflowSimulationManager::StartSimulationTasks(const float DeltaTime, con
 
 	for(TPair<TObjectPtr<UDataflow>, Dataflow::FDataflowSimulationData>& DataflowData : SimulationData)
 	{
-		// Add a simulation task linked to that solver
-		SimulationTasks.Add(TGraphTask<Dataflow::FDataflowSimulationTask>::CreateTask(nullptr, ENamedThreads::GameThread).ConstructAndDispatchWhenReady(
-				DataflowData.Key, DataflowData.Value.SimulationContext, DeltaTime, SimulationTime));
+		if(!DataflowData.Value.IsEmpty())
+		{
+			// Add a simulation task linked to that solver
+			SimulationTasks.Add(TGraphTask<Dataflow::FDataflowSimulationTask>::CreateTask(nullptr, ENamedThreads::GameThread).ConstructAndDispatchWhenReady(
+					DataflowData.Key, DataflowData.Value.SimulationContext, DeltaTime, SimulationTime));
+		}
 	}
 }
 
@@ -381,3 +448,5 @@ void UDataflowSimulationManager::Deinitialize()
 
 	CompleteSimulationTasks();
 }
+
+
diff --git a/Engine/Plugins/Experimental/Dataflow/Source/DataflowSimulation/Public/Dataflow/DataflowSimulationManager.h b/Engine/Plugins/Experimental/Dataflow/Source/DataflowSimulation/Public/Dataflow/DataflowSimulationManager.h
--- a/Engine/Plugins/Experimental/Dataflow/Source/DataflowSimulation/Public/Dataflow/DataflowSimulationManager.h
+++ b/Engine/Plugins/Experimental/Dataflow/Source/DataflowSimulation/Public/Dataflow/DataflowSimulationManager.h
@@ -4,6 +4,7 @@
 
 #include "CoreMinimal.h"
 #include "UObject/ObjectMacros.h"
+#include "UObject/Interface.h"
 #include "Subsystems/WorldSubsystem.h"
 #include "Templates/SharedPointer.h"
 #include "Dataflow/DataflowSimulationContext.h"
@@ -23,6 +24,19 @@ namespace Dataflow
 
 		/** Simulation context used to evaluate the graph on PT */
 		TSharedPtr<Dataflow::FDataflowSimulationContext> SimulationContext;
+
+		/** Check is there is any datas to process */
+		bool IsEmpty() const
+		{
+			for(const TPair<FString, TSet<IDataflowSimulationInterface*>>& InterfacesPair : SimulationInterfaces)
+			{
+				if(!InterfacesPair.Value.IsEmpty())
+				{
+					return false;
+				}
+			}
+			return true;
+		}
 	};
 }
 
@@ -106,10 +120,31 @@ private :
 	
 	/** Simulation tasks in which the graph will be evaluated */
 	TArray<FGraphEventRef> SimulationTasks;
-
+ 
 	/** Boolean to control if the simulation should be disabled or not */
 	bool bIsSimulationEnabled = true;
 	
 	/** Boolean to check if we are stepping the simulation scene */
 	bool bStepSimulationScene = false;
 };
+
+/** Dataflow simulation actor interface to be able to call BP events before/after the manager ticking in case we need it */
+UINTERFACE(MinimalAPI, Blueprintable)
+class UDataflowSimulationActor : public UInterface
+{
+	GENERATED_BODY()
+};
+
+class IDataflowSimulationActor
+{
+	GENERATED_BODY()
+ 
+public:
+	/** Pre simulation callback function that can be implemented in C++ or Blueprint. */
+	UFUNCTION(BlueprintCallable, BlueprintNativeEvent, Category="Dataflow")
+	void PreDataflowSimulationTick(const float SimulationTime, const float DeltaTime);
+
+	/** Post simulation callback function that can be implemented in C++ or Blueprint. */
+	UFUNCTION(BlueprintCallable, BlueprintNativeEvent, Category="Dataflow")
+	void PostDataflowSimulationTick(const float SimulationTime, const float DeltaTime);
+};
diff --git a/Engine/Plugins/Experimental/Dataflow/Source/DataflowSimulation/Public/Dataflow/DataflowSimulationProxy.h b/Engine/Plugins/Experimental/Dataflow/Source/DataflowSimulation/Public/Dataflow/DataflowSimulationProxy.h
--- a/Engine/Plugins/Experimental/Dataflow/Source/DataflowSimulation/Public/Dataflow/DataflowSimulationProxy.h
+++ b/Engine/Plugins/Experimental/Dataflow/Source/DataflowSimulation/Public/Dataflow/DataflowSimulationProxy.h
@@ -16,6 +16,9 @@ struct DATAFLOWSIMULATION_API FDataflowSimulationProxy
 	FDataflowSimulationProxy() = default;
 	virtual ~FDataflowSimulationProxy() = default;
 	
+	/** Check if the proxy is valid */
+	virtual bool IsValid() const { return true;}
+	
 	/** Get the proxy script struct */
 	virtual const UScriptStruct* GetScriptStruct() const
 	{
diff --git a/Engine/Source/Runtime/Experimental/Dataflow/Engine/Private/Dataflow/DataflowContent.cpp b/Engine/Source/Runtime/Experimental/Dataflow/Engine/Private/Dataflow/DataflowContent.cpp
--- a/Engine/Source/Runtime/Experimental/Dataflow/Engine/Private/Dataflow/DataflowContent.cpp
+++ b/Engine/Source/Runtime/Experimental/Dataflow/Engine/Private/Dataflow/DataflowContent.cpp
@@ -1,12 +1,11 @@
 // Copyright Epic Games, Inc. All Rights Reserved.
 
 #include "Dataflow/DataflowContent.h"
-
+#include "Dataflow/DataflowObjectInterface.h"
+#include "Dataflow/DataflowObject.h"
 #include "Animation/AnimSingleNodeInstance.h"
 #include "Animation/Skeleton.h"
 #include "Components/SkeletalMeshComponent.h"
-#include "Dataflow/DataflowNodeParameters.h"
-#include "Dataflow/DataflowObject.h"
 #include "Engine/SkeletalMesh.h"
 #include "GameFramework/Actor.h"
 #include "Misc/PackageName.h"
diff --git a/Engine/Source/Runtime/Experimental/Dataflow/Engine/Public/Dataflow/DataflowContent.h b/Engine/Source/Runtime/Experimental/Dataflow/Engine/Public/Dataflow/DataflowContent.h
--- a/Engine/Source/Runtime/Experimental/Dataflow/Engine/Public/Dataflow/DataflowContent.h
+++ b/Engine/Source/Runtime/Experimental/Dataflow/Engine/Public/Dataflow/DataflowContent.h
@@ -2,10 +2,9 @@
 
 #pragma once
 
-#include "Dataflow/DataflowObjectInterface.h"
-#include "Components/PrimitiveComponent.h"
-#include "Templates/SharedPointer.h"
+#include "UObject/Interface.h"
 #include "Dataflow/DataflowEdNode.h"
+#include "Templates/SharedPointer.h"
 #include "DataflowContent.generated.h"
 
 class FDataflowEditorToolkit;
@@ -18,7 +17,15 @@ class UDataflowBaseContent;
 class FPreviewScene;
 class UAnimSingleNodeInstance;
 class AActor;
-namespace Dataflow { class IDataflowConstructionViewMode; }
+class UDataflowEdNode;
+struct FManagedArrayCollection;
+
+namespace Dataflow
+{
+	class IDataflowConstructionViewMode;
+	template<class Base> class TEngineContext;
+	typedef TEngineContext<FContextSingle> FEngineContext;
+}
 
 namespace DataflowContextHelpers
 {
diff --git a/Engine/Source/Runtime/Experimental/Dataflow/Engine/Public/Dataflow/DataflowObject.h b/Engine/Source/Runtime/Experimental/Dataflow/Engine/Public/Dataflow/DataflowObject.h
--- a/Engine/Source/Runtime/Experimental/Dataflow/Engine/Public/Dataflow/DataflowObject.h
+++ b/Engine/Source/Runtime/Experimental/Dataflow/Engine/Public/Dataflow/DataflowObject.h
@@ -4,7 +4,7 @@
 
 #include "Containers/Array.h"
 #include "CoreMinimal.h"
-#include "DataflowPreview.h"
+#include "Dataflow/DataflowPreview.h"
 #include "Dataflow/DataflowCore.h"
 #include "EdGraph/EdGraph.h"
 #include "Templates/Function.h"
