Fix CIS with dataflow selection view

diff --git a/Engine/Plugins/Experimental/Dataflow/Source/DataflowEditor/Private/Dataflow/DataflowSelectionView.cpp b/Engine/Plugins/Experimental/Dataflow/Source/DataflowEditor/Private/Dataflow/DataflowSelectionView.cpp
--- a/Engine/Plugins/Experimental/Dataflow/Source/DataflowEditor/Private/Dataflow/DataflowSelectionView.cpp
+++ b/Engine/Plugins/Experimental/Dataflow/Source/DataflowEditor/Private/Dataflow/DataflowSelectionView.cpp
@@ -5,6 +5,7 @@
 #include "Templates/EnableIf.h"
 #include "Dataflow/DataflowContent.h"
 //#include "Dataflow/DataflowEdNode.h"
+#include "Dataflow/DataflowObjectInterface.h"
 #include "Dataflow/DataflowSelection.h"
 #include "Dataflow/SelectionViewWidget.h"
 
