Fix dataflow CIS regarding primitvecomponent/enginecontext

diff --git a/Engine/Plugins/Experimental/Dataflow/Source/DataflowEditor/Private/Dataflow/DataflowCollectionSpreadSheet.cpp b/Engine/Plugins/Experimental/Dataflow/Source/DataflowEditor/Private/Dataflow/DataflowCollectionSpreadSheet.cpp
--- a/Engine/Plugins/Experimental/Dataflow/Source/DataflowEditor/Private/Dataflow/DataflowCollectionSpreadSheet.cpp
+++ b/Engine/Plugins/Experimental/Dataflow/Source/DataflowEditor/Private/Dataflow/DataflowCollectionSpreadSheet.cpp
@@ -1,7 +1,7 @@
 // Copyright Epic Games, Inc. All Rights Reserved.
 
 #include "Dataflow/DataflowCollectionSpreadSheet.h"
-
+#include "Dataflow/DataflowObjectInterface.h"
 #include "Dataflow/DataflowCollectionSpreadSheetWidget.h"
 #include "Dataflow/DataflowContent.h"
 #include "Dataflow/DataflowSelection.h"
diff --git a/Engine/Plugins/Experimental/Dataflow/Source/DataflowEditor/Private/Dataflow/DataflowEditorPreviewSceneBase.cpp b/Engine/Plugins/Experimental/Dataflow/Source/DataflowEditor/Private/Dataflow/DataflowEditorPreviewSceneBase.cpp
--- a/Engine/Plugins/Experimental/Dataflow/Source/DataflowEditor/Private/Dataflow/DataflowEditorPreviewSceneBase.cpp
+++ b/Engine/Plugins/Experimental/Dataflow/Source/DataflowEditor/Private/Dataflow/DataflowEditorPreviewSceneBase.cpp
@@ -1,10 +1,9 @@
 // Copyright Epic Games, Inc. All Rights Reserved.
 
 #include "Dataflow/DataflowEditorPreviewSceneBase.h"
-
-
 #include "AssetEditorModeManager.h"
 #include "Dataflow/DataflowEditor.h"
+#include "Components/PrimitiveComponent.h"
 #include "Elements/Framework/EngineElementsLibrary.h"
 #include "Selection.h"
 
diff --git a/Engine/Plugins/Experimental/Dataflow/Source/DataflowEditor/Private/Dataflow/DataflowSimulationScene.cpp b/Engine/Plugins/Experimental/Dataflow/Source/DataflowEditor/Private/Dataflow/DataflowSimulationScene.cpp
--- a/Engine/Plugins/Experimental/Dataflow/Source/DataflowEditor/Private/Dataflow/DataflowSimulationScene.cpp
+++ b/Engine/Plugins/Experimental/Dataflow/Source/DataflowEditor/Private/Dataflow/DataflowSimulationScene.cpp
@@ -3,6 +3,7 @@
 #include "Dataflow/DataflowSimulationScene.h"
 #include "Dataflow/DataflowSimulationManager.h"
 #include "Dataflow/DataflowEditor.h"
+#include "Components/PrimitiveComponent.h"
 #include "Chaos/CacheManagerActor.h"
 #include "Misc/TransactionObjectEvent.h"
 #include "EngineUtils.h"
diff --git a/Engine/Plugins/Experimental/Dataflow/Source/DataflowEditor/Private/Dataflow/DataflowSkeletonView.cpp b/Engine/Plugins/Experimental/Dataflow/Source/DataflowEditor/Private/Dataflow/DataflowSkeletonView.cpp
--- a/Engine/Plugins/Experimental/Dataflow/Source/DataflowEditor/Private/Dataflow/DataflowSkeletonView.cpp
+++ b/Engine/Plugins/Experimental/Dataflow/Source/DataflowEditor/Private/Dataflow/DataflowSkeletonView.cpp
@@ -5,6 +5,7 @@
 #include "Animation/Skeleton.h"
 #include "Dataflow/DataflowEditor.h"
 #include "Dataflow/DataflowEdNode.h"
+#include "Dataflow/DataflowObjectInterface.h"
 #include "GeometryCollection/GeometryCollectionEngineConversion.h"
 #include "IEditableSkeleton.h"
 #include "ISkeletonTree.h"
diff --git a/Engine/Plugins/Experimental/Dataflow/Source/DataflowEditor/Private/DataflowEditorTools/DataflowEditorWeightMapPaintTool.cpp b/Engine/Plugins/Experimental/Dataflow/Source/DataflowEditor/Private/DataflowEditorTools/DataflowEditorWeightMapPaintTool.cpp
--- a/Engine/Plugins/Experimental/Dataflow/Source/DataflowEditor/Private/DataflowEditorTools/DataflowEditorWeightMapPaintTool.cpp
+++ b/Engine/Plugins/Experimental/Dataflow/Source/DataflowEditor/Private/DataflowEditorTools/DataflowEditorWeightMapPaintTool.cpp
@@ -14,6 +14,7 @@
 #include "Dataflow/DataflowCollectionAddScalarVertexPropertyNode.h"
 #include "Dataflow/DataflowConstructionViewportClient.h"
 #include "Dataflow/DataflowContent.h"
+#include "Dataflow/DataflowObjectInterface.h"
 #include "Dataflow/DataflowEdNode.h"
 #include "Dataflow/DataflowEditorCollectionComponent.h"
 #include "Dataflow/DataflowEditorMode.h"
diff --git a/Engine/Plugins/Experimental/Dataflow/Source/DataflowEditor/Public/Dataflow/DataflowRenderingFactory.h b/Engine/Plugins/Experimental/Dataflow/Source/DataflowEditor/Public/Dataflow/DataflowRenderingFactory.h
--- a/Engine/Plugins/Experimental/Dataflow/Source/DataflowEditor/Public/Dataflow/DataflowRenderingFactory.h
+++ b/Engine/Plugins/Experimental/Dataflow/Source/DataflowEditor/Public/Dataflow/DataflowRenderingFactory.h
@@ -7,6 +7,7 @@
 #include "Dataflow/DataflowNodeParameters.h"
 #include "Dataflow/DataflowInputOutput.h"
 #include "Dataflow/DataflowNode.h"
+#include "Dataflow/DataflowObjectInterface.h"
 #include "Dataflow/DataflowRenderingViewMode.h"
 #include "GeometryCollection/Facades/CollectionRenderingFacade.h"
 
