Dataflow : Fix dataflow content asset to return the right clas and make it non viewable /editable

diff --git a/Engine/Plugins/Experimental/Dataflow/Source/DataflowEditor/Private/Dataflow/AssetDefinition_DataflowContext.cpp b/Engine/Plugins/Experimental/Dataflow/Source/DataflowEditor/Private/Dataflow/AssetDefinition_DataflowContext.cpp
--- a/Engine/Plugins/Experimental/Dataflow/Source/DataflowEditor/Private/Dataflow/AssetDefinition_DataflowContext.cpp
+++ b/Engine/Plugins/Experimental/Dataflow/Source/DataflowEditor/Private/Dataflow/AssetDefinition_DataflowContext.cpp
@@ -36,7 +36,7 @@ FText UAssetDefinition_DataflowContext::GetAssetDisplayName() const
 
 TSoftClassPtr<UObject> UAssetDefinition_DataflowContext::GetAssetClass() const
 {
-	return UDataflow::StaticClass();
+	return UDataflowBaseContent::StaticClass();
 }
 
 FLinearColor UAssetDefinition_DataflowContext::GetAssetColor() const
@@ -55,5 +55,10 @@ UThumbnailInfo* UAssetDefinition_DataflowContext::LoadThumbnailInfo(const FAsset
 	return UE::Editor::FindOrCreateThumbnailInfo(InAsset.GetAsset(), USceneThumbnailInfo::StaticClass());
 }
 
+FAssetOpenSupport UAssetDefinition_DataflowContext::GetAssetOpenSupport(const FAssetOpenSupportArgs& OpenSupportArgs) const
+{
+	// this asset shoul dnot be editable at any time
+	return FAssetOpenSupport(EAssetOpenMethod::View, false);
+}
 
 #undef LOCTEXT_NAMESPACE
diff --git a/Engine/Plugins/Experimental/Dataflow/Source/DataflowEditor/Public/Dataflow/AssetDefinition_DataflowContext.h b/Engine/Plugins/Experimental/Dataflow/Source/DataflowEditor/Public/Dataflow/AssetDefinition_DataflowContext.h
--- a/Engine/Plugins/Experimental/Dataflow/Source/DataflowEditor/Public/Dataflow/AssetDefinition_DataflowContext.h
+++ b/Engine/Plugins/Experimental/Dataflow/Source/DataflowEditor/Public/Dataflow/AssetDefinition_DataflowContext.h
@@ -20,5 +20,6 @@ private:
 	virtual FLinearColor GetAssetColor() const override;
 	virtual TConstArrayView<FAssetCategoryPath> GetAssetCategories() const override;
 	virtual UThumbnailInfo* LoadThumbnailInfo(const FAssetData& InAssetData) const override;
+	virtual FAssetOpenSupport GetAssetOpenSupport(const FAssetOpenSupportArgs& OpenSupportArgs) const override;
 };
 
