4 - Cloth support for simulation dataflow

#virtualized

diff --git a/Engine/Plugins/Animation/MLDeformer/ChaosClothGenerator/Source/ChaosClothGenerator/ChaosClothGenerator.build.cs b/Engine/Plugins/Animation/MLDeformer/ChaosClothGenerator/Source/ChaosClothGenerator/ChaosClothGenerator.build.cs
--- a/Engine/Plugins/Animation/MLDeformer/ChaosClothGenerator/Source/ChaosClothGenerator/ChaosClothGenerator.build.cs
+++ b/Engine/Plugins/Animation/MLDeformer/ChaosClothGenerator/Source/ChaosClothGenerator/ChaosClothGenerator.build.cs
@@ -46,7 +46,8 @@ public class ChaosClothGenerator : ModuleRules
 				"SkeletalMeshDescription",
 				"Slate",
 				"SlateCore",
-				"UnrealEd"
+				"UnrealEd",
+				"DataflowSimulation"
 			}
 			);
 		
diff --git a/Engine/Plugins/ChaosCloth/Source/ChaosCloth/Private/ChaosCloth/ChaosClothingCacheSchema.cpp b/Engine/Plugins/ChaosCloth/Source/ChaosCloth/Private/ChaosCloth/ChaosClothingCacheSchema.cpp
--- a/Engine/Plugins/ChaosCloth/Source/ChaosCloth/Private/ChaosCloth/ChaosClothingCacheSchema.cpp
+++ b/Engine/Plugins/ChaosCloth/Source/ChaosCloth/Private/ChaosCloth/ChaosClothingCacheSchema.cpp
@@ -46,7 +46,7 @@ void FClothingCacheSchema::RecordPostSolve(const FClothingSimulationSolver& Clot
 				const Softs::FSolverVec3& ParticleX = ParticleXs[ParticleIndex];
 
 				// Adding the vertices relative position to the particle write datas
-				PendingID.Add(ParticleIndex + GlobalOffset);
+				PendingID.Add(ParticleIndex);
 				PendingVX.Add(ParticleV.X);
 				PendingVY.Add(ParticleV.Y);
 				PendingVZ.Add(ParticleV.Z);
diff --git a/Engine/Plugins/ChaosClothAsset/ChaosClothAsset.uplugin b/Engine/Plugins/ChaosClothAsset/ChaosClothAsset.uplugin
--- a/Engine/Plugins/ChaosClothAsset/ChaosClothAsset.uplugin
+++ b/Engine/Plugins/ChaosClothAsset/ChaosClothAsset.uplugin
@@ -34,6 +34,10 @@
 		{
 			"Name": "GeometryCache",
 			"Enabled": true
+		},
+		{
+			"Name": "Dataflow",
+			"Enabled": true
 		}
 	]
 }
diff --git a/Engine/Plugins/ChaosClothAsset/Source/ChaosClothAssetEngine/ChaosClothAssetEngine.Build.cs b/Engine/Plugins/ChaosClothAsset/Source/ChaosClothAssetEngine/ChaosClothAssetEngine.Build.cs
--- a/Engine/Plugins/ChaosClothAsset/Source/ChaosClothAssetEngine/ChaosClothAssetEngine.Build.cs
+++ b/Engine/Plugins/ChaosClothAsset/Source/ChaosClothAssetEngine/ChaosClothAssetEngine.Build.cs
@@ -7,6 +7,15 @@ public class ChaosClothAssetEngine : ModuleRules
 	public ChaosClothAssetEngine(ReadOnlyTargetRules Target) : base(Target)
 	{
 		PCHUsage = ModuleRules.PCHUsageMode.UseExplicitOrSharedPCHs;
+		
+		PublicDependencyModuleNames.AddRange(
+			new string[]
+			{
+				"DataflowCore",
+				"DataflowEngine",
+				"DataflowSimulation"
+			}
+		);
 
 		PrivateDependencyModuleNames.AddRange(
 			new string[]
@@ -21,8 +30,7 @@ public class ChaosClothAssetEngine : ModuleRules
 				"ChaosCloth",
 				"ChaosCaching",
 				"ClothingSystemRuntimeCommon",
-				"ClothingSystemRuntimeInterface",
-				"DataflowEngine"
+				"ClothingSystemRuntimeInterface"
 			}
 		);
 
diff --git a/Engine/Plugins/ChaosClothAsset/Source/ChaosClothAssetEngine/Private/ChaosClothAsset/ClothAsset.cpp b/Engine/Plugins/ChaosClothAsset/Source/ChaosClothAssetEngine/Private/ChaosClothAsset/ClothAsset.cpp
--- a/Engine/Plugins/ChaosClothAsset/Source/ChaosClothAssetEngine/Private/ChaosClothAsset/ClothAsset.cpp
+++ b/Engine/Plugins/ChaosClothAsset/Source/ChaosClothAssetEngine/Private/ChaosClothAsset/ClothAsset.cpp
@@ -261,6 +261,7 @@ void UChaosClothAsset::PostEditChangeProperty(FPropertyChangedEvent& PropertyCha
 		ReregisterComponents();
 	}
 	Super::PostEditChangeProperty(PropertyChangedEvent);
+	InvalidateDataflowContents();
 }
 #endif // #if WITH_EDITOR
 
@@ -1044,3 +1045,43 @@ UAnimationAsset* UChaosClothAsset::GetPreviewSceneAnimation() const
 }
 
 #endif
+
+TObjectPtr<UDataflowBaseContent> UChaosClothAsset::CreateDataflowContent()
+{
+	TObjectPtr<UDataflowSkeletalContent> SkeletalContent = DataflowContextHelpers::CreateNewDataflowContent<UDataflowSkeletalContent>(this);
+
+	SkeletalContent->SetDataflowOwner(this);
+	SkeletalContent->SetTerminalAsset(this);
+
+	WriteDataflowContent(SkeletalContent);
+	
+	return SkeletalContent;
+}
+
+void UChaosClothAsset::WriteDataflowContent(const TObjectPtr<UDataflowBaseContent>& DataflowContent) const
+{
+	if(const TObjectPtr<UDataflowSkeletalContent> SkeletalContent = Cast<UDataflowSkeletalContent>(DataflowContent))
+	{
+		SkeletalContent->SetDataflowAsset(DataflowAsset);
+		SkeletalContent->SetDataflowTerminal(DataflowTerminal);
+		SkeletalContent->SetSkeleton(nullptr);
+
+#if WITH_EDITORONLY_DATA
+		SkeletalContent->SetAnimationAsset(GetPreviewSceneAnimation());
+		SkeletalContent->SetSkeletalMesh(GetPreviewSceneSkeletalMesh());
+#endif
+	}
+}
+
+void UChaosClothAsset::ReadDataflowContent(const TObjectPtr<UDataflowBaseContent>& DataflowContent)
+{
+	if(const TObjectPtr<UDataflowSkeletalContent> SkeletalContent = Cast<UDataflowSkeletalContent>(DataflowContent))
+	{
+#if WITH_EDITORONLY_DATA
+		PreviewSceneAnimation = SkeletalContent->GetAnimationAsset();
+		PreviewSceneSkeletalMesh = SkeletalContent->GetSkeletalMesh();
+#endif
+	}
+}
+
+
diff --git a/Engine/Plugins/ChaosClothAsset/Source/ChaosClothAssetEngine/Private/ChaosClothAsset/ClothComponent.cpp b/Engine/Plugins/ChaosClothAsset/Source/ChaosClothAssetEngine/Private/ChaosClothAsset/ClothComponent.cpp
--- a/Engine/Plugins/ChaosClothAsset/Source/ChaosClothAssetEngine/Private/ChaosClothAsset/ClothComponent.cpp
+++ b/Engine/Plugins/ChaosClothAsset/Source/ChaosClothAssetEngine/Private/ChaosClothAsset/ClothComponent.cpp
@@ -13,6 +13,7 @@
 #include "HAL/IConsoleManager.h"
 #include "Rendering/SkeletalMeshRenderData.h"
 #include "SkeletalRenderPublic.h"
+#include "Dataflow/DataflowSimulationManager.h"
 #include "PhysicsEngine/PhysicsAsset.h"
 #include "Stats/Stats.h"
 
@@ -207,6 +208,9 @@ void UChaosClothComponent::OnRegister()
 
 	// Update render visibility, so that an empty LODs doesn't unnecessarily go to render
 	UpdateVisibility();
+
+	// Register the dataflow simulation interface
+	Dataflow::RegisterSimulationInterface(this);
 }
 
 void UChaosClothComponent::OnUnregister()
@@ -220,6 +224,9 @@ void UChaosClothComponent::OnUnregister()
 	ClothOutfitInteractor->ResetProperties();
 	CollectionPropertyFacades.Empty();
 	PropertyCollections.Empty();
+
+	// Unregister the dataflow simulation interface
+	Dataflow::UnregisterSimulationInterface(this);
 }
 
 bool UChaosClothComponent::IsComponentTickEnabled() const
@@ -234,19 +241,22 @@ void UChaosClothComponent::TickComponent(float DeltaTime, enum ELevelTick TickTy
 	
 	// Tick USkinnedMeshComponent first so it will update the predicted lod
 	Super::TickComponent(DeltaTime, TickType, ThisTickFunction);
-
+	
 	// Make sure that the previous frame simulation has completed
 	HandleExistingParallelSimulation();
+	
+	if(!SimulationAsset.DataflowAsset)
+	{
+		// < This would be the right place to update the preset/use an interactor, ...etc.
 
-	// < This would be the right place to update the preset/use an interactor, ...etc.
-
-	// Update the proxy and start the simulation parallel task
-	StartNewParallelSimulation(DeltaTime);
+		// Update the proxy and start the simulation parallel task
+		StartNewParallelSimulation(DeltaTime);
 
-	// Wait in tick function for the simulation results if required
-	if (ShouldWaitForParallelSimulationInTickComponent())
-	{
-		HandleExistingParallelSimulation();
+		// Wait in tick function for the simulation results if required
+		if (ShouldWaitForParallelSimulationInTickComponent())
+		{
+			HandleExistingParallelSimulation();
+		}
 	}
 
 #if WITH_EDITOR
@@ -606,3 +616,45 @@ UChaosClothAssetInteractor* UChaosClothComponent::GetClothOutfitInteractor()
 	return ClothOutfitInteractor;
 }
 
+void UChaosClothComponent::BuildSimulationProxy()
+{
+	RecreateClothSimulationProxy();
+}
+
+void UChaosClothComponent::ResetSimulationProxy()
+{
+	ClothSimulationProxy.Reset();
+}
+
+void UChaosClothComponent::WriteToSimulation(const float DeltaTime)
+{
+	if (ClothSimulationProxy.IsValid())
+	{
+		const bool bIsSimulating = ClothSimulationProxy->PreSimulate_GameThread(DeltaTime);
+		const int32 CurrentLOD = GetPredictedLODLevel();
+
+		if (bIsSimulating && CollectionPropertyFacades.IsValidIndex(CurrentLOD) && CollectionPropertyFacades[CurrentLOD].IsValid())
+		{
+			CollectionPropertyFacades[CurrentLOD]->ClearDirtyFlags();
+		}
+	}
+}
+
+void UChaosClothComponent::ReadFromSimulation(const float DeltaTime)
+{
+	if (ClothSimulationProxy.IsValid())
+	{
+		ClothSimulationProxy->PostSimulate_GameThread();
+	}
+}
+
+FDataflowSimulationProxy* UChaosClothComponent::GetSimulationProxy()
+{
+	return ClothSimulationProxy.Get();
+}
+
+const FDataflowSimulationProxy* UChaosClothComponent::GetSimulationProxy() const
+{
+	return ClothSimulationProxy.Get();
+}
+
diff --git a/Engine/Plugins/ChaosClothAsset/Source/ChaosClothAssetEngine/Private/ChaosClothAsset/ClothSimulationProxy.cpp b/Engine/Plugins/ChaosClothAsset/Source/ChaosClothAssetEngine/Private/ChaosClothAsset/ClothSimulationProxy.cpp
--- a/Engine/Plugins/ChaosClothAsset/Source/ChaosClothAssetEngine/Private/ChaosClothAsset/ClothSimulationProxy.cpp
+++ b/Engine/Plugins/ChaosClothAsset/Source/ChaosClothAssetEngine/Private/ChaosClothAsset/ClothSimulationProxy.cpp
@@ -175,7 +175,7 @@ namespace UE::Chaos::ClothAsset
 		CompleteParallelSimulation_GameThread();
 	}
 
-	bool FClothSimulationProxy::Tick_GameThread(float DeltaTime)
+	bool FClothSimulationProxy::SetupSimulationData(float DeltaTime)
 	{
 		SCOPE_CYCLE_COUNTER(STAT_ClothSimulationProxy_TickGame);
 
@@ -224,9 +224,6 @@ namespace UE::Chaos::ClothAsset
 				CollisionSourcesProxy->ExtractCollisionData();
 			}
 
-			// Start the the cloth simulation thread
-			ParallelTask = TGraphTask<FClothSimulationProxyParallelTask>::CreateTask(nullptr, ENamedThreads::GameThread).ConstructAndDispatchWhenReady(*this);
-
 			return true;  // Simulating
 		}
 
@@ -244,6 +241,32 @@ namespace UE::Chaos::ClothAsset
 		return false;  // Not simulating
 	}
 
+	bool FClothSimulationProxy::PreSimulate_GameThread(float DeltaTime)
+	{
+		SCOPE_CYCLE_COUNTER(STAT_ClothSimulationProxy_TickGame);
+
+		const bool bIsSimulating = SetupSimulationData(DeltaTime);
+		if(bIsSimulating && ClothSimulationContext->CacheData.HasData())
+		{
+			// Start the the cloth simulation thread only when reading cache
+			ParallelTask = TGraphTask<FClothSimulationProxyParallelTask>::CreateTask(nullptr, ENamedThreads::GameThread).ConstructAndDispatchWhenReady(*this);
+		}
+		return bIsSimulating;
+	}
+
+	bool FClothSimulationProxy::Tick_GameThread(float DeltaTime)
+	{
+		SCOPE_CYCLE_COUNTER(STAT_ClothSimulationProxy_TickGame);
+
+		const bool bIsSimulating = SetupSimulationData(DeltaTime);
+		if(bIsSimulating)
+		{
+			// Start the the cloth simulation thread
+			ParallelTask = TGraphTask<FClothSimulationProxyParallelTask>::CreateTask(nullptr, ENamedThreads::GameThread).ConstructAndDispatchWhenReady(*this);
+		}
+		return bIsSimulating;
+	}
+
 	void FClothSimulationProxy::Tick()
 	{
 		using namespace ::Chaos;
diff --git a/Engine/Plugins/ChaosClothAsset/Source/ChaosClothAssetEngine/Public/ChaosClothAsset/ClothAsset.h b/Engine/Plugins/ChaosClothAsset/Source/ChaosClothAssetEngine/Public/ChaosClothAsset/ClothAsset.h
--- a/Engine/Plugins/ChaosClothAsset/Source/ChaosClothAssetEngine/Public/ChaosClothAsset/ClothAsset.h
+++ b/Engine/Plugins/ChaosClothAsset/Source/ChaosClothAssetEngine/Public/ChaosClothAsset/ClothAsset.h
@@ -6,6 +6,7 @@
 #include "ReferenceSkeleton.h"
 #include "RenderCommandFence.h"
 #include "PerQualityLevelProperties.h"
+#include "Dataflow/DataflowContent.h"
 #include "ClothAsset.generated.h"
 
 class FSkeletalMeshRenderData;
@@ -36,7 +37,7 @@ ENUM_CLASS_FLAGS(EClothAssetAsyncProperties);
  * Cloth asset for pattern based simulation.
  */
 UCLASS(hidecategories = Object, BlueprintType)
-class CHAOSCLOTHASSETENGINE_API UChaosClothAsset : public USkinnedAsset
+class CHAOSCLOTHASSETENGINE_API UChaosClothAsset : public USkinnedAsset, public IDataflowContentOwner
 {
 	GENERATED_BODY()
 public:
@@ -55,6 +56,12 @@ public:
 	virtual void GetResourceSizeEx(FResourceSizeEx& CumulativeResourceSize) override;
 	//~ End UObject interface
 
+	//~ Begin IDataflowContentOwner interface 
+	virtual TObjectPtr<UDataflowBaseContent> CreateDataflowContent() override;
+	virtual void WriteDataflowContent(const TObjectPtr<UDataflowBaseContent>& DataflowContent) const override;
+	virtual void ReadDataflowContent(const TObjectPtr<UDataflowBaseContent>& DataflowContent) override;
+	//~ End IDataflowContentOwner interface 
+
 	//~ Begin USkinnedAsset interface
 	virtual FReferenceSkeleton& GetRefSkeleton()								
 	{
diff --git a/Engine/Plugins/ChaosClothAsset/Source/ChaosClothAssetEngine/Public/ChaosClothAsset/ClothComponent.h b/Engine/Plugins/ChaosClothAsset/Source/ChaosClothAssetEngine/Public/ChaosClothAsset/ClothComponent.h
--- a/Engine/Plugins/ChaosClothAsset/Source/ChaosClothAssetEngine/Public/ChaosClothAsset/ClothComponent.h
+++ b/Engine/Plugins/ChaosClothAsset/Source/ChaosClothAssetEngine/Public/ChaosClothAsset/ClothComponent.h
@@ -3,6 +3,7 @@
 #pragma once
 
 #include "Components/SkinnedMeshComponent.h"
+#include "Interfaces/DataflowPhysicsSolver.h"
 #include "ClothComponent.generated.h"
 
 class UChaosClothAsset;
@@ -30,7 +31,7 @@ UCLASS(
 	Meta = (BlueprintSpawnableComponent, ToolTip = "Chaos cloth component."),
 	DisplayName = "Chaos cloth component",
 	HideCategories = (Object, "Mesh|SkeletalAsset", Constraints, Advanced, Cooking, Collision, Navigation))
-class CHAOSCLOTHASSETENGINE_API UChaosClothComponent : public USkinnedMeshComponent
+class CHAOSCLOTHASSETENGINE_API UChaosClothComponent : public USkinnedMeshComponent, public IDataflowPhysicsSolverInterface
 {
 	GENERATED_BODY()	
 public:
@@ -179,6 +180,18 @@ protected:
 	virtual void UnregisterOnBoneTransformsFinalizedDelegate(const FDelegateHandle& DelegateHandle) override;
 	//~ End USkinnedMeshComponent Interface
 
+	// Begin IDataflowPhysicsSolverInterface overrides
+	virtual FString GetSimulationName() const override {return GetName();};
+	virtual FDataflowSimulationAsset& GetSimulationAsset() override {return SimulationAsset;};
+	virtual const FDataflowSimulationAsset& GetSimulationAsset() const override {return SimulationAsset;};
+	virtual FDataflowSimulationProxy* GetSimulationProxy() override;
+	virtual const FDataflowSimulationProxy* GetSimulationProxy() const  override;
+	virtual void BuildSimulationProxy() override;
+	virtual void ResetSimulationProxy() override;
+	virtual void WriteToSimulation(const float DeltaTime) override;
+	virtual void ReadFromSimulation(const float DeltaTime) override;
+	// End IDataflowPhysicsSolverInterface overrides
+
 	/** Override this function for setting up custom simulation proxies when the component is registered. */
 	virtual TSharedPtr<UE::Chaos::ClothAsset::FClothSimulationProxy> CreateClothSimulationProxy();
 
@@ -202,6 +215,10 @@ private:
 	uint8 bSimulateInEditor : 1;
 #endif
 
+	/* Solver dataflow asset used to advance in time */
+	UPROPERTY(EditAnywhere, Category = ClothComponent)
+	FDataflowSimulationAsset SimulationAsset;
+	
 	/** If enabled, and the parent is another Skinned Mesh Component (e.g. another Cloth Component, Poseable Mesh Component, Skeletal Mesh Component, ...etc.), use its pose. */
 	UPROPERTY(EditAnywhere, Category = ClothComponent)
 	uint8 bUseAttachedParentAsPoseComponent : 1;
diff --git a/Engine/Plugins/ChaosClothAsset/Source/ChaosClothAssetEngine/Public/ChaosClothAsset/ClothSimulationProxy.h b/Engine/Plugins/ChaosClothAsset/Source/ChaosClothAssetEngine/Public/ChaosClothAsset/ClothSimulationProxy.h
--- a/Engine/Plugins/ChaosClothAsset/Source/ChaosClothAssetEngine/Public/ChaosClothAsset/ClothSimulationProxy.h
+++ b/Engine/Plugins/ChaosClothAsset/Source/ChaosClothAssetEngine/Public/ChaosClothAsset/ClothSimulationProxy.h
@@ -7,6 +7,7 @@
 #include "Math/Transform.h"
 #include "Templates/UniquePtr.h"
 #include "ClothingSystemRuntimeTypes.h"
+#include "Interfaces/DataflowPhysicsSolver.h"
 
 namespace Chaos
 {
@@ -33,7 +34,7 @@ namespace UE::Chaos::ClothAsset
 	 * Cloth simulation proxy.
 	 * Class used to share data between the cloth simulation and the cloth component.
 	 */
-	class CHAOSCLOTHASSETENGINE_API FClothSimulationProxy
+	class CHAOSCLOTHASSETENGINE_API FClothSimulationProxy : public FDataflowPhysicsSolverProxy
 	{
 	public:
 		explicit FClothSimulationProxy(const UChaosClothComponent& InClothComponent);
@@ -51,6 +52,12 @@ namespace UE::Chaos::ClothAsset
 		/** Wait for the parallel task to complete if one was running, and update the simulation data. */
 		void CompleteParallelSimulation_GameThread();
 
+		/** write simulation data back onto GT after the simulation is done */
+		void PostSimulate_GameThread() {WriteSimulationData();}
+		
+		/** setup simulation data from GT before the simulation starts */
+        bool PreSimulate_GameThread(float DeltaTime);
+
 		/**
 		 * Return a map of all simulation data as used by the skeletal rendering code.
 		 * The map key is the rendering section's cloth index as set in FSkelMeshRenderSection::CorrespondClothAssetIndex,
@@ -75,6 +82,7 @@ namespace UE::Chaos::ClothAsset
 	protected:
 		void Tick();
 		void WriteSimulationData();
+		bool SetupSimulationData(float DeltaTime);
 		void InitializeConfigs();
 		void FillSimulationContext(float DeltaTime, bool bIsInitialization = false);
 
@@ -82,6 +90,10 @@ namespace UE::Chaos::ClothAsset
 		bool ShouldEnableSolver(bool bSolverCurrentlyEnabled) const;
 		void UpdateClothLODs();
 
+		// Begin FDataflowPhysicsSolverProxy overrides
+		virtual void AdvanceSolverDatas(const float DeltaTime) override {Tick();}
+		// End FDataflowPhysicsSolverProxy overrides
+
 		// Internal physics thread object
 		friend class FClothSimulationProxyParallelTask;
 
diff --git a/Engine/Plugins/ChaosClothAssetEditor/Source/ChaosClothAssetEditor/Private/ChaosClothAsset/AssetDefinition_ClothAsset.cpp b/Engine/Plugins/ChaosClothAssetEditor/Source/ChaosClothAssetEditor/Private/ChaosClothAsset/AssetDefinition_ClothAsset.cpp
--- a/Engine/Plugins/ChaosClothAssetEditor/Source/ChaosClothAssetEditor/Private/ChaosClothAsset/AssetDefinition_ClothAsset.cpp
+++ b/Engine/Plugins/ChaosClothAssetEditor/Source/ChaosClothAssetEditor/Private/ChaosClothAsset/AssetDefinition_ClothAsset.cpp
@@ -15,10 +15,17 @@
 #include "Misc/WarnIfAssetsLoadedInScope.h"
 #include "Dialog/SMessageDialog.h"
 #include "ChaosClothAsset/TerminalNode.h"
+#include "Dataflow/DataflowEditor.h"
 #include "Dataflow/DataflowSNode.h"
 
 #define LOCTEXT_NAMESPACE "AssetDefinition_ClothAsset"
 
+namespace UE::Chaos::ClothAsset::Private
+{
+static int32 bEnableClothDataflowEditor = false;
+static FAutoConsoleVariableRef CVarEnableClothDataflowEditor(TEXT("p.ChaosCloth.EnableDataflowEditor"), bEnableClothDataflowEditor, TEXT("Enable the use of the core dataflow editor for cloth asset (WIP)"));
+}
+
 namespace ClothAssetDefinitionHelpers
 {
 	// Return true if we should proceed, false if we should re-open the dialog
@@ -174,9 +181,6 @@ EAssetCommandResult UAssetDefinition_ClothAsset::OpenAssets(const FAssetOpenArgs
 
 	if (ClothObjects.Num() > 0)
 	{
-		UAssetEditorSubsystem* const AssetEditorSubsystem = GEditor->GetEditorSubsystem<UAssetEditorSubsystem>();
-		UChaosClothAssetEditor* const AssetEditor = NewObject<UChaosClothAssetEditor>(AssetEditorSubsystem, NAME_None, RF_Transient);
-
 		// Validate the asset
 		UChaosClothAsset* const ClothAsset = CastChecked<UChaosClothAsset>(ClothObjects[0]);
 		if (!ClothAsset->GetDataflow())
@@ -186,13 +190,20 @@ EAssetCommandResult UAssetDefinition_ClothAsset::OpenAssets(const FAssetOpenArgs
 				ClothAsset->SetDataflow(NewDataflowAsset);
 			}
 		}
-
-		TArray<TObjectPtr<UObject>> Objects;
-		for (UChaosClothAsset* const ClothObject : ClothObjects)
+		
+		UAssetEditorSubsystem* const AssetEditorSubsystem = GEditor->GetEditorSubsystem<UAssetEditorSubsystem>();
+		if(!UE::Chaos::ClothAsset::Private::bEnableClothDataflowEditor)
+		{
+			UChaosClothAssetEditor* const AssetEditor = NewObject<UChaosClothAssetEditor>(AssetEditorSubsystem, NAME_None, RF_Transient);
+			AssetEditor->Initialize({ClothObjects[0] });
+		}
+		else
 		{
-			Objects.Add(ClothObject);
+			UDataflowEditor* const AssetEditor = NewObject<UDataflowEditor>(AssetEditorSubsystem, NAME_None, RF_Transient);
+            const TSubclassOf<AActor> ActorClass = StaticLoadClass(AActor::StaticClass(), nullptr,
+            	TEXT("/ChaosClothAsset/BP_ClothPreview.BP_ClothPreview_C"), nullptr, LOAD_None, nullptr);
+			AssetEditor->Initialize({ ClothObjects[0] }, ActorClass);
 		}
-		AssetEditor->Initialize(Objects);
 
 		return EAssetCommandResult::Handled;
 	}
diff --git a/Engine/Plugins/Experimental/Dataflow/Source/DataflowSimulation/Private/Dataflow/DataflowSimulationManager.cpp b/Engine/Plugins/Experimental/Dataflow/Source/DataflowSimulation/Private/Dataflow/DataflowSimulationManager.cpp
--- a/Engine/Plugins/Experimental/Dataflow/Source/DataflowSimulation/Private/Dataflow/DataflowSimulationManager.cpp
+++ b/Engine/Plugins/Experimental/Dataflow/Source/DataflowSimulation/Private/Dataflow/DataflowSimulationManager.cpp
@@ -19,20 +19,20 @@ static FAutoConsoleTaskPriority DataflowSimulationTaskPriority(
 
 namespace Dataflow
 {
-	enum class ESimulationThreadingMode : uint8
-	{
-		GameThread,
-		BlockingThread,
-		AsyncThread,
-	};
-	
-	namespace CVars
-	{
-		/** Simulation threading mode */
-		int32  DataflowSimulationThreadingMode = (uint8)ESimulationThreadingMode::AsyncThread;
-		FAutoConsoleVariableRef CVarDataflowSimulationThreadingMode(TEXT("p.Dataflow.Simulation.ThreadingMode"), DataflowSimulationThreadingMode,
-			TEXT("0 : run simulation on GT | 1 : run simulation on PT (GT is blocked in manager Tick) | 2 : run simulation on PT (GT will be blocked at the end of the world tick)"));
-	}
+enum class ESimulationThreadingMode : uint8
+{
+	GameThread,
+	BlockingThread,
+	AsyncThread,
+};
+
+namespace CVars
+{
+	/** Simulation threading mode */
+	int32  DataflowSimulationThreadingMode = (uint8)ESimulationThreadingMode::AsyncThread;
+	FAutoConsoleVariableRef CVarDataflowSimulationThreadingMode(TEXT("p.Dataflow.Simulation.ThreadingMode"), DataflowSimulationThreadingMode,
+		TEXT("0 : run simulation on GT | 1 : run simulation on PT (GT is blocked in manager Tick) | 2 : run simulation on PT (GT will be blocked at the end of the world tick)"));
+}
 	
 class FDataflowSimulationTask
 {
@@ -80,15 +80,6 @@ private:
 	/** World simulation time  */
 	float SimulationTime;
 };
-}
-
-FDelegateHandle UDataflowSimulationManager::OnObjectPropertyChangedHandle;
-FDelegateHandle UDataflowSimulationManager::OnWorldTickEndHandle;
-FDelegateHandle UDataflowSimulationManager::OnCreatePhysicsStateHandle;
-FDelegateHandle UDataflowSimulationManager::OnDestroyPhysicsStateHandle;
-
-UDataflowSimulationManager::UDataflowSimulationManager()
-{}
 
 inline void PreSimulationTick(const TObjectPtr<UObject>& SimulationWorld, const float SimulationTime, const float DeltaTime)
 {
@@ -96,7 +87,7 @@ inline void PreSimulationTick(const TObjectPtr<UObject>& SimulationWorld, const
 	{
 		TArray<AActor*> Actors;
 		UGameplayStatics::GetAllActorsWithInterface(SimulationWorld, UDataflowSimulationActor::StaticClass(), Actors);
-	
+
 		for (AActor* CurrentActor : Actors)
 		{
 			IDataflowSimulationActor::Execute_PreDataflowSimulationTick(CurrentActor, SimulationTime, DeltaTime);
@@ -110,7 +101,7 @@ inline void PostSimulationTick(const TObjectPtr<UObject>& SimulationWorld, const
 	{
 		TArray<AActor*> Actors;
 		UGameplayStatics::GetAllActorsWithInterface(SimulationWorld, UDataflowSimulationActor::StaticClass(), Actors);
-	
+
 		for (AActor* CurrentActor : Actors)
 		{
 			IDataflowSimulationActor::Execute_PostDataflowSimulationTick(CurrentActor, SimulationTime, DeltaTime);
@@ -118,6 +109,55 @@ inline void PostSimulationTick(const TObjectPtr<UObject>& SimulationWorld, const
 	}
 }
 
+void RegisterSimulationInterface(const TObjectPtr<UObject>& SimulationObject)
+{
+	if(IDataflowSimulationInterface* SimulationInterface = Cast<IDataflowSimulationInterface>(SimulationObject))
+	{
+		if(SimulationInterface->GetSimulationAsset().DataflowAsset)
+		{
+			FDataflowSimulationProxy* SimulationProxy = SimulationInterface->GetSimulationProxy();
+			
+			if(!SimulationProxy || (SimulationProxy && !SimulationProxy->IsValid()))
+			{
+				// Build the simulation proxy
+				SimulationInterface->BuildSimulationProxy();
+			}
+
+			// Register the simulation interface to the manager
+			SimulationInterface->RegisterManagerInterface(SimulationObject->GetWorld());
+		}
+	}
+}
+
+void UnregisterSimulationInterface(const TObjectPtr<UObject>& SimulationObject)
+{
+	if(IDataflowSimulationInterface* SimulationInterface = Cast<IDataflowSimulationInterface>(SimulationObject))
+	{
+		if(SimulationInterface->GetSimulationAsset().DataflowAsset)
+		{
+			FDataflowSimulationProxy* SimulationProxy = SimulationInterface->GetSimulationProxy();
+			
+			if(SimulationProxy && SimulationProxy->IsValid())
+			{
+				// Reset the simulation proxy
+				SimulationInterface->ResetSimulationProxy();
+			}
+
+			// Unregister the simulation interface from the manager
+			SimulationInterface->UnregisterManagerInterface(SimulationObject->GetWorld());
+		}
+	}
+}
+}
+
+FDelegateHandle UDataflowSimulationManager::OnObjectPropertyChangedHandle;
+FDelegateHandle UDataflowSimulationManager::OnWorldTickEndHandle;
+FDelegateHandle UDataflowSimulationManager::OnCreatePhysicsStateHandle;
+FDelegateHandle UDataflowSimulationManager::OnDestroyPhysicsStateHandle;
+
+UDataflowSimulationManager::UDataflowSimulationManager()
+{}
+
 void UDataflowSimulationManager::Tick(float DeltaTime)
 {
 	Super::Tick(DeltaTime);
@@ -182,42 +222,12 @@ void UDataflowSimulationManager::OnStartup()
 
 	OnCreatePhysicsStateHandle = UActorComponent::GlobalCreatePhysicsDelegate.AddLambda([](UActorComponent* ActorComponent)
 	{
-		if(IDataflowSimulationInterface* SimulationInterface = Cast<IDataflowSimulationInterface>(ActorComponent))
-		{
-			if(SimulationInterface->GetSimulationAsset().DataflowAsset)
-			{
-				FDataflowSimulationProxy* SimulationProxy = SimulationInterface->GetSimulationProxy();
-				
-				if(!SimulationProxy || (SimulationProxy && !SimulationProxy->IsValid()))
-				{
-					// Build the simulation proxy
-					SimulationInterface->BuildSimulationProxy();
-				}
-
-				// Register the simulation interface to the manager
-				SimulationInterface->RegisterManagerInterface(ActorComponent->GetWorld());
-			}
-		}
+		Dataflow::RegisterSimulationInterface(ActorComponent);
 	});
 
 	OnDestroyPhysicsStateHandle = UActorComponent::GlobalDestroyPhysicsDelegate.AddLambda([](UActorComponent* ActorComponent)
 	{
-		if(IDataflowSimulationInterface* SimulationInterface = Cast<IDataflowSimulationInterface>(ActorComponent))
-		{
-			if(SimulationInterface->GetSimulationAsset().DataflowAsset)
-			{
-				FDataflowSimulationProxy* SimulationProxy = SimulationInterface->GetSimulationProxy();
-				
-				if(SimulationProxy && SimulationProxy->IsValid())
-				{
-					// Reset the simulation proxy
-					SimulationInterface->ResetSimulationProxy();
-				}
-
-				// Unregister the simulation interface from the manager
-				SimulationInterface->UnregisterManagerInterface(ActorComponent->GetWorld());
-			}
-		}
+		Dataflow::UnregisterSimulationInterface(ActorComponent);
 	});
 
 #if WITH_EDITOR
@@ -248,10 +258,50 @@ void UDataflowSimulationManager::OnShutdown()
 	UActorComponent::GlobalDestroyPhysicsDelegate.Remove(OnDestroyPhysicsStateHandle);
 }
 
+void RegisterSimulationInterface(const TObjectPtr<UObject>& SimulationObject)
+{
+	if(IDataflowSimulationInterface* SimulationInterface = Cast<IDataflowSimulationInterface>(SimulationObject))
+	{
+		if(SimulationInterface->GetSimulationAsset().DataflowAsset)
+		{
+			FDataflowSimulationProxy* SimulationProxy = SimulationInterface->GetSimulationProxy();
+				
+			if(!SimulationProxy || (SimulationProxy && !SimulationProxy->IsValid()))
+			{
+				// Build the simulation proxy
+				SimulationInterface->BuildSimulationProxy();
+			}
+
+			// Register the simulation interface to the manager
+			SimulationInterface->RegisterManagerInterface(SimulationObject->GetWorld());
+		}
+	}
+}
+
+void UnregisterSimulationInterface(const TObjectPtr<UObject>& SimulationObject)
+{
+	if(IDataflowSimulationInterface* SimulationInterface = Cast<IDataflowSimulationInterface>(SimulationObject))
+	{
+		if(SimulationInterface->GetSimulationAsset().DataflowAsset)
+		{
+			FDataflowSimulationProxy* SimulationProxy = SimulationInterface->GetSimulationProxy();
+				
+			if(SimulationProxy && SimulationProxy->IsValid())
+			{
+				// Reset the simulation proxy
+				SimulationInterface->ResetSimulationProxy();
+			}
+
+			// Unregister the simulation interface from the manager
+			SimulationInterface->UnregisterManagerInterface(SimulationObject->GetWorld());
+		}
+	}
+}
+
 void UDataflowSimulationManager::WriteSimulationData(const float DeltaTime)
 {
 	// Pre-simulation callback that could be used in BP before the simulation
-	PreSimulationTick(GetWorld(), GetWorld()->GetTimeSeconds(), DeltaTime);
+	Dataflow::PreSimulationTick(GetWorld(), GetWorld()->GetTimeSeconds(), DeltaTime);
 	
 	for(const TPair<TObjectPtr<UDataflow>, Dataflow::FDataflowSimulationData>& DataflowData : SimulationData)
 	{
@@ -314,7 +364,7 @@ void UDataflowSimulationManager::ReadSimulationData(const float DeltaTime)
 		bStepSimulationScene = false;
 	}
 	// Post-simulation callback that could be used in BP after the simulation
-	PostSimulationTick(GetWorld(), GetWorld()->GetTimeSeconds(), DeltaTime);
+	Dataflow::PostSimulationTick(GetWorld(), GetWorld()->GetTimeSeconds(), DeltaTime);
 }
 
 void UDataflowSimulationManager::AdvanceSimulationData(const float DeltaTime, const float SimulationTime)
diff --git a/Engine/Plugins/Experimental/Dataflow/Source/DataflowSimulation/Public/Dataflow/DataflowSimulationManager.h b/Engine/Plugins/Experimental/Dataflow/Source/DataflowSimulation/Public/Dataflow/DataflowSimulationManager.h
--- a/Engine/Plugins/Experimental/Dataflow/Source/DataflowSimulation/Public/Dataflow/DataflowSimulationManager.h
+++ b/Engine/Plugins/Experimental/Dataflow/Source/DataflowSimulation/Public/Dataflow/DataflowSimulationManager.h
@@ -38,6 +38,8 @@ namespace Dataflow
 			return true;
 		}
 	};
+	DATAFLOWSIMULATION_API void RegisterSimulationInterface(const TObjectPtr<UObject>& SimulationObject);
+	DATAFLOWSIMULATION_API void UnregisterSimulationInterface(const TObjectPtr<UObject>& SimulationObject);
 }
 
 UCLASS()
