Fix GT regression due to dataflow simulation interface

diff --git a/Engine/Plugins/Experimental/Dataflow/Source/DataflowSimulation/Private/Dataflow/DataflowSimulationManager.cpp b/Engine/Plugins/Experimental/Dataflow/Source/DataflowSimulation/Private/Dataflow/DataflowSimulationManager.cpp
--- a/Engine/Plugins/Experimental/Dataflow/Source/DataflowSimulation/Private/Dataflow/DataflowSimulationManager.cpp
+++ b/Engine/Plugins/Experimental/Dataflow/Source/DataflowSimulation/Private/Dataflow/DataflowSimulationManager.cpp
@@ -301,7 +301,7 @@ void UnregisterSimulationInterface(const TObjectPtr<UObject>& SimulationObject)
 void UDataflowSimulationManager::WriteSimulationData(const float DeltaTime)
 {
 	// Pre-simulation callback that could be used in BP before the simulation
-	Dataflow::PreSimulationTick(GetWorld(), GetWorld()->GetTimeSeconds(), DeltaTime);
+	//Dataflow::PreSimulationTick(GetWorld(), GetWorld()->GetTimeSeconds(), DeltaTime);
 	
 	for(const TPair<TObjectPtr<UDataflow>, Dataflow::FDataflowSimulationData>& DataflowData : SimulationData)
 	{
@@ -364,7 +364,7 @@ void UDataflowSimulationManager::ReadSimulationData(const float DeltaTime)
 		bStepSimulationScene = false;
 	}
 	// Post-simulation callback that could be used in BP after the simulation
-	Dataflow::PostSimulationTick(GetWorld(), GetWorld()->GetTimeSeconds(), DeltaTime);
+	//Dataflow::PostSimulationTick(GetWorld(), GetWorld()->GetTimeSeconds(), DeltaTime);
 }
 
 void UDataflowSimulationManager::AdvanceSimulationData(const float DeltaTime, const float SimulationTime)
