Dataflow : Wireframe selection

diff --git a/Engine/Plugins/Experimental/Dataflow/Source/DataflowEditor/Private/Dataflow/DataflowConstructionScene.cpp b/Engine/Plugins/Experimental/Dataflow/Source/DataflowEditor/Private/Dataflow/DataflowConstructionScene.cpp
--- a/Engine/Plugins/Experimental/Dataflow/Source/DataflowEditor/Private/Dataflow/DataflowConstructionScene.cpp
+++ b/Engine/Plugins/Experimental/Dataflow/Source/DataflowEditor/Private/Dataflow/DataflowConstructionScene.cpp
@@ -277,10 +277,10 @@ TObjectPtr<UDynamicMeshComponent>& FDataflowConstructionScene::AddDynamicMeshCom
 
 void FDataflowConstructionScene::AddWireframeMeshElementsVisualizer()
 {
-	ensure(WireframeElements.Num()==0);
-	for(FRenderElement Elem : DynamicMeshComponents)
+	ensure(WireframeElements.Num() == 0);
+	for (UDynamicMeshComponent* Elem : MeshComponentsForWireframeRendering)
 	{
-		if( TObjectPtr<UDataflowEditorCollectionComponent> DynamicMeshComponent = Cast<UDataflowEditorCollectionComponent>(Elem.Value) )
+		if (TObjectPtr<UDataflowEditorCollectionComponent> DynamicMeshComponent = Cast<UDataflowEditorCollectionComponent>(Elem))
 		{
 			// Set up the wireframe display of the rest space mesh.
 
