Dataflow - Fixed a crash when using more than 1024 chars in a node's name.

#rnx

diff --git a/Engine/Plugins/Experimental/Dataflow/Source/DataflowEditor/Private/Dataflow/DataflowEditorCommands.cpp b/Engine/Plugins/Experimental/Dataflow/Source/DataflowEditor/Private/Dataflow/DataflowEditorCommands.cpp
--- a/Engine/Plugins/Experimental/Dataflow/Source/DataflowEditor/Private/Dataflow/DataflowEditorCommands.cpp
+++ b/Engine/Plugins/Experimental/Dataflow/Source/DataflowEditor/Private/Dataflow/DataflowEditorCommands.cpp
@@ -258,13 +258,20 @@ void FDataflowEditorCommands::EvaluateTerminalNode(Dataflow::FContext& Context,
 
 bool FDataflowEditorCommands::OnNodeVerifyTitleCommit(const FText& NewText, UEdGraphNode* GraphNode, FText& OutErrorMessage)
 {
+	const FString NewString = NewText.ToString();
+	const int32 NewStringLen = NewString.Len();
+	if (NewStringLen >= NAME_SIZE)
+	{
+		OutErrorMessage = FText::FromString(FString::Printf(TEXT("Name length is %d characters which exceeds the maximum allowed of %d characters"), NewStringLen, NAME_SIZE - 1));
+		return false;
+	}
 	if (GraphNode)
 	{
 		if (UDataflowEdNode* DataflowNode = Cast<UDataflowEdNode>(GraphNode))
 		{
 			if (TSharedPtr<Dataflow::FGraph> Graph = DataflowNode->GetDataflowGraph())
 			{
-				if( Graph->FindBaseNode(FName(NewText.ToString())).Get()==nullptr )
+				if( Graph->FindBaseNode(FName(NewString)).Get()==nullptr )
 				{
 					return true;
 				}
@@ -275,7 +282,7 @@ bool FDataflowEditorCommands::OnNodeVerifyTitleCommit(const FText& NewText, UEdG
 			return true;
 		}
 	}
-	OutErrorMessage = FText::FromString(FString::Printf(TEXT("Non-unique name for graph node (%s)"), *NewText.ToString()));
+	OutErrorMessage = FText::FromString(FString::Printf(TEXT("Non-unique name for graph node (%s)"), *NewString));
 	return false;
 }
 
