Dataflow : Construction screen render refresh based on context state only.

diff --git a/Engine/Plugins/Experimental/Dataflow/Source/DataflowEditor/Private/Dataflow/DataflowConstructionScene.cpp b/Engine/Plugins/Experimental/Dataflow/Source/DataflowEditor/Private/Dataflow/DataflowConstructionScene.cpp
--- a/Engine/Plugins/Experimental/Dataflow/Source/DataflowEditor/Private/Dataflow/DataflowConstructionScene.cpp
+++ b/Engine/Plugins/Experimental/Dataflow/Source/DataflowEditor/Private/Dataflow/DataflowConstructionScene.cpp
@@ -85,15 +85,11 @@ void FDataflowConstructionScene::TickDataflowScene(const float DeltaSeconds)
 		{
 			if (const UDataflow* Dataflow = EditorContent->GetDataflowAsset())
 			{
-				const Dataflow::FTimestamp SystemTimestamp = LatestTimestamp(Dataflow, DataflowContext.Get());
-				if (SystemTimestamp >= EditorContent->GetLastModifiedTimestamp() || EditorContent->IsConstructionDirty())
+				const Dataflow::FTimestamp SystemTimestamp = DataflowContext->GetTimestamp().Value;
+				if (LastRenderedTimestamp < SystemTimestamp || EditorContent->IsConstructionDirty())
 				{
-					EditorContent->SetLastModifiedTimestamp(SystemTimestamp.Value + 1);
-
-					if (EditorContent->IsConstructionDirty())
-					{
-						UpdateConstructionScene();
-					}
+					LastRenderedTimestamp = SystemTimestamp;
+					UpdateConstructionScene();
 				}
 			}
 		}
diff --git a/Engine/Plugins/Experimental/Dataflow/Source/DataflowEditor/Private/Dataflow/DataflowEditorToolkit.cpp b/Engine/Plugins/Experimental/Dataflow/Source/DataflowEditor/Private/Dataflow/DataflowEditorToolkit.cpp
--- a/Engine/Plugins/Experimental/Dataflow/Source/DataflowEditor/Private/Dataflow/DataflowEditorToolkit.cpp
+++ b/Engine/Plugins/Experimental/Dataflow/Source/DataflowEditor/Private/Dataflow/DataflowEditorToolkit.cpp
@@ -923,6 +923,8 @@ void FDataflowEditorToolkit::OnNodeSelectionChanged(const TSet<UObject*>& InNewS
 				SelectComponentsInView(PrimarySelection);
 			}
 		}
+
+		EditorContent->SetConstructionDirty(true);
 	}
 
 	//
diff --git a/Engine/Plugins/Experimental/Dataflow/Source/DataflowEditor/Public/Dataflow/DataflowConstructionScene.h b/Engine/Plugins/Experimental/Dataflow/Source/DataflowEditor/Public/Dataflow/DataflowConstructionScene.h
--- a/Engine/Plugins/Experimental/Dataflow/Source/DataflowEditor/Public/Dataflow/DataflowConstructionScene.h
+++ b/Engine/Plugins/Experimental/Dataflow/Source/DataflowEditor/Public/Dataflow/DataflowConstructionScene.h
@@ -3,6 +3,7 @@
 #pragma once
 
 #include "Dataflow/DataflowEditorPreviewSceneBase.h"
+#include "Dataflow/DataflowNodeParameters.h"
 
 class UDataflowEditor;
 class USkeletalMeshComponent;
@@ -88,6 +89,9 @@ private:
 	/** Property objects (visible or not) that get ticked. [ticked] */
 	TArray<TObjectPtr<UInteractiveToolPropertySet>> PropertyObjectsToTick;
 
+	/** Last Rendered Timestamp */
+	Dataflow::FTimestamp LastRenderedTimestamp = Dataflow::FTimestamp::Invalid;
+
 	/** Show the wireframe on render */ 
 	bool bConstructionViewWireframe = true;
 };
