Dataflow: allow the FDataflowCollectionAddScalarVertexPropertyNode to assign weights to different vertex targets depending on the schema of the Collection passing through the node. This will require developers of different Collection types to opt in by registering their vertex types.

diff --git a/Engine/Plugins/ChaosClothAssetEditor/Source/ChaosClothAssetDataflowNodes/ChaosClothAssetDataflowNodes.Build.cs b/Engine/Plugins/ChaosClothAssetEditor/Source/ChaosClothAssetDataflowNodes/ChaosClothAssetDataflowNodes.Build.cs
--- a/Engine/Plugins/ChaosClothAssetEditor/Source/ChaosClothAssetDataflowNodes/ChaosClothAssetDataflowNodes.Build.cs
+++ b/Engine/Plugins/ChaosClothAssetEditor/Source/ChaosClothAssetDataflowNodes/ChaosClothAssetDataflowNodes.Build.cs
@@ -27,6 +27,7 @@ public class ChaosClothAssetDataflowNodes : ModuleRules
 				"DataflowCore",
 				"DataflowEditor",
 				"DataflowEngine",
+				"DataflowNodes",
 				"DesktopWidgets",  // For SFilePathPicker
 				"DetailCustomizations",
 				"DynamicMesh",
diff --git a/Engine/Plugins/ChaosClothAssetEditor/Source/ChaosClothAssetDataflowNodes/Private/ChaosClothAsset/ChaosClothAssetDataflowNodesModule.cpp b/Engine/Plugins/ChaosClothAssetEditor/Source/ChaosClothAssetDataflowNodes/Private/ChaosClothAsset/ChaosClothAssetDataflowNodesModule.cpp
--- a/Engine/Plugins/ChaosClothAssetEditor/Source/ChaosClothAssetDataflowNodes/Private/ChaosClothAsset/ChaosClothAssetDataflowNodesModule.cpp
+++ b/Engine/Plugins/ChaosClothAssetEditor/Source/ChaosClothAssetDataflowNodes/Private/ChaosClothAsset/ChaosClothAssetDataflowNodesModule.cpp
@@ -4,6 +4,7 @@
 #include "ChaosClothAsset/AddWeightMapNode.h"
 #include "ChaosClothAsset/AttributeNode.h"
 #include "ChaosClothAsset/BindToRootBoneNode.h"
+#include "ChaosClothAsset/ClothCollectionGroup.h"
 #include "ChaosClothAsset/ColorScheme.h"
 #include "ChaosClothAsset/CopySimulationToRenderMeshNode.h"
 #include "ChaosClothAsset/SelectionGroupCustomization.h"
@@ -66,6 +67,7 @@
 #include "ChaosClothAsset/ConnectableValueCustomization.h"
 #include "ChaosClothAsset/ConnectableValue.h"
 #include "ChaosClothAsset/ClothDataflowViewModes.h"
+#include "Dataflow/DataflowCollectionAddScalarVertexPropertyNode.h"
 #include "Dataflow/DataflowNodeColorsRegistry.h"
 #include "Dataflow/DataflowNodeFactory.h"
 #include "Dataflow/DataflowRenderingFactory.h"
@@ -219,10 +221,6 @@ namespace UE::Chaos::ClothAsset
 						Normals = ClothFacade.GetRenderNormal();
 						// TODO: Get materials from the render mesh
 					}
-					else
-					{
-						checkf(false, TEXT("Invalid View Mode for FClothCollection rendering"));
-					}
 
 					TArray<FLinearColor> Colors;
 					Colors.Init(FLinearColor::Gray, Vertices.Num());	// TODO: Choose a vertex color
@@ -234,7 +232,7 @@ namespace UE::Chaos::ClothAsset
 			}
 		};
 
-		Dataflow::FRenderKey FClothSurfaceRenderCallbacks::RenderKey = { "SurfaceRender", FName("FClothCollection") };
+		Dataflow::FRenderKey FClothSurfaceRenderCallbacks::RenderKey = { TEXT("SurfaceRender"), FName("FClothCollection") };
 
 		static void RegisterRenderingCallbacks()
 		{
@@ -254,6 +252,34 @@ namespace UE::Chaos::ClothAsset
 			Dataflow::FRenderingViewModeFactory::GetInstance().DeregisterViewMode(FClothRenderViewMode::Name);
 		}
 
+
+
+		class FClothCollectionAddScalarVertexPropertyCallbacks : public IDataflowAddScalarVertexPropertyCallbacks
+		{
+		public:
+
+			const static FName Name;
+
+			virtual ~FClothCollectionAddScalarVertexPropertyCallbacks() override = default;
+
+			virtual FName GetName() const override
+			{
+				return Name;
+			}
+
+			virtual TArray<FName> GetTargetGroupNames() const override
+			{
+				return { ClothCollectionGroup::SimVertices2D, ClothCollectionGroup::SimVertices3D, ClothCollectionGroup::RenderVertices };
+			}
+
+			virtual TArray<Dataflow::FRenderingParameter> GetRenderingParameters() const override
+			{
+				return { {TEXT("SurfaceRender"), FName("FClothCollection"), {TEXT("Collection")} } };
+			}
+		};
+
+		const FName FClothCollectionAddScalarVertexPropertyCallbacks::Name = FName(TEXT("FClothCollectionAddScalarVertexPropertyCallbacks"));
+
 	}  // End namespace Private
 
 	class FChaosClothAssetDataflowNodesModule : public IModuleInterface
@@ -286,6 +312,8 @@ namespace UE::Chaos::ClothAsset
 			}
 
 			Private::RegisterRenderingCallbacks();
+
+			DataflowAddScalarVertexPropertyCallbackRegistry::Get().RegisterCallbacks(MakeUnique<Private::FClothCollectionAddScalarVertexPropertyCallbacks>());
 		}
 
 		virtual void ShutdownModule() override
@@ -314,6 +342,8 @@ namespace UE::Chaos::ClothAsset
 					PropertyModule->UnregisterCustomPropertyTypeLayout(FChaosClothAssetImportedIntValue::StaticStruct()->GetFName());
 				}
 			}
+
+			DataflowAddScalarVertexPropertyCallbackRegistry::Get().DeregisterCallbacks(Private::FClothCollectionAddScalarVertexPropertyCallbacks::Name);
 		}
 	};
 }  // End namespace UE::Chaos::ClothAsset
diff --git a/Engine/Plugins/Experimental/Dataflow/Source/DataflowEditor/Private/Dataflow/DataflowEditorModule.cpp b/Engine/Plugins/Experimental/Dataflow/Source/DataflowEditor/Private/Dataflow/DataflowEditorModule.cpp
--- a/Engine/Plugins/Experimental/Dataflow/Source/DataflowEditor/Private/Dataflow/DataflowEditorModule.cpp
+++ b/Engine/Plugins/Experimental/Dataflow/Source/DataflowEditor/Private/Dataflow/DataflowEditorModule.cpp
@@ -9,6 +9,8 @@
 #include "Dataflow/DataflowEditorCommands.h"
 #include "Dataflow/DataflowEngineRendering.h"
 #include "Dataflow/DataflowSNodeFactories.h"
+#include "Dataflow/ScalarVertexPropertyGroupCustomization.h"
+#include "Dataflow/DataflowCollectionAddScalarVertexPropertyNode.h"
 
 #include "CoreMinimal.h"
 #include "EdGraphUtilities.h"
@@ -25,12 +27,24 @@ void FDataflowEditorModule::StartupModule()
 	
 	FPropertyEditorModule& PropertyEditorModule = FModuleManager::LoadModuleChecked<FPropertyEditorModule>(TEXT("PropertyEditor"));
 
+	// Register type customizations
+	if (FPropertyEditorModule* const PropertyModule = FModuleManager::GetModulePtr<FPropertyEditorModule>("PropertyEditor"))
+	{
+		PropertyModule->RegisterCustomPropertyTypeLayout(FScalarVertexPropertyGroup::StaticStruct()->GetFName(), FOnGetPropertyTypeCustomizationInstance::CreateStatic(&Dataflow::FScalarVertexPropertyGroupCustomization::MakeInstance));
+	}
+
 	Dataflow::RenderingCallbacks();
 }
 
 void FDataflowEditorModule::ShutdownModule()
 {	
 	FEditorModeRegistry::Get().UnregisterMode(UDataflowEditorMode::EM_DataflowEditorModeId);
+
+	// Deregister type customizations
+	if (FPropertyEditorModule* const PropertyModule = FModuleManager::GetModulePtr<FPropertyEditorModule>("PropertyEditor"))
+	{
+		PropertyModule->UnregisterCustomPropertyTypeLayout(FScalarVertexPropertyGroup::StaticStruct()->GetFName());
+	}
 }
 
 IMPLEMENT_MODULE(FDataflowEditorModule, DataflowEditor)
diff --git a/Engine/Plugins/Experimental/Dataflow/Source/DataflowEditor/Private/Dataflow/DataflowEditorToolkit.cpp b/Engine/Plugins/Experimental/Dataflow/Source/DataflowEditor/Private/Dataflow/DataflowEditorToolkit.cpp
--- a/Engine/Plugins/Experimental/Dataflow/Source/DataflowEditor/Private/Dataflow/DataflowEditorToolkit.cpp
+++ b/Engine/Plugins/Experimental/Dataflow/Source/DataflowEditor/Private/Dataflow/DataflowEditorToolkit.cpp
@@ -916,6 +916,13 @@ void FDataflowEditorToolkit::OnNodeSelectionChanged(const TSet<UObject*>& InNewS
 				}
 
 				EditorContent->SetPrimarySelectedNode(PrimarySelection);
+
+				// Call the node's OnSelected function. Some nodes use this to cache information from the inputs (e.g. FDataflowCollectionAddScalarVertexPropertyNode::CachedCollectionGroupNames)
+				TSharedPtr<Dataflow::FEngineContext> DataflowContext = EditorContent->GetDataflowContext();
+				if (PrimarySelection && DataflowContext.IsValid())
+				{
+					PrimarySelection->GetDataflowNode()->OnSelected(*DataflowContext);
+				}
 			}
 
 			if (GetDataflowGraphEditor()->IsAltDown())
diff --git a/Engine/Plugins/Experimental/Dataflow/Source/DataflowEditor/Private/Dataflow/ScalarVertexPropertyGroupCustomization.cpp b/Engine/Plugins/Experimental/Dataflow/Source/DataflowEditor/Private/Dataflow/ScalarVertexPropertyGroupCustomization.cpp
new file mode 100644
--- /dev/null
+++ b/Engine/Plugins/Experimental/Dataflow/Source/DataflowEditor/Private/Dataflow/ScalarVertexPropertyGroupCustomization.cpp
@@ -0,0 +1,185 @@
+// Copyright Epic Games, Inc. All Rights Reserved.
+
+#include "ScalarVertexPropertyGroupCustomization.h"
+#include "Dataflow/DataflowCollectionAddScalarVertexPropertyNode.h"
+#include "Widgets/Input/SEditableTextBox.h"
+#include "Widgets/Input/SComboButton.h"
+#include "Widgets/Views/SListView.h"
+#include "Widgets/Views/STableRow.h"
+#include "DetailWidgetRow.h"
+
+#define LOCTEXT_NAMESPACE "ScalarVertexPropertyGroupCustomization"
+
+namespace Dataflow
+{
+	TSharedRef<IPropertyTypeCustomization> FScalarVertexPropertyGroupCustomization::MakeInstance()
+	{
+		return MakeShareable(new FScalarVertexPropertyGroupCustomization);
+	}
+
+	void FScalarVertexPropertyGroupCustomization::CustomizeHeader(TSharedRef<IPropertyHandle> PropertyHandle, FDetailWidgetRow& HeaderRow, IPropertyTypeCustomizationUtils& CustomizationUtils)
+	{
+		uint32 NumChildren;
+		const FPropertyAccess::Result Result = PropertyHandle->GetNumChildren(NumChildren);
+
+		ChildPropertyHandle = (Result == FPropertyAccess::Success && NumChildren) ? PropertyHandle->GetChildHandle(0) : nullptr;
+
+		GroupNames.Reset();
+
+		HeaderRow
+			.NameContent()
+			[
+				PropertyHandle->CreatePropertyNameWidget(PropertyHandle->GetPropertyDisplayName())
+			]
+			.ValueContent()
+			.MinDesiredWidth(250)
+			.MaxDesiredWidth(350.f)
+			[
+				SNew(SHorizontalBox)
+				+ SHorizontalBox::Slot()
+				.MaxWidth(145.f)
+				[
+					SAssignNew(ComboButton, SComboButton)
+					.ButtonStyle(FAppStyle::Get(), "NoBorder")
+					.ContentPadding(0)
+					.OnGetMenuContent(this, &FScalarVertexPropertyGroupCustomization::OnGetMenuContent)
+					.ButtonContent()
+					[
+						SNew(SEditableTextBox)
+						.Text_Raw(this, &FScalarVertexPropertyGroupCustomization::GetText)
+						.Font(FAppStyle::GetFontStyle(TEXT("PropertyWindow.NormalFont")))
+						.OnTextCommitted(this, &FScalarVertexPropertyGroupCustomization::OnTextCommitted)
+					]
+				]
+			];
+	}
+
+	FText FScalarVertexPropertyGroupCustomization::GetText() const
+	{
+		FText Text;
+		if (ChildPropertyHandle)
+		{
+			ChildPropertyHandle->GetValueAsFormattedText(Text);
+		}
+		return Text;
+	}
+
+	void FScalarVertexPropertyGroupCustomization::OnTextCommitted(const FText& NewText, ETextCommit::Type CommitInfo)
+	{
+		if (ChildPropertyHandle)
+		{
+			FText CurrentText;
+			ChildPropertyHandle->GetValueAsFormattedText(CurrentText);
+
+			if (!NewText.ToString().Equals(CurrentText.ToString(), ESearchCase::CaseSensitive))
+			{
+				FString String = NewText.ToString();
+				//FClothDataflowTools::MakeCollectionName(String);
+				ChildPropertyHandle->SetValueFromFormattedString(String);
+			}
+		}
+	}
+
+	void FScalarVertexPropertyGroupCustomization::OnSelectionChanged(TSharedPtr<FText> ItemSelected, ESelectInfo::Type /*SelectInfo*/)
+	{
+		if (ChildPropertyHandle)
+		{
+			// Set the child property's value
+			if (ItemSelected)
+			{
+				FText CurrentText;
+				ChildPropertyHandle->GetValueAsFormattedText(CurrentText);
+
+				if (!ItemSelected->EqualTo(CurrentText))
+				{
+					ChildPropertyHandle->SetValueFromFormattedString(ItemSelected->ToString());
+				}
+
+				if (TSharedPtr<SComboButton> PinnedComboButton = ComboButton.Pin())
+				{
+					PinnedComboButton->SetIsOpen(false);
+				}
+			}
+		}
+	}
+
+	TSharedRef<ITableRow> FScalarVertexPropertyGroupCustomization::MakeCategoryViewWidget(TSharedPtr<FText> Item, const TSharedRef< STableViewBase >& OwnerTable)
+	{
+		if (Item)
+		{
+			return SNew(STableRow<TSharedPtr<FString>>, OwnerTable)
+			[
+				SNew(STextBlock).Text(*Item)
+			];
+		}
+		else
+		{
+			return SNew(STableRow<TSharedPtr<FString>>, OwnerTable);
+		}
+	}
+
+	TSharedRef<SWidget> FScalarVertexPropertyGroupCustomization::OnGetMenuContent()
+	{
+		GroupNames.Reset();
+
+		// Find all group names in the parent node's collection
+		FName NodeType;
+		TArray<FName> CollectionGroupNames;
+		if (const FDataflowCollectionAddScalarVertexPropertyNode* const ScalarVertexPropertyNode = GetOwnerStruct<FDataflowCollectionAddScalarVertexPropertyNode>())
+		{
+			CollectionGroupNames = ScalarVertexPropertyNode->GetCachedCollectionGroupNames();
+			NodeType = ScalarVertexPropertyNode->GetType();
+
+			const TArray<FName> AvailableTargetGroups = DataflowAddScalarVertexPropertyCallbackRegistry::Get().GetTargetGroupNames();
+
+			CollectionGroupNames = CollectionGroupNames.FilterByPredicate([&AvailableTargetGroups](const FName& Name)
+				{
+					return AvailableTargetGroups.Contains(Name);
+				});
+
+			for (const FName& CollectionGroupName : CollectionGroupNames)
+			{
+				GroupNames.Add(MakeShareable(new FText(FText::FromName(CollectionGroupName))));
+			}
+		}
+
+		return SNew(SVerticalBox)
+			+ SVerticalBox::Slot()
+			.AutoHeight()
+			.MaxHeight(400.0f)
+			[
+				SNew(SListView<TSharedPtr<FText>>)
+					.ListItemsSource(&GroupNames)
+					.SelectionMode(ESelectionMode::Single)
+					.OnGenerateRow(this, &FScalarVertexPropertyGroupCustomization::MakeCategoryViewWidget)
+					.OnSelectionChanged(this, &FScalarVertexPropertyGroupCustomization::OnSelectionChanged)
+			];
+	}
+
+	template<typename T>
+	T* FScalarVertexPropertyGroupCustomization::GetOwnerStruct() const
+	{
+		if (ChildPropertyHandle)
+		{
+			if (const TSharedPtr<IPropertyHandle> PropertyHandle = ChildPropertyHandle->GetParentHandle())
+			{
+				if (const TSharedPtr<IPropertyHandle> OwnerHandle = PropertyHandle->GetParentHandle())  // Assume that the group struct is only used at the node struct level
+				{
+					if (const TSharedPtr<IPropertyHandleStruct> OwnerHandleStruct = OwnerHandle->AsStruct())
+					{
+						if (TSharedPtr<FStructOnScope> StructOnScope = OwnerHandleStruct->GetStructData())
+						{
+							if (StructOnScope->GetStruct() == T::StaticStruct())
+							{
+								return reinterpret_cast<T*>(StructOnScope->GetStructMemory());
+							}
+						}
+					}
+				}
+			}
+		}
+		return nullptr;
+	}
+}  // End namespace UE::Chaos::ClothAsset
+
+#undef LOCTEXT_NAMESPACE
diff --git a/Engine/Plugins/Experimental/Dataflow/Source/DataflowEditor/Private/Dataflow/ScalarVertexPropertyGroupCustomization.h b/Engine/Plugins/Experimental/Dataflow/Source/DataflowEditor/Private/Dataflow/ScalarVertexPropertyGroupCustomization.h
new file mode 100644
--- /dev/null
+++ b/Engine/Plugins/Experimental/Dataflow/Source/DataflowEditor/Private/Dataflow/ScalarVertexPropertyGroupCustomization.h
@@ -0,0 +1,39 @@
+// Copyright Epic Games, Inc. All Rights Reserved.
+
+#pragma once
+
+#include "IPropertyTypeCustomization.h"
+
+class IPropertyHandle;
+class ITableRow;
+class SComboButton;
+class STableViewBase;
+class SWidget;
+
+namespace Dataflow
+{
+
+	class FScalarVertexPropertyGroupCustomization : public IPropertyTypeCustomization
+	{
+	public:
+
+		static TSharedRef<IPropertyTypeCustomization> MakeInstance();
+
+	private:
+		virtual void CustomizeHeader(TSharedRef<IPropertyHandle> InPropertyHandle, FDetailWidgetRow& HeaderRow, IPropertyTypeCustomizationUtils& CustomizationUtils) override;
+		virtual void CustomizeChildren(TSharedRef<IPropertyHandle> /*InPropertyHandle*/, IDetailChildrenBuilder& /*ChildBuilder*/, IPropertyTypeCustomizationUtils& /*CustomizationUtils*/) override {}
+
+		FText GetText() const;
+		void OnTextCommitted(const FText& NewText, ETextCommit::Type CommitInfo);
+		void OnSelectionChanged(TSharedPtr<FText> ItemSelected, ESelectInfo::Type SelectInfo);
+		TSharedRef<ITableRow> MakeCategoryViewWidget(TSharedPtr<FText> Item, const TSharedRef<STableViewBase>& OwnerTable);
+		TSharedRef<SWidget> OnGetMenuContent();
+		
+		template<typename T>
+		T* GetOwnerStruct() const;
+
+		TSharedPtr<IPropertyHandle> ChildPropertyHandle;
+		TWeakPtr<SComboButton> ComboButton;
+		TArray<TSharedPtr<FText>> GroupNames;
+	};
+}
diff --git a/Engine/Plugins/Experimental/Dataflow/Source/DataflowEditor/Private/DataflowEditorTools/DataflowEditorWeightMapPaintTool.cpp b/Engine/Plugins/Experimental/Dataflow/Source/DataflowEditor/Private/DataflowEditorTools/DataflowEditorWeightMapPaintTool.cpp
--- a/Engine/Plugins/Experimental/Dataflow/Source/DataflowEditor/Private/DataflowEditorTools/DataflowEditorWeightMapPaintTool.cpp
+++ b/Engine/Plugins/Experimental/Dataflow/Source/DataflowEditor/Private/DataflowEditorTools/DataflowEditorWeightMapPaintTool.cpp
@@ -76,7 +76,7 @@ void UDataflowEditorWeightMapPaintToolBuilder::GetSupportedViewModes(TArray<Data
 
 bool UDataflowEditorWeightMapPaintToolBuilder::CanBuildTool(const FToolBuilderState& SceneState) const
 {
-	auto HasRenderableCollection = [](const FDataflowNode* InDataflowNode, const TSharedPtr<Dataflow::FEngineContext> Context)
+	auto HasManagedArrayCollection = [](const FDataflowNode* InDataflowNode, const TSharedPtr<Dataflow::FEngineContext> Context)
 	{
 		if (InDataflowNode && Context)
 		{
@@ -84,12 +84,7 @@ bool UDataflowEditorWeightMapPaintToolBuilder::CanBuildTool(const FToolBuilderSt
 			{
 				if (Output->GetType() == FName("FManagedArrayCollection"))
 				{
-					const FManagedArrayCollection DefaultValue;
-					const FManagedArrayCollection& Collection = Output->GetValue<FManagedArrayCollection>(*Context, DefaultValue);
-					if (Collection.HasGroup("Geometry"))
-					{
-						return true;
-					}
+					return true;
 				}
 			}
 		}
@@ -111,7 +106,7 @@ bool UDataflowEditorWeightMapPaintToolBuilder::CanBuildTool(const FToolBuilderSt
 						{
 							if (const FDataflowNode* PrimarySelection = ContextObject->GetPrimarySelectedNodeOfType<FDataflowCollectionAddScalarVertexPropertyNode>())
 							{
-								return HasRenderableCollection(PrimarySelection, EvaluationContext);
+								return HasManagedArrayCollection(PrimarySelection, EvaluationContext);
 							}
 						}
 					}
diff --git a/Engine/Plugins/Experimental/Dataflow/Source/DataflowNodes/Private/Dataflow/DataflowCollectionAddScalarVertexPropertyNode.cpp b/Engine/Plugins/Experimental/Dataflow/Source/DataflowNodes/Private/Dataflow/DataflowCollectionAddScalarVertexPropertyNode.cpp
--- a/Engine/Plugins/Experimental/Dataflow/Source/DataflowNodes/Private/Dataflow/DataflowCollectionAddScalarVertexPropertyNode.cpp
+++ b/Engine/Plugins/Experimental/Dataflow/Source/DataflowNodes/Private/Dataflow/DataflowCollectionAddScalarVertexPropertyNode.cpp
@@ -4,11 +4,62 @@
 
 #include "Dataflow/DataflowInputOutput.h"
 #include "Dataflow/DataflowTools.h"
+#include "Misc/LazySingleton.h"
 
 #include UE_INLINE_GENERATED_CPP_BY_NAME(DataflowCollectionAddScalarVertexPropertyNode)
 
 #define LOCTEXT_NAMESPACE "DataflowCollectionAddScalarVertexProperty"
 
+
+DataflowAddScalarVertexPropertyCallbackRegistry& DataflowAddScalarVertexPropertyCallbackRegistry::Get()
+{
+	return TLazySingleton<DataflowAddScalarVertexPropertyCallbackRegistry>::Get();
+}
+
+void DataflowAddScalarVertexPropertyCallbackRegistry::TearDown()
+{
+	TLazySingleton<DataflowAddScalarVertexPropertyCallbackRegistry>::TearDown();
+}
+
+void DataflowAddScalarVertexPropertyCallbackRegistry::RegisterCallbacks(TUniquePtr<IDataflowAddScalarVertexPropertyCallbacks>&& Callbacks)
+{
+	AllCallbacks.Add(Callbacks->GetName(), MoveTemp(Callbacks));
+}
+
+void DataflowAddScalarVertexPropertyCallbackRegistry::DeregisterCallbacks(const FName& CallbacksName)
+{
+	AllCallbacks.Remove(CallbacksName);
+}
+
+TArray<FName> DataflowAddScalarVertexPropertyCallbackRegistry::GetTargetGroupNames() const
+{
+	TArray<FName> UniqueNames;
+	
+	for (const TPair<FName, TUniquePtr<IDataflowAddScalarVertexPropertyCallbacks>>& CallbacksEntry : AllCallbacks)
+	{
+		for (const FName& GroupName : CallbacksEntry.Value->GetTargetGroupNames())
+		{
+			UniqueNames.AddUnique(GroupName);
+		}
+	}
+	return UniqueNames;
+}
+
+TArray<Dataflow::FRenderingParameter> DataflowAddScalarVertexPropertyCallbackRegistry::GetRenderingParameters() const
+{
+	TArray<Dataflow::FRenderingParameter> UniqueParameters;
+
+	for (const TPair<FName, TUniquePtr<IDataflowAddScalarVertexPropertyCallbacks>>& CallbacksEntry : AllCallbacks)
+	{
+		for (const Dataflow::FRenderingParameter& RenderingParameter : CallbacksEntry.Value->GetRenderingParameters())
+		{
+			UniqueParameters.AddUnique(RenderingParameter);
+		}
+	}
+	return UniqueParameters;
+}
+
+
 FDataflowCollectionAddScalarVertexPropertyNode::FDataflowCollectionAddScalarVertexPropertyNode(const Dataflow::FNodeParameters& InParam, FGuid InGuid)
 	: FDataflowNode(InParam, InGuid)
 {
@@ -17,6 +68,12 @@ FDataflowCollectionAddScalarVertexPropertyNode::FDataflowCollectionAddScalarVert
 	RegisterOutputConnection(&AttributeKey);
 }
 
+TArray<::Dataflow::FRenderingParameter> FDataflowCollectionAddScalarVertexPropertyNode::GetRenderParametersImpl() const
+{
+	return DataflowAddScalarVertexPropertyCallbackRegistry::Get().GetRenderingParameters();
+}
+
+
 void FDataflowCollectionAddScalarVertexPropertyNode::Evaluate(Dataflow::FContext& Context, const FDataflowOutput* Out) const
 {
 	using namespace Dataflow;
@@ -28,16 +85,18 @@ void FDataflowCollectionAddScalarVertexPropertyNode::Evaluate(Dataflow::FContext
 
 		if (!Name.IsEmpty())
 		{
-			FName nName(Name), nGroup("Vertices");
-			TManagedArray<float>& Scalar = InCollection.AddAttribute<float>(nName,nGroup);
+			const FName InName(Name);
+			const FName InGroup = TargetGroup.Name;
+			TManagedArray<float>& Scalar = InCollection.AddAttribute<float>(InName, InGroup);
 
 			const int32 MaxWeightIndex = FMath::Min(VertexWeights.Num(), Scalar.Num());
 			if (VertexWeights.Num() > 0 && VertexWeights.Num() != Scalar.Num())
 			{
 				FDataflowTools::LogAndToastWarning(*this,
 					LOCTEXT("VertexCountMismatchHeadline", "Vertex count mismatch."),
-					FText::Format(LOCTEXT("VertexCountMismatchDetails", "Vertex weights in the node: {0}\n3D vertices in the cloth: {1}"),
+					FText::Format(LOCTEXT("VertexCountMismatchDetails", "Vertex weights in the node: {0}\n Vertices in group \"{1}\" in the Collection: {2}"),
 						VertexWeights.Num(),
+						FText::FromName(InGroup),
 						Scalar.Num()));
 			}
 
@@ -55,4 +114,26 @@ void FDataflowCollectionAddScalarVertexPropertyNode::Evaluate(Dataflow::FContext
 	}
 }
 
+void FDataflowCollectionAddScalarVertexPropertyNode::OnSelected(Dataflow::FContext& Context)
+{
+	// Re-evaluate the input collection
+	FManagedArrayCollection InCollection = GetValue<FManagedArrayCollection>(Context, &Collection);
+	const TSharedRef<FManagedArrayCollection> ManagedArrayCollection = MakeShared<FManagedArrayCollection>(MoveTemp(InCollection));
+
+	// Update the list of used group for the UI customization
+	const TArray<FName> GroupNames = ManagedArrayCollection->GroupNames();
+	CachedCollectionGroupNames.Reset(GroupNames.Num());
+	for (const FName& GroupName : GroupNames)
+	{
+		CachedCollectionGroupNames.Emplace(GroupName);
+	}
+}
+
+void FDataflowCollectionAddScalarVertexPropertyNode::OnDeselected()
+{
+	// Clean up, to avoid another toolkit picking up the wrong context evaluation
+	CachedCollectionGroupNames.Reset();
+}
+
+
 #undef LOCTEXT_NAMESPACE
diff --git a/Engine/Plugins/Experimental/Dataflow/Source/DataflowNodes/Private/Dataflow/DataflowNodesPlugin.cpp b/Engine/Plugins/Experimental/Dataflow/Source/DataflowNodes/Private/Dataflow/DataflowNodesPlugin.cpp
--- a/Engine/Plugins/Experimental/Dataflow/Source/DataflowNodes/Private/Dataflow/DataflowNodesPlugin.cpp
+++ b/Engine/Plugins/Experimental/Dataflow/Source/DataflowNodes/Private/Dataflow/DataflowNodesPlugin.cpp
@@ -13,6 +13,31 @@
 
 #define LOCTEXT_NAMESPACE "DataflowNodes"
 
+class FGeometryCollectionAddScalarVertexPropertyCallbacks : public IDataflowAddScalarVertexPropertyCallbacks
+{
+public:
+
+	const static FName Name;
+
+	virtual ~FGeometryCollectionAddScalarVertexPropertyCallbacks() = default;
+
+	virtual FName GetName() const override
+	{
+		return Name;
+	}
+
+	virtual TArray<FName> GetTargetGroupNames() const override
+	{
+		return { FGeometryCollection::VerticesGroup };
+	}
+
+	virtual TArray<Dataflow::FRenderingParameter> GetRenderingParameters() const override
+	{
+		return { { TEXT("SurfaceRender"), FGeometryCollection::StaticType(), {TEXT("Collection")} } };
+	}
+};
+
+const FName FGeometryCollectionAddScalarVertexPropertyCallbacks::Name = FName("FGeometryCollectionAddScalarVertexPropertyCallbacks");
 
 void IDataflowNodesPlugin::StartupModule()
 {
@@ -24,10 +49,13 @@ void IDataflowNodesPlugin::StartupModule()
 	DATAFLOW_NODE_REGISTER_CREATION_FACTORY(FDataflowCollectionAddScalarVertexPropertyNode);
 
 	Dataflow::RegisterNodeFilter(FDataflowTerminalNode::StaticType());
+
+	DataflowAddScalarVertexPropertyCallbackRegistry::Get().RegisterCallbacks(MakeUnique<FGeometryCollectionAddScalarVertexPropertyCallbacks>());
 }
 
 void IDataflowNodesPlugin::ShutdownModule()
 {
+	DataflowAddScalarVertexPropertyCallbackRegistry::Get().DeregisterCallbacks(FGeometryCollectionAddScalarVertexPropertyCallbacks::Name);
 }
 
 
diff --git a/Engine/Plugins/Experimental/Dataflow/Source/DataflowNodes/Public/Dataflow/DataflowCollectionAddScalarVertexPropertyNode.h b/Engine/Plugins/Experimental/Dataflow/Source/DataflowNodes/Public/Dataflow/DataflowCollectionAddScalarVertexPropertyNode.h
--- a/Engine/Plugins/Experimental/Dataflow/Source/DataflowNodes/Public/Dataflow/DataflowCollectionAddScalarVertexPropertyNode.h
+++ b/Engine/Plugins/Experimental/Dataflow/Source/DataflowNodes/Public/Dataflow/DataflowCollectionAddScalarVertexPropertyNode.h
@@ -9,14 +9,58 @@
 
 #include "DataflowCollectionAddScalarVertexPropertyNode.generated.h"
 
+
+class IDataflowAddScalarVertexPropertyCallbacks
+{
+public:
+	virtual ~IDataflowAddScalarVertexPropertyCallbacks() = default;
+	virtual FName GetName() const = 0;
+	virtual TArray<FName> GetTargetGroupNames() const = 0;
+	virtual TArray<Dataflow::FRenderingParameter> GetRenderingParameters() const = 0;
+};
+
+class DataflowAddScalarVertexPropertyCallbackRegistry
+{
+public:
+
+	DATAFLOWNODES_API static DataflowAddScalarVertexPropertyCallbackRegistry& Get();
+	DATAFLOWNODES_API static void TearDown();
+
+	DATAFLOWNODES_API void RegisterCallbacks(TUniquePtr<IDataflowAddScalarVertexPropertyCallbacks>&& Callbacks);
+
+	DATAFLOWNODES_API void DeregisterCallbacks(const FName& CallbacksName);
+
+	DATAFLOWNODES_API TArray<FName> GetTargetGroupNames() const;
+
+	DATAFLOWNODES_API TArray<Dataflow::FRenderingParameter> GetRenderingParameters() const;
+
+private:
+
+	TMap<FName, TUniquePtr<IDataflowAddScalarVertexPropertyCallbacks>> AllCallbacks;
+};
+
+
+/*
+* Custom type so that we can use property type customization
+*/
+USTRUCT()
+struct FScalarVertexPropertyGroup
+{
+	GENERATED_USTRUCT_BODY()
+
+	UPROPERTY(EditAnywhere, Category = "Vertex Group")
+	FName Name;
+};
+
+
 /** Scalar vertex properties. */
 USTRUCT(Meta = (DataflowCollection))
 struct DATAFLOWNODES_API FDataflowCollectionAddScalarVertexPropertyNode : public FDataflowNode
 {
 	GENERATED_USTRUCT_BODY()
 	DATAFLOW_NODE_DEFINE_INTERNAL(FDataflowCollectionAddScalarVertexPropertyNode, "AddScalarVertexProperty", "Collection", "Add a saved scalar property to a collection")
-	DATAFLOW_NODE_RENDER_TYPE("SurfaceRender", FGeometryCollection::StaticType(), "Collection")
 
+	virtual TArray<::Dataflow::FRenderingParameter> GetRenderParametersImpl() const override;
 
 public:
 
@@ -33,8 +77,18 @@ public:
 	UPROPERTY()
 	TArray<float> VertexWeights;
 
+	UPROPERTY(EditAnywhere, Category = "Vertex Attribute")
+	FScalarVertexPropertyGroup TargetGroup;
+
 	FDataflowCollectionAddScalarVertexPropertyNode(const Dataflow::FNodeParameters& InParam, FGuid InGuid = FGuid::NewGuid());
 
+	/** Return a cached array of all the groups used by the input collection during at the time of the latest evaluation. */
+	const TArray<FName>& GetCachedCollectionGroupNames() const { return CachedCollectionGroupNames; }
+
 private:
 	virtual void Evaluate(Dataflow::FContext& Context, const FDataflowOutput* Out) const override;
+	virtual void OnSelected(Dataflow::FContext& Context) override;
+	virtual void OnDeselected() override;
+
+	TArray<FName> CachedCollectionGroupNames;
 };
diff --git a/Engine/Source/Runtime/Experimental/Dataflow/Core/Public/Dataflow/DataflowNodeParameters.h b/Engine/Source/Runtime/Experimental/Dataflow/Core/Public/Dataflow/DataflowNodeParameters.h
--- a/Engine/Source/Runtime/Experimental/Dataflow/Core/Public/Dataflow/DataflowNodeParameters.h
+++ b/Engine/Source/Runtime/Experimental/Dataflow/Core/Public/Dataflow/DataflowNodeParameters.h
@@ -42,6 +42,8 @@ namespace Dataflow
 			: Name(InRenderName), Type(InTypeName), Outputs(InOutputs) {}
 		FRenderingParameter(FString InRenderName, FName InTypeName, TArray<FName>&& InOutputs)
 			: Name(InRenderName), Type(InTypeName), Outputs(InOutputs) {}
+		
+		bool operator==(const FRenderingParameter& Other) const = default;
 
 		FString Name = FString("");
 		FName Type = FName("");
