Update Dataflow Node Details panel to handle showing multiple selected nodes.

diff --git a/Engine/Plugins/Experimental/Dataflow/Source/DataflowEditor/Private/Dataflow/DataflowEditorCommands.cpp b/Engine/Plugins/Experimental/Dataflow/Source/DataflowEditor/Private/Dataflow/DataflowEditorCommands.cpp
--- a/Engine/Plugins/Experimental/Dataflow/Source/DataflowEditor/Private/Dataflow/DataflowEditorCommands.cpp
+++ b/Engine/Plugins/Experimental/Dataflow/Source/DataflowEditor/Private/Dataflow/DataflowEditorCommands.cpp
@@ -20,6 +20,7 @@
 #include "Editor.h"
 #include "IDetailsView.h"
 #include "IStructureDetailsView.h"
+#include "IStructureDataProvider.h"
 #include "Serialization/ObjectWriter.h"
 #include "Serialization/ObjectReader.h"
 #include "Dataflow/DataflowGraph.h"
@@ -429,7 +430,8 @@ void FDataflowEditorCommands::OnSelectedNodesChanged(TSharedPtr<IStructureDetail
 			FGraphPanelSelectionSet SelectedNodes = AsRawPointers(NewSelection);
 			if (SelectedNodes.Num())
 			{
-				TArray<UObject*> Objects;
+				TArray<TSharedPtr<FStructOnScope>> StructData;
+				StructData.Reserve(SelectedNodes.Num());
 				for (UObject* SelectedObject : SelectedNodes)
 				{
 					if (UDataflowEdNode* EdNode = Cast<UDataflowEdNode>(SelectedObject))
@@ -447,16 +449,18 @@ void FDataflowEditorCommands::OnSelectedNodesChanged(TSharedPtr<IStructureDetail
 								});
 							PropertiesEditor->GetDetailsView()->SetIsPropertyReadOnlyDelegate(Delegate);
 
-							TSharedPtr<FStructOnScope> Struct(DataflowNode->NewStructOnScope());
-							PropertiesEditor->SetStructureData(Struct);
+							StructData.Emplace(DataflowNode->NewStructOnScope());
 						}
 					}
 					else if (UEdGraphNode_Comment* CommentNode = Cast<UEdGraphNode_Comment>(SelectedObject))
 					{
-						TSharedPtr<FStructOnScope> Struct(new FStructOnScope(UEdGraphNode_Comment::StaticClass(), (uint8*)CommentNode));
-						PropertiesEditor->SetStructureData(Struct);
+						StructData.Emplace(new FStructOnScope(UEdGraphNode_Comment::StaticClass(), (uint8*)CommentNode));
 					}
 				}
+				if (StructData.Num())
+				{
+					PropertiesEditor->SetStructureProvider(MakeShared<FStructOnScopeStructureDataProvider>(StructData));
+				}
 			}
 		}
 	}
