Dataflow:
- When new nodes get created they don't start in 'Rename mode'
- Implemented right mouse click 'Rename F2' for nodes
- When new comment boxes get created they start in 'Rename mode'
- Implemented right mouse click 'Rename F2' for comment boxes

diff --git a/Engine/Plugins/Experimental/Dataflow/Source/DataflowEditor/Private/Dataflow/DataflowEditorCommands.cpp b/Engine/Plugins/Experimental/Dataflow/Source/DataflowEditor/Private/Dataflow/DataflowEditorCommands.cpp
--- a/Engine/Plugins/Experimental/Dataflow/Source/DataflowEditor/Private/Dataflow/DataflowEditorCommands.cpp
+++ b/Engine/Plugins/Experimental/Dataflow/Source/DataflowEditor/Private/Dataflow/DataflowEditorCommands.cpp
@@ -415,6 +415,15 @@ void FDataflowEditorCommands::DeleteNodes(UDataflow* Graph, const FGraphPanelSel
 	}
 }
 
+void FDataflowEditorCommands::RenameNode(const TSharedPtr<SDataflowGraphEditor>& DataflowGraphEditor, UEdGraphNode* EdNode)
+{
+	if (EdNode && DataflowGraphEditor.IsValid())
+	{
+		// GMelich: There is no direct rename function, this function can rename a node without recentering the selected node
+		DataflowGraphEditor->IsNodeTitleVisible(EdNode, /*bRequestRename*/true);
+	}
+}
+
 void FDataflowEditorCommands::OnSelectedNodesChanged(TSharedPtr<IStructureDetailsView> PropertiesEditor, UObject* Asset, UDataflow* Graph, const TSet<TObjectPtr<UObject> >& NewSelection)
 {
 	PropertiesEditor->SetStructureData(nullptr);
diff --git a/Engine/Plugins/Experimental/Dataflow/Source/DataflowEditor/Private/Dataflow/DataflowGraphEditor.cpp b/Engine/Plugins/Experimental/Dataflow/Source/DataflowEditor/Private/Dataflow/DataflowGraphEditor.cpp
--- a/Engine/Plugins/Experimental/Dataflow/Source/DataflowEditor/Private/Dataflow/DataflowGraphEditor.cpp
+++ b/Engine/Plugins/Experimental/Dataflow/Source/DataflowEditor/Private/Dataflow/DataflowGraphEditor.cpp
@@ -142,6 +142,11 @@ void SDataflowGraphEditor::Construct(const FArguments& InArgs, UObject* InAssetO
 				FGenericCommands::Get().Paste,
 				FExecuteAction::CreateSP(this, &SDataflowGraphEditor::PasteSelectedNodes)
 			);
+			GraphEditorCommands->MapAction(
+				FGenericCommands::Get().Rename,
+				FExecuteAction::CreateSP(this, &SDataflowGraphEditor::RenameNode),
+				FCanExecuteAction::CreateSP(this, &SDataflowGraphEditor::CanRenameNode)
+			);
 		}
 	}
 
@@ -221,6 +226,48 @@ void SDataflowGraphEditor::DeleteNode()
 	}
 }
 
+void SDataflowGraphEditor::RenameNode()
+{
+	if (UDataflow* Graph = DataflowAsset.Get())
+	{
+		const TSharedPtr<SDataflowGraphEditor>& DataflowGraphEditor = SharedThis(this);
+		const FGraphPanelSelectionSet& SelectedNodes = GetSelectedNodes();
+
+		if (SelectedNodes.Num() == 1)
+		{
+			if (CanRenameNode())
+			{
+				if (UDataflowEdNode* SelectedNode = Cast<UDataflowEdNode>(*SelectedNodes.CreateConstIterator()))
+				{
+					FDataflowEditorCommands::RenameNode(DataflowGraphEditor, SelectedNode);
+				}
+				else if (UEdGraphNode_Comment* SelectedCommentNode = Cast<UEdGraphNode_Comment>(*SelectedNodes.CreateConstIterator()))
+				{
+					FDataflowEditorCommands::RenameNode(DataflowGraphEditor, SelectedCommentNode);
+				}
+			}
+		}
+	}
+}
+
+bool SDataflowGraphEditor::CanRenameNode() const
+{
+	const FGraphPanelSelectionSet SelectedNodes = GetSelectedNodes();
+	if (SelectedNodes.Num() == 1)
+	{
+		if (UDataflowEdNode* SelectedNode = Cast<UDataflowEdNode>(*SelectedNodes.CreateConstIterator()))
+		{
+			return SelectedNode->bCanRenameNode;
+		}
+		else if (UEdGraphNode_Comment* SelectedCommentNode = Cast<UEdGraphNode_Comment>(*SelectedNodes.CreateConstIterator()))
+		{
+			return SelectedCommentNode->bCanRenameNode;
+		}
+	}
+
+	return false;
+}
+
 void SDataflowGraphEditor::OnSelectedNodesChanged(const TSet<UObject*>& NewSelection)
 {
 	OnSelectionChangedMulticast.Broadcast(NewSelection);  // Broadcast the selection change before refreshing the DetailsView, the nodes' specific UI data have to be updated before the UI is being redrawn
diff --git a/Engine/Plugins/Experimental/Dataflow/Source/DataflowEditor/Public/Dataflow/DataflowEditorCommands.h b/Engine/Plugins/Experimental/Dataflow/Source/DataflowEditor/Public/Dataflow/DataflowEditorCommands.h
--- a/Engine/Plugins/Experimental/Dataflow/Source/DataflowEditor/Public/Dataflow/DataflowEditorCommands.h
+++ b/Engine/Plugins/Experimental/Dataflow/Source/DataflowEditor/Public/Dataflow/DataflowEditorCommands.h
@@ -156,5 +156,10 @@ public:
 	*  PasteSelectedNodes
 	*/
 	static void PasteNodes(UDataflow* Graph, const TSharedPtr<SDataflowGraphEditor>& DataflowGraphEditor);
+
+	/*
+	*  RenameNode
+	*/
+	static void RenameNode(const TSharedPtr<SDataflowGraphEditor>& DataflowGraphEditor, UEdGraphNode* EdNode);
 };
 
diff --git a/Engine/Plugins/Experimental/Dataflow/Source/DataflowEditor/Public/Dataflow/DataflowGraphEditor.h b/Engine/Plugins/Experimental/Dataflow/Source/DataflowEditor/Public/Dataflow/DataflowGraphEditor.h
--- a/Engine/Plugins/Experimental/Dataflow/Source/DataflowEditor/Public/Dataflow/DataflowGraphEditor.h
+++ b/Engine/Plugins/Experimental/Dataflow/Source/DataflowEditor/Public/Dataflow/DataflowGraphEditor.h
@@ -135,6 +135,10 @@ public:
 	/** */
 	void PasteSelectedNodes();
 
+	/** */
+	void RenameNode();
+	bool CanRenameNode() const;
+
 	SGraphEditor* GetGraphEditor() { return (SGraphEditor*)this; }
 
 	/** FGCObject interface */
diff --git a/Engine/Plugins/Experimental/Dataflow/Source/DataflowEditor/Public/Dataflow/DataflowSNode.h b/Engine/Plugins/Experimental/Dataflow/Source/DataflowEditor/Public/Dataflow/DataflowSNode.h
--- a/Engine/Plugins/Experimental/Dataflow/Source/DataflowEditor/Public/Dataflow/DataflowSNode.h
+++ b/Engine/Plugins/Experimental/Dataflow/Source/DataflowEditor/Public/Dataflow/DataflowSNode.h
@@ -37,6 +37,7 @@ public:
 	virtual FReply OnMouseButtonDoubleClick(const FGeometry& InMyGeometry, const FPointerEvent& InMouseEvent) override;
 	TArray<FOverlayWidgetInfo> GetOverlayWidgets(bool bSelected, const FVector2D& WidgetSize) const;
 	virtual void UpdateErrorInfo() override;
+	virtual void RequestRenameOnSpawn() override { /* No auto rename on spawn, because it can interfere with Copy/Paste and cause a crash */ }
 
 	static void CopyDataflowNodeSettings(TSharedPtr<FDataflowNode> SourceDataflowNode, TSharedPtr<FDataflowNode> TargetDataflowNode);
 
diff --git a/Engine/Source/Editor/GraphEditor/Public/SGraphNodeComment.h b/Engine/Source/Editor/GraphEditor/Public/SGraphNodeComment.h
--- a/Engine/Source/Editor/GraphEditor/Public/SGraphNodeComment.h
+++ b/Engine/Source/Editor/GraphEditor/Public/SGraphNodeComment.h
@@ -55,6 +55,12 @@ public:
 	//~ Begin SGraphNode Interface
 	virtual bool IsNameReadOnly() const override;
 	virtual FSlateColor GetCommentColor() const override { return GetCommentBodyColor(); }
+	/** Requests a rename when the node was initially spawned */
+	virtual void RequestRenameOnSpawn() 
+	{ 
+		RequestRename(); 
+		ApplyRename(); 
+	}
 	//~ End SGraphNode Interface
 
 	void Construct( const FArguments& InArgs, UEdGraphNode_Comment* InNode );
