Dataflow - Added a Functions property structure customization to emulate the CallInEditor function buttons.

#rnx
#tests Added three functions with different setup and checked that the node update was triggered after the delegate was executed.

diff --git a/Engine/Plugins/Experimental/Dataflow/Source/DataflowEditor/Private/Dataflow/DataflowEditorModule.cpp b/Engine/Plugins/Experimental/Dataflow/Source/DataflowEditor/Private/Dataflow/DataflowEditorModule.cpp
--- a/Engine/Plugins/Experimental/Dataflow/Source/DataflowEditor/Private/Dataflow/DataflowEditorModule.cpp
+++ b/Engine/Plugins/Experimental/Dataflow/Source/DataflowEditor/Private/Dataflow/DataflowEditorModule.cpp
@@ -2,18 +2,17 @@
 
 #include "Dataflow/DataflowEditorModule.h"
 
-#include "Dataflow/DataflowGraphEditor.h"
 #include "Dataflow/DataflowEditorStyle.h"
 #include "Dataflow/DataflowEditorMode.h"
 #include "Dataflow/DataflowEditorToolkit.h"
 #include "Dataflow/DataflowEditorCommands.h"
 #include "Dataflow/DataflowEngineRendering.h"
+#include "Dataflow/DataflowFunctionsProperty.h"
+#include "Dataflow/DataflowFunctionsPropertyCustomization.h"
 #include "Dataflow/DataflowSNodeFactories.h"
 #include "Dataflow/ScalarVertexPropertyGroupCustomization.h"
 #include "Dataflow/DataflowCollectionAddScalarVertexPropertyNode.h"
 
-#include "CoreMinimal.h"
-#include "EdGraphUtilities.h"
 #include "PropertyEditorModule.h"
 #include "EditorModeRegistry.h"
 
@@ -25,12 +24,11 @@ void FDataflowEditorModule::StartupModule()
 {
 	FDataflowEditorStyle::Get();
 	
-	FPropertyEditorModule& PropertyEditorModule = FModuleManager::LoadModuleChecked<FPropertyEditorModule>(TEXT("PropertyEditor"));
-
 	// Register type customizations
 	if (FPropertyEditorModule* const PropertyModule = FModuleManager::GetModulePtr<FPropertyEditorModule>("PropertyEditor"))
 	{
 		PropertyModule->RegisterCustomPropertyTypeLayout(FScalarVertexPropertyGroup::StaticStruct()->GetFName(), FOnGetPropertyTypeCustomizationInstance::CreateStatic(&Dataflow::FScalarVertexPropertyGroupCustomization::MakeInstance));
+		PropertyModule->RegisterCustomPropertyTypeLayout(FDataflowFunctionsProperty::StaticStruct()->GetFName(), FOnGetPropertyTypeCustomizationInstance::CreateStatic(&Dataflow::FFunctionsPropertyCustomization::MakeInstance));
 	}
 
 	Dataflow::RenderingCallbacks();
@@ -44,6 +42,7 @@ void FDataflowEditorModule::ShutdownModule()
 	if (FPropertyEditorModule* const PropertyModule = FModuleManager::GetModulePtr<FPropertyEditorModule>("PropertyEditor"))
 	{
 		PropertyModule->UnregisterCustomPropertyTypeLayout(FScalarVertexPropertyGroup::StaticStruct()->GetFName());
+		PropertyModule->UnregisterCustomPropertyTypeLayout(FDataflowFunctionsProperty::StaticStruct()->GetFName());
 	}
 }
 
diff --git a/Engine/Plugins/Experimental/Dataflow/Source/DataflowEditor/Private/Dataflow/DataflowFunctionsPropertyCustomization.cpp b/Engine/Plugins/Experimental/Dataflow/Source/DataflowEditor/Private/Dataflow/DataflowFunctionsPropertyCustomization.cpp
new file mode 100644
--- /dev/null
+++ b/Engine/Plugins/Experimental/Dataflow/Source/DataflowEditor/Private/Dataflow/DataflowFunctionsPropertyCustomization.cpp
@@ -0,0 +1,141 @@
+// Copyright Epic Games, Inc. All Rights Reserved.
+
+#include "Dataflow/DataflowFunctionsPropertyCustomization.h"
+#include "Dataflow/DataflowFunctionsProperty.h"
+#include "Widgets/Images/SImage.h"
+#include "Widgets/Input/SButton.h"
+#include "Widgets/Layout/SWrapBox.h"
+#include "Widgets/Text/STextBlock.h"
+#include "DetailWidgetRow.h"
+#include "PropertyHandle.h"
+
+#define LOCTEXT_NAMESPACE "FunctionPropertyCustomization"
+
+namespace Dataflow
+{
+	namespace Private
+	{
+		static const FDataflowFunctionsProperty* GetFunctionProperty(const TSharedPtr<IPropertyHandle>& PropertyHandle)
+		{
+			const FStructProperty* const StructProperty = CastField<FStructProperty>(PropertyHandle->GetProperty());
+			check(StructProperty && StructProperty->Struct);
+			checkf(StructProperty->Struct->IsA(FDataflowFunctionsProperty::StaticStruct()->StaticClass()),
+				TEXT("This customisation can only apply to FDataflowFunctionsProperty and its derived structures."));
+
+			const FDataflowFunctionsProperty* FunctionProperty = nullptr;
+			void* Data;
+			if (PropertyHandle && PropertyHandle->GetValueData(Data) == FPropertyAccess::Success)
+			{
+				FunctionProperty = reinterpret_cast<const FDataflowFunctionsProperty*>(Data);
+			}
+			return FunctionProperty;
+		}
+	}
+
+	TSharedRef<IPropertyTypeCustomization> FFunctionsPropertyCustomization::MakeInstance()
+	{
+		return MakeShareable(new FFunctionsPropertyCustomization);
+	}
+
+	void FFunctionsPropertyCustomization::CustomizeHeader(TSharedRef<IPropertyHandle> StructPropertyHandle, FDetailWidgetRow& HeaderRow, IPropertyTypeCustomizationUtils& CustomizationUtils)
+	{
+		StructProperty = StructPropertyHandle;
+
+		TSharedPtr<SWrapBox> WrapBox;
+		HeaderRow
+		[
+			SAssignNew(WrapBox, SWrapBox)
+			.PreferredSize(2000)  // Copied from FObjectDetails::AddCallInEditorMethods()
+			.UseAllottedSize(true)
+		];
+
+		const FDataflowFunctionsProperty* const FunctionBlockProperty = Private::GetFunctionProperty(StructPropertyHandle);
+		check(FunctionBlockProperty);
+
+		for (const FDataflowFunctionsProperty::FFunction& Function : FunctionBlockProperty->GetFunctions())
+		{
+			auto OnClicked = [this, &Function]() -> FReply
+				{
+					// Execute function
+					Function.Delegate.ExecuteIfBound();
+					StructProperty->NotifyFinishedChangingProperties();  // Triggers node invalidation
+					return FReply::Handled();
+				};
+
+			if (!Function.ImageStyle.IsNone())
+			{
+				const FString& ButtonImage = StructProperty->GetMetaData("ButtonImage");  // e.g. FAppStyle::GetBrush("Persona.ReimportAsset") will be Meta = (ButtonImage = "Persona.ReimportAsset")
+				if (Function.Name.IsEmptyOrWhitespace())
+				{
+					// Create button image
+					WrapBox->AddSlot()
+						.Padding(0.f, 0.f, 5.f, 3.f)
+						[
+							SNew(SButton)
+							.ToolTipText(Function.ToolTip)
+							.OnClicked_Lambda(MoveTemp(OnClicked))
+							.ContentPadding(FMargin(0.f, 4.f))  // Too much horizontal padding otherwise (default is 4, 2)
+							.IsEnabled_Lambda([this]() -> bool { return StructProperty->IsEditable(); })
+							[
+								SNew(SImage)
+								.DesiredSizeOverride(FVector2D(16, 16))
+								.Image(FAppStyle::GetBrush(Function.ImageStyle))
+							]
+						];
+				}
+				else
+				{
+					// Create button image + text
+					WrapBox->AddSlot()
+						.Padding(0.f, 0.f, 5.f, 3.f)
+						[
+							SNew(SButton)
+							.Text(Function.Name)
+							.ToolTipText(Function.ToolTip)
+							.OnClicked_Lambda(MoveTemp(OnClicked))
+							.ContentPadding(FMargin(0.f, 2.f))  // Too much horizontal padding otherwise (default is 4, 2)
+							.IsEnabled_Lambda([this]() -> bool { return StructProperty->IsEditable(); })
+							[
+								SNew(SHorizontalBox)
+								+ SHorizontalBox::Slot()
+								.AutoWidth()
+								.Padding(0, 2)
+								[
+									SNew(SImage)
+									.DesiredSizeOverride(FVector2D(16, 16))
+									.Image(FAppStyle::GetBrush(Function.ImageStyle))
+								]
+								+ SHorizontalBox::Slot()
+								.AutoWidth()
+								.Padding(5.f, 0.f)
+								+ SHorizontalBox::Slot()
+								.AutoWidth()
+								.VAlign(VAlign_Center)
+								[
+									SNew(STextBlock)
+									.Text(Function.Name)
+								]
+							]
+						];
+				}
+			}
+			else
+			{
+				// Create button text
+				WrapBox->AddSlot()
+					.Padding(0.f, 0.f, 5.f, 3.f)
+					[
+						SNew(SButton)
+						.Text(Function.Name)
+						.ToolTipText(Function.ToolTip)
+						.OnClicked_Lambda(MoveTemp(OnClicked))
+						.ContentPadding(FMargin(0.f, 2.f))  // Too much horizontal padding otherwise (default is 4, 2)
+						.VAlign(VAlign_Center)  // No Slot, so need vertical align
+						.IsEnabled_Lambda([this]() -> bool { return StructProperty->IsEditable(); })
+					];
+			}
+		}
+	}
+}  // End namespace Dataflow
+
+#undef LOCTEXT_NAMESPACE
diff --git a/Engine/Plugins/Experimental/Dataflow/Source/DataflowEditor/Private/Dataflow/DataflowFunctionsPropertyCustomization.h b/Engine/Plugins/Experimental/Dataflow/Source/DataflowEditor/Private/Dataflow/DataflowFunctionsPropertyCustomization.h
new file mode 100644
--- /dev/null
+++ b/Engine/Plugins/Experimental/Dataflow/Source/DataflowEditor/Private/Dataflow/DataflowFunctionsPropertyCustomization.h
@@ -0,0 +1,24 @@
+// Copyright Epic Games, Inc. All Rights Reserved.
+
+#pragma once
+
+#include "IPropertyTypeCustomization.h"
+
+namespace Dataflow
+{
+	/**
+	 * Customization for buttons in the UStruct UI of dataflow nodes.
+	 */
+	class FFunctionsPropertyCustomization : public IPropertyTypeCustomization
+	{
+	public:
+		static TSharedRef<IPropertyTypeCustomization> MakeInstance();
+
+	protected:
+		virtual void CustomizeHeader(TSharedRef<IPropertyHandle> StructPropertyHandle, FDetailWidgetRow& HeaderRow, IPropertyTypeCustomizationUtils& CustomizationUtils) override;
+		virtual void CustomizeChildren(TSharedRef<IPropertyHandle> /*StructPropertyHandle*/, IDetailChildrenBuilder& /*ChildBuilder*/, IPropertyTypeCustomizationUtils& /*CustomizationUtils*/) override {}
+
+	private:
+		TSharedPtr<IPropertyHandle> StructProperty;
+	};
+}
diff --git a/Engine/Plugins/Experimental/Dataflow/Source/DataflowNodes/Public/Dataflow/DataflowFunctionsProperty.h b/Engine/Plugins/Experimental/Dataflow/Source/DataflowNodes/Public/Dataflow/DataflowFunctionsProperty.h
new file mode 100644
--- /dev/null
+++ b/Engine/Plugins/Experimental/Dataflow/Source/DataflowNodes/Public/Dataflow/DataflowFunctionsProperty.h
@@ -0,0 +1,48 @@
+// Copyright Epic Games, Inc. All Rights Reserved.
+
+#pragma once
+
+#include "Delegates/Delegate.h"
+#include "DataflowFunctionsProperty.generated.h"
+
+/**
+ * Function property for all Dataflow nodes.
+ * The structure is also used in DataFlow::FFunctionDetailCustomization to appear as text and/or image buttons.
+ * This helps with the equivalent UCLASS UFUNCTION CallInEditor functionality that is missing from the USTRUCT implementation.
+ */
+USTRUCT()
+struct FDataflowFunctionsProperty
+{
+	GENERATED_BODY()
+
+public:
+
+	struct FFunction
+	{
+		FSimpleDelegate Delegate;
+
+#if WITH_EDITORONLY_DATA
+		FName ImageStyle;  // Button icon, optional if a name is provided, but can have both
+		FText Name;  // Display Name, optional if an icon is provided, but can have both
+		FText ToolTip;
+#endif
+
+		FFunction(const FName& InImageStyle, const FText& InName, const FText& InToolTip, FSimpleDelegate&& Delegate)
+			: Delegate(MoveTemp(Delegate))
+#if WITH_EDITORONLY_DATA
+			, ImageStyle(InImageStyle)
+			, Name(InName)
+			, ToolTip(InToolTip)
+#endif
+		{}
+	};
+
+	FDataflowFunctionsProperty() = default;
+
+	explicit FDataflowFunctionsProperty(TArray<FFunction>&& InFunctions) { Functions = MoveTemp(InFunctions); }
+
+	const TArray<FFunction>& GetFunctions() const { return Functions; }
+
+private:
+	TArray<FFunction> Functions;
+};
