Dataflow - Update function customization to work with one property per function.

#rnx

diff --git a/Engine/Plugins/Experimental/Dataflow/Source/DataflowEditor/Private/Dataflow/DataflowEditorModule.cpp b/Engine/Plugins/Experimental/Dataflow/Source/DataflowEditor/Private/Dataflow/DataflowEditorModule.cpp
--- a/Engine/Plugins/Experimental/Dataflow/Source/DataflowEditor/Private/Dataflow/DataflowEditorModule.cpp
+++ b/Engine/Plugins/Experimental/Dataflow/Source/DataflowEditor/Private/Dataflow/DataflowEditorModule.cpp
@@ -7,8 +7,8 @@
 #include "Dataflow/DataflowEditorToolkit.h"
 #include "Dataflow/DataflowEditorCommands.h"
 #include "Dataflow/DataflowEngineRendering.h"
-#include "Dataflow/DataflowFunctionsProperty.h"
-#include "Dataflow/DataflowFunctionsPropertyCustomization.h"
+#include "Dataflow/DataflowFunctionProperty.h"
+#include "Dataflow/DataflowFunctionPropertyCustomization.h"
 #include "Dataflow/DataflowSNodeFactories.h"
 #include "Dataflow/ScalarVertexPropertyGroupCustomization.h"
 #include "Dataflow/DataflowCollectionAddScalarVertexPropertyNode.h"
@@ -28,7 +28,7 @@ void FDataflowEditorModule::StartupModule()
 	if (FPropertyEditorModule* const PropertyModule = FModuleManager::GetModulePtr<FPropertyEditorModule>("PropertyEditor"))
 	{
 		PropertyModule->RegisterCustomPropertyTypeLayout(FScalarVertexPropertyGroup::StaticStruct()->GetFName(), FOnGetPropertyTypeCustomizationInstance::CreateStatic(&Dataflow::FScalarVertexPropertyGroupCustomization::MakeInstance));
-		PropertyModule->RegisterCustomPropertyTypeLayout(FDataflowFunctionsProperty::StaticStruct()->GetFName(), FOnGetPropertyTypeCustomizationInstance::CreateStatic(&Dataflow::FFunctionsPropertyCustomization::MakeInstance));
+		PropertyModule->RegisterCustomPropertyTypeLayout(FDataflowFunctionProperty::StaticStruct()->GetFName(), FOnGetPropertyTypeCustomizationInstance::CreateStatic(&Dataflow::FFunctionPropertyCustomization::MakeInstance));
 	}
 
 	Dataflow::RenderingCallbacks();
@@ -42,7 +42,7 @@ void FDataflowEditorModule::ShutdownModule()
 	if (FPropertyEditorModule* const PropertyModule = FModuleManager::GetModulePtr<FPropertyEditorModule>("PropertyEditor"))
 	{
 		PropertyModule->UnregisterCustomPropertyTypeLayout(FScalarVertexPropertyGroup::StaticStruct()->GetFName());
-		PropertyModule->UnregisterCustomPropertyTypeLayout(FDataflowFunctionsProperty::StaticStruct()->GetFName());
+		PropertyModule->UnregisterCustomPropertyTypeLayout(FDataflowFunctionProperty::StaticStruct()->GetFName());
 	}
 }
 
diff --git a/Engine/Plugins/Experimental/Dataflow/Source/DataflowEditor/Private/Dataflow/DataflowFunctionPropertyCustomization.cpp b/Engine/Plugins/Experimental/Dataflow/Source/DataflowEditor/Private/Dataflow/DataflowFunctionPropertyCustomization.cpp
new file mode 100644
--- /dev/null
+++ b/Engine/Plugins/Experimental/Dataflow/Source/DataflowEditor/Private/Dataflow/DataflowFunctionPropertyCustomization.cpp
@@ -0,0 +1,187 @@
+// Copyright Epic Games, Inc. All Rights Reserved.
+
+#include "Dataflow/DataflowFunctionPropertyCustomization.h"
+#include "Dataflow/DataflowFunctionProperty.h"
+#include "Widgets/Images/SImage.h"
+#include "Widgets/Input/SButton.h"
+#include "Widgets/Layout/SWrapBox.h"
+#include "Widgets/Text/STextBlock.h"
+#include "DetailWidgetRow.h"
+#include "PropertyHandle.h"
+
+#define LOCTEXT_NAMESPACE "FunctionPropertyCustomization"
+
+namespace Dataflow
+{
+	namespace Private
+	{
+		static const FDataflowFunctionProperty* GetFunctionProperty(const TSharedPtr<IPropertyHandle>& PropertyHandle)
+		{
+			const FDataflowFunctionProperty* FunctionProperty = nullptr;
+			if (PropertyHandle)
+			{
+				if (const FStructProperty* const StructProperty = CastField<FStructProperty>(PropertyHandle->GetProperty()))
+				{
+					if (StructProperty->Struct && StructProperty->Struct->IsChildOf(FDataflowFunctionProperty::StaticStruct()))
+					{
+						void* Data;
+						if (PropertyHandle->GetValueData(Data) == FPropertyAccess::Success)
+						{
+							FunctionProperty = reinterpret_cast<const FDataflowFunctionProperty*>(Data);
+						}
+					}
+				}
+			}
+			return FunctionProperty;
+		}
+	}
+
+	TSharedRef<IPropertyTypeCustomization> FFunctionPropertyCustomization::MakeInstance()
+	{
+		return MakeShareable(new FFunctionPropertyCustomization);
+	}
+
+	void FFunctionPropertyCustomization::CustomizeHeader(TSharedRef<IPropertyHandle> StructPropertyHandle, FDetailWidgetRow& HeaderRow, IPropertyTypeCustomizationUtils& CustomizationUtils)
+	{
+		StructProperty = StructPropertyHandle;
+		// Find all functions under the same category
+		const FName CategoryName = StructProperty->GetDefaultCategoryName();
+		const TSharedPtr<IPropertyHandle> OwnerProperty = StructProperty->GetParentHandle();
+
+		uint32 NumChildren;
+		const FPropertyAccess::Result Result = OwnerProperty->GetNumChildren(NumChildren);
+		check(Result == FPropertyAccess::Success);
+
+		TSharedPtr<SWrapBox> WrapBox;
+		bool bFirstFunction = true;
+
+		for (uint32 Index = 0; Index < NumChildren; ++Index)
+		{
+			const TSharedPtr<IPropertyHandle> ChildProperty = OwnerProperty->GetChildHandle(Index);
+
+			if (const FDataflowFunctionProperty* const FunctionProperty = Private::GetFunctionProperty(ChildProperty))
+			{
+				if (ChildProperty->GetDefaultCategoryName() == CategoryName)
+				{
+					if (bFirstFunction)
+					{
+						if (ChildProperty->GetProperty() != StructProperty->GetProperty())
+						{
+							return;  // The first function property draws all others, no need to do it again
+						}
+						bFirstFunction = false;
+
+						// Create WrapBox
+						HeaderRow
+						.EditCondition(true, FOnBooleanValueChanged())  // Add a custom edit condition so that the first button doesn't affect the others
+						[
+							SAssignNew(WrapBox, SWrapBox)
+							.PreferredSize(2000)  // Copied from FObjectDetails::AddCallInEditorMethods()
+							.UseAllottedSize(true)
+						];
+					}
+
+					auto OnClicked = [ChildProperty, FunctionProperty]() -> FReply
+						{
+							// Execute function
+							FunctionProperty->Execute();
+							ChildProperty->NotifyFinishedChangingProperties();  // Triggers node invalidation
+							return FReply::Handled();
+						};
+
+					const FText Name = ChildProperty->HasMetaData("ButtonText") ?
+						FText::FromString(ChildProperty->GetMetaData("ButtonText")) :
+						ChildProperty->GetPropertyDisplayName();
+					const FText ToolTip = ChildProperty->GetToolTipText();
+					const FName ButtonImage = *ChildProperty->GetMetaData("ButtonImage");  // e.g. FAppStyle::GetBrush("Persona.ReimportAsset") will be Meta = (ButtonImage = "Persona.ReimportAsset")
+
+					if (!ButtonImage.IsNone())
+					{
+						if (Name.IsEmptyOrWhitespace())
+						{
+							// Create button image
+							WrapBox->AddSlot()
+							.Padding(0.f, 0.f, 5.f, 3.f)
+							[
+								SNew(SButton)
+								.ToolTipText(ToolTip)
+								.OnClicked_Lambda(MoveTemp(OnClicked))
+								.ContentPadding(FMargin(0.f, 4.f))  // Too much horizontal padding otherwise (default is 4, 2)
+								.IsEnabled_Lambda([ChildProperty]() -> bool { return ChildProperty->IsEditable(); })
+								.Visibility_Lambda([ChildProperty]() -> EVisibility
+									{
+										return !ChildProperty->HasMetaData("EditConditionHides") || ChildProperty->IsEditable() ? EVisibility::Visible : EVisibility::Collapsed;
+									})
+								[
+									SNew(SImage)
+										.DesiredSizeOverride(FVector2D(16, 16))
+										.Image(FAppStyle::GetBrush(ButtonImage))
+								]
+							];
+						}
+						else
+						{
+							// Create button image + text
+							WrapBox->AddSlot()
+							.Padding(0.f, 0.f, 5.f, 3.f)
+							[
+								SNew(SButton)
+								.Text(Name)
+								.ToolTipText(ToolTip)
+								.OnClicked_Lambda(MoveTemp(OnClicked))
+								.ContentPadding(FMargin(0.f, 2.f))  // Too much horizontal padding otherwise (default is 4, 2)
+								.IsEnabled_Lambda([ChildProperty]() -> bool { return ChildProperty->IsEditable(); })
+								.Visibility_Lambda([ChildProperty]() -> EVisibility
+									{
+										return !ChildProperty->HasMetaData("EditConditionHides") || ChildProperty->IsEditable() ? EVisibility::Visible : EVisibility::Collapsed;
+									})
+								[
+									SNew(SHorizontalBox)
+									+ SHorizontalBox::Slot()
+									.AutoWidth()
+									.Padding(0, 2)
+									[
+										SNew(SImage)
+										.DesiredSizeOverride(FVector2D(16, 16))
+										.Image(FAppStyle::GetBrush(ButtonImage))
+									]
+									+ SHorizontalBox::Slot()
+									.AutoWidth()
+									.Padding(5.f, 0.f)
+									+ SHorizontalBox::Slot()
+									.AutoWidth()
+									.VAlign(VAlign_Center)
+									[
+										SNew(STextBlock)
+											.Text(Name)
+									]
+								]
+							];
+						}
+					}
+					else
+					{
+						// Create button text
+						WrapBox->AddSlot()
+						.Padding(0.f, 0.f, 5.f, 3.f)
+						[
+							SNew(SButton)
+							.Text(Name)
+							.ToolTipText(ToolTip)
+							.OnClicked_Lambda(MoveTemp(OnClicked))
+							.ContentPadding(FMargin(0.f, 2.f))  // Too much horizontal padding otherwise (default is 4, 2)
+							.VAlign(VAlign_Center)  // No Slot, so need vertical align
+							.IsEnabled_Lambda([ChildProperty]() -> bool { return ChildProperty->IsEditable(); })
+							.Visibility_Lambda([ChildProperty]() -> EVisibility
+								{
+									return !ChildProperty->HasMetaData("EditConditionHides") || ChildProperty->IsEditable() ? EVisibility::Visible : EVisibility::Collapsed;
+								})
+						];
+					}
+				}
+			}
+		}
+	}
+}  // End namespace Dataflow
+
+#undef LOCTEXT_NAMESPACE
diff --git a/Engine/Plugins/Experimental/Dataflow/Source/DataflowEditor/Private/Dataflow/DataflowFunctionPropertyCustomization.h b/Engine/Plugins/Experimental/Dataflow/Source/DataflowEditor/Private/Dataflow/DataflowFunctionPropertyCustomization.h
new file mode 100644
--- /dev/null
+++ b/Engine/Plugins/Experimental/Dataflow/Source/DataflowEditor/Private/Dataflow/DataflowFunctionPropertyCustomization.h
@@ -0,0 +1,24 @@
+// Copyright Epic Games, Inc. All Rights Reserved.
+
+#pragma once
+
+#include "IPropertyTypeCustomization.h"
+
+namespace Dataflow
+{
+	/**
+	 * Customization for buttons in the UStruct UI of dataflow nodes.
+	 */
+	class FFunctionPropertyCustomization : public IPropertyTypeCustomization
+	{
+	public:
+		static TSharedRef<IPropertyTypeCustomization> MakeInstance();
+
+	protected:
+		virtual void CustomizeHeader(TSharedRef<IPropertyHandle> StructPropertyHandle, FDetailWidgetRow& HeaderRow, IPropertyTypeCustomizationUtils& CustomizationUtils) override;
+		virtual void CustomizeChildren(TSharedRef<IPropertyHandle> /*StructPropertyHandle*/, IDetailChildrenBuilder& /*ChildBuilder*/, IPropertyTypeCustomizationUtils& /*CustomizationUtils*/) override {}
+
+	private:
+		TSharedPtr<IPropertyHandle> StructProperty;
+	};
+}
diff --git a/Engine/Plugins/Experimental/Dataflow/Source/DataflowEditor/Private/Dataflow/DataflowFunctionsPropertyCustomization.cpp b/Engine/Plugins/Experimental/Dataflow/Source/DataflowEditor/Private/Dataflow/DataflowFunctionsPropertyCustomization.cpp
deleted file mode 100644
diff --git a/Engine/Plugins/Experimental/Dataflow/Source/DataflowEditor/Private/Dataflow/DataflowFunctionsPropertyCustomization.h b/Engine/Plugins/Experimental/Dataflow/Source/DataflowEditor/Private/Dataflow/DataflowFunctionsPropertyCustomization.h
deleted file mode 100644
diff --git a/Engine/Plugins/Experimental/Dataflow/Source/DataflowNodes/Public/Dataflow/DataflowFunctionProperty.h b/Engine/Plugins/Experimental/Dataflow/Source/DataflowNodes/Public/Dataflow/DataflowFunctionProperty.h
new file mode 100644
--- /dev/null
+++ b/Engine/Plugins/Experimental/Dataflow/Source/DataflowNodes/Public/Dataflow/DataflowFunctionProperty.h
@@ -0,0 +1,53 @@
+// Copyright Epic Games, Inc. All Rights Reserved.
+
+#pragma once
+
+#include "Delegates/Delegate.h"
+#include "DataflowFunctionProperty.generated.h"
+
+/**
+ * Function property for all Dataflow nodes.
+ * The structure is also used in DataFlow::FFunctionPropertyCustomization to appear as text and/or image buttons.
+ * This helps with the equivalent UCLASS UFUNCTION CallInEditor functionality that is missing from the USTRUCT implementation.
+ *
+ * By default the text of the button is the name of the structure property.
+ * The tooltip is the property source documentation.
+ * Further (but optional) customizations can be achieved by using the following Meta tags where declaring the property:
+ *   ButtonText
+ *   ButtonImage
+ *
+ * Specifying an empty ButtonText string will only display the icon and no text.
+ *
+ * For example:
+ *   UPROPERTY(EditAnywhere, Category = "Functions")
+ *   FDataflowFunctionProperty ReimportAssetTextOnly;
+ *
+ *   UPROPERTY(EditAnywhere, Category = "Functions", Meta = (ButtonImage = "Persona.ReimportAsset"))
+ *   FDataflowFunctionProperty ReimportAssetTextAndIcon;
+ * 
+ *   UPROPERTY(EditAnywhere, Category = "Functions", Meta = (ButtonText = "", ButtonImage = "Persona.ReimportAsset"))
+ *   FDataflowFunctionProperty ReimportAssetIconOnly;
+ *
+ *   UPROPERTY(EditAnywhere, Category = "Functions", Meta = (ButtonText = "Reimport Asset"))
+ *   FDataflowFunctionProperty ReimportAssetOverriddenText;
+ *
+ *   UPROPERTY(EditAnywhere, Category = "Functions", Meta = (ButtonText = "Reimport Asset", ButtonImage = "Persona.ReimportAsset"))
+ *   FDataflowFunctionProperty ReimportAssetOverriddenTextAndIcon;
+ *
+ */
+USTRUCT()
+struct FDataflowFunctionProperty
+{
+	GENERATED_BODY()
+
+public:
+
+	FDataflowFunctionProperty() = default;
+
+	explicit FDataflowFunctionProperty(FSimpleDelegate&& InDelegate) { Delegate = MoveTemp(InDelegate); }
+
+	void Execute() const { Delegate.ExecuteIfBound(); }
+
+private:
+	FSimpleDelegate Delegate;
+};
diff --git a/Engine/Plugins/Experimental/Dataflow/Source/DataflowNodes/Public/Dataflow/DataflowFunctionsProperty.h b/Engine/Plugins/Experimental/Dataflow/Source/DataflowNodes/Public/Dataflow/DataflowFunctionsProperty.h
deleted file mode 100644
