Dataflow Editor: fix crash when resetting simulation

diff --git a/Engine/Plugins/Experimental/Dataflow/Source/DataflowEditor/Private/Dataflow/DataflowSimulationScene.cpp b/Engine/Plugins/Experimental/Dataflow/Source/DataflowEditor/Private/Dataflow/DataflowSimulationScene.cpp
--- a/Engine/Plugins/Experimental/Dataflow/Source/DataflowEditor/Private/Dataflow/DataflowSimulationScene.cpp
+++ b/Engine/Plugins/Experimental/Dataflow/Source/DataflowEditor/Private/Dataflow/DataflowSimulationScene.cpp
@@ -8,6 +8,8 @@
 #include "Misc/TransactionObjectEvent.h"
 #include "EngineUtils.h"
 #include "Animation/AnimSingleNodeInstance.h"
+#include "AssetEditorModeManager.h"
+#include "Engine/Selection.h"
 
 #define LOCTEXT_NAMESPACE "FDataflowSimulationScene"
 
@@ -63,6 +65,15 @@ void FDataflowSimulationScene::UnbindSceneSelection()
 
 void FDataflowSimulationScene::ResetSimulationScene()
 {
+	// Release any selected components before the PreviewActor is deleted from the scene
+	if (const TSharedPtr<FAssetEditorModeManager> ModeManager = GetDataflowModeManager())
+	{
+		if (USelection* const SelectedComponents = ModeManager->GetSelectedComponents())
+		{
+			SelectedComponents->DeselectAll();
+		}
+	}
+
 	// Destroy the spawned root actor
 	if(PreviewActor && GetWorld())
 	{
