Fix use-after-free error: we were accessing the name of a property that has already been unloaded in a ShutdownModule call.

ASAN preflight: 66b3a23ea82c869bf2232b71

diff --git a/Engine/Plugins/Experimental/Dataflow/Source/DataflowEditor/Private/Dataflow/DataflowEditorModule.cpp b/Engine/Plugins/Experimental/Dataflow/Source/DataflowEditor/Private/Dataflow/DataflowEditorModule.cpp
--- a/Engine/Plugins/Experimental/Dataflow/Source/DataflowEditor/Private/Dataflow/DataflowEditorModule.cpp
+++ b/Engine/Plugins/Experimental/Dataflow/Source/DataflowEditor/Private/Dataflow/DataflowEditorModule.cpp
@@ -19,6 +19,8 @@
 #define LOCTEXT_NAMESPACE "DataflowEditor"
 
 const FColor FDataflowEditorModule::SurfaceColor = FLinearColor(0.6, 0.6, 0.6).ToRGBE();
+static const FName ScalarVertexPropertyGroupName = TEXT("ScalarVertexPropertyGroup");
+static const FName DataflowFunctionPropertyName = TEXT("DataflowFunctionProperty");
 
 void FDataflowEditorModule::StartupModule()
 {
@@ -27,8 +29,8 @@ void FDataflowEditorModule::StartupModule()
 	// Register type customizations
 	if (FPropertyEditorModule* const PropertyModule = FModuleManager::GetModulePtr<FPropertyEditorModule>("PropertyEditor"))
 	{
-		PropertyModule->RegisterCustomPropertyTypeLayout(FScalarVertexPropertyGroup::StaticStruct()->GetFName(), FOnGetPropertyTypeCustomizationInstance::CreateStatic(&Dataflow::FScalarVertexPropertyGroupCustomization::MakeInstance));
-		PropertyModule->RegisterCustomPropertyTypeLayout(FDataflowFunctionProperty::StaticStruct()->GetFName(), FOnGetPropertyTypeCustomizationInstance::CreateStatic(&Dataflow::FFunctionPropertyCustomization::MakeInstance));
+		PropertyModule->RegisterCustomPropertyTypeLayout(ScalarVertexPropertyGroupName, FOnGetPropertyTypeCustomizationInstance::CreateStatic(&Dataflow::FScalarVertexPropertyGroupCustomization::MakeInstance));
+		PropertyModule->RegisterCustomPropertyTypeLayout(DataflowFunctionPropertyName, FOnGetPropertyTypeCustomizationInstance::CreateStatic(&Dataflow::FFunctionPropertyCustomization::MakeInstance));
 	}
 
 	Dataflow::RenderingCallbacks();
@@ -41,8 +43,8 @@ void FDataflowEditorModule::ShutdownModule()
 	// Deregister type customizations
 	if (FPropertyEditorModule* const PropertyModule = FModuleManager::GetModulePtr<FPropertyEditorModule>("PropertyEditor"))
 	{
-		PropertyModule->UnregisterCustomPropertyTypeLayout(FScalarVertexPropertyGroup::StaticStruct()->GetFName());
-		PropertyModule->UnregisterCustomPropertyTypeLayout(FDataflowFunctionProperty::StaticStruct()->GetFName());
+		PropertyModule->UnregisterCustomPropertyTypeLayout(ScalarVertexPropertyGroupName);
+		PropertyModule->UnregisterCustomPropertyTypeLayout(DataflowFunctionPropertyName);
 	}
 }
 
