Fix CIS Clang build: auto-generated constructors warning about deprecated members

Win64 Clang Preflight: 66b637234b50a6f1dc56a729

diff --git a/Engine/Plugins/Experimental/Dataflow/Source/DataflowEditor/Public/Dataflow/DataflowEditorCommands.h b/Engine/Plugins/Experimental/Dataflow/Source/DataflowEditor/Public/Dataflow/DataflowEditorCommands.h
--- a/Engine/Plugins/Experimental/Dataflow/Source/DataflowEditor/Public/Dataflow/DataflowEditorCommands.h
+++ b/Engine/Plugins/Experimental/Dataflow/Source/DataflowEditor/Public/Dataflow/DataflowEditorCommands.h
@@ -30,6 +30,16 @@ public:
 
 	FDataflowEditorCommandsImpl();
 
+	// Note: We need to explicitly disable warnings on these constructors/operators for clang to be happy with deprecated variables
+	// (See Confluence page "How to Deprecate Code in UE > How to deprecate variables inside a UStruct")
+	PRAGMA_DISABLE_DEPRECATION_WARNINGS
+	~FDataflowEditorCommandsImpl() = default;
+	FDataflowEditorCommandsImpl(const FDataflowEditorCommandsImpl&) = default;
+	FDataflowEditorCommandsImpl(FDataflowEditorCommandsImpl&&) = default;
+	FDataflowEditorCommandsImpl& operator=(const FDataflowEditorCommandsImpl&) = default;
+	FDataflowEditorCommandsImpl& operator=(FDataflowEditorCommandsImpl&&) = default;
+	PRAGMA_ENABLE_DEPRECATION_WARNINGS
+
 	// TBaseCharacterFXEditorCommands<> interface
 	 virtual void RegisterCommands() override;
 
