Dataflow Editor:
- if p.Dataflow.Editor.ContextCaching CVAR is false, mark the UDataflowBaseContent object as transient to prevent it from setting the Dataflow asset's package dirty
- when user hits the Save button, only prompt to checkout and save assets that are actuallly modified

diff --git a/Engine/Plugins/Experimental/Dataflow/Source/DataflowEditor/Private/Dataflow/DataflowEditorToolkit.cpp b/Engine/Plugins/Experimental/Dataflow/Source/DataflowEditor/Private/Dataflow/DataflowEditorToolkit.cpp
--- a/Engine/Plugins/Experimental/Dataflow/Source/DataflowEditor/Private/Dataflow/DataflowEditorToolkit.cpp
+++ b/Engine/Plugins/Experimental/Dataflow/Source/DataflowEditor/Private/Dataflow/DataflowEditorToolkit.cpp
@@ -81,6 +81,10 @@ FDataflowEditorToolkit::FDataflowEditorToolkit(UAssetEditor* InOwningAssetEditor
 	: FBaseCharacterFXEditorToolkit(InOwningAssetEditor, FName("DataflowEditor"))
 	, DataflowEditor(Cast< UDataflowEditor>(InOwningAssetEditor))
 {
+
+	// When saving, only prompt to checkout and save assets that are actually modified
+	bCheckDirtyOnAssetSave = true;
+
 	check(DataflowEditor);
 
 	ConstructionDefaultLayout = FTabManager::NewLayout(FName("DataflowConstructionLayout03"))
@@ -533,7 +537,7 @@ void FDataflowEditorToolkit::GetSaveableObjects(TArray<UObject*>& OutObjects) co
 		if (UDataflow* DataflowAsset = GetEditorContent()->GetDataflowAsset())
 		{
 			check(DataflowAsset->IsAsset());
-			OutObjects.Add(DataflowAsset);
+			OutObjects.AddUnique(DataflowAsset);
 		}
 	}
 }
diff --git a/Engine/Source/Runtime/Experimental/Dataflow/Engine/Private/Dataflow/DataflowContent.cpp b/Engine/Source/Runtime/Experimental/Dataflow/Engine/Private/Dataflow/DataflowContent.cpp
--- a/Engine/Source/Runtime/Experimental/Dataflow/Engine/Private/Dataflow/DataflowContent.cpp
+++ b/Engine/Source/Runtime/Experimental/Dataflow/Engine/Private/Dataflow/DataflowContent.cpp
@@ -232,7 +232,7 @@ namespace DataflowContextHelpers
 		}
 		else
 		{
-			Asset = NewObject<T>(ContentOwner, T::StaticClass());
+			Asset = NewObject<T>(GetTransientPackage(), T::StaticClass(), NAME_None, RF_Transient);
 			if (UDataflowBaseContent* BaseContent = Cast< UDataflowBaseContent>(Asset.Get()))
 			{
 				BaseContent->SetIsSaved(false);
