ChaosFlesh: Dataflow skeleton view fix

diff --git a/Engine/Plugins/Experimental/Dataflow/Source/DataflowEditor/Private/Dataflow/DataflowEditorToolkit.cpp b/Engine/Plugins/Experimental/Dataflow/Source/DataflowEditor/Private/Dataflow/DataflowEditorToolkit.cpp
--- a/Engine/Plugins/Experimental/Dataflow/Source/DataflowEditor/Private/Dataflow/DataflowEditorToolkit.cpp
+++ b/Engine/Plugins/Experimental/Dataflow/Source/DataflowEditor/Private/Dataflow/DataflowEditorToolkit.cpp
@@ -1307,7 +1307,7 @@ TSharedRef<SDockTab> FDataflowEditorToolkit::SpawnTab_SkeletonView(const FSpawnT
 
 	TSharedPtr<ISkeletonTree> SkeletonEditor = SkeletonEditorView->CreateEditor(SkeletonTreeArgs);
 	return SNew(SDockTab)
-		.Label(LOCTEXT("DataflowEditor_Skeleton_TabTitle", "Skeleton View"))
+		.Label(LOCTEXT("DataflowEditor_Outliner_TabTitle", "Outliner"))
 		[
 			SkeletonEditor.ToSharedRef()
 		];
@@ -1501,7 +1501,7 @@ void FDataflowEditorToolkit::RegisterTabSpawners(const TSharedRef<FTabManager>&
 	if (bDataflowEnableSkeletonView)
 	{
 		InTabManager->RegisterTabSpawner(SkeletonViewTabId, FOnSpawnTab::CreateSP(this, &FDataflowEditorToolkit::SpawnTab_SkeletonView))
-			.SetDisplayName(LOCTEXT("DataflowSkeletonTab", "Skeleton View"))
+			.SetDisplayName(LOCTEXT("OutlinerTab", "Outliner"))
 			.SetGroup(EditorMenuCategory.ToSharedRef())
 			.SetIcon(FSlateIcon(FAppStyle::GetAppStyleSetName(), "LevelEditor.Tabs.SkeletonHierarchy"));
 	}
diff --git a/Engine/Plugins/Experimental/Dataflow/Source/DataflowEditor/Private/Dataflow/DataflowSkeletonView.cpp b/Engine/Plugins/Experimental/Dataflow/Source/DataflowEditor/Private/Dataflow/DataflowSkeletonView.cpp
--- a/Engine/Plugins/Experimental/Dataflow/Source/DataflowEditor/Private/Dataflow/DataflowSkeletonView.cpp
+++ b/Engine/Plugins/Experimental/Dataflow/Source/DataflowEditor/Private/Dataflow/DataflowSkeletonView.cpp
@@ -13,7 +13,7 @@
 #include "ISkeletonEditorModule.h"
 #include "Engine/SkeletalMesh.h"
 
-
+const FName FDataflowSkeletonView::SkeletonName("DataflowSkeleton");
 FDataflowSkeletonView::FDataflowSkeletonView(TObjectPtr<UDataflowBaseContent> InContent)
 	: FDataflowNodeView(InContent)
 	, SkeletonEditor(nullptr)
@@ -21,16 +21,7 @@ FDataflowSkeletonView::FDataflowSkeletonView(TObjectPtr<UDataflowBaseContent> In
 	, CollectionIndexRemap(TArray<int32>())
 {
 	check(InContent);
-
-	SetSkeleton(NewObject<USkeleton>(SkeletalMesh, NAME_Name));
-	if (const TObjectPtr<UDataflowSkeletalContent> SkeletalContent = Cast<UDataflowSkeletalContent>(InContent))
-	{
-		if (SkeletalContent->GetDataflowAsset())
-		{
-			SetSkeleton(SkeletalContent->GetSkeleton());
-		}
-	}
-
+	UpdateSkeleton();
 }
 
 FDataflowSkeletonView::~FDataflowSkeletonView()
@@ -44,7 +35,7 @@ FDataflowSkeletonView::~FDataflowSkeletonView()
 TSharedPtr<ISkeletonTree> FDataflowSkeletonView::CreateEditor(FSkeletonTreeArgs& InSkeletonTreeArgs)
 {
 	ISkeletonEditorModule& SkeletonEditorModule = FModuleManager::LoadModuleChecked<ISkeletonEditorModule>("SkeletonEditor");
-	SkeletonEditor = SkeletonEditorModule.CreateSkeletonTree(GetSkeleton(), InSkeletonTreeArgs);
+	SkeletonEditor = SkeletonEditorModule.CreateSkeletonTree(Skeleton, InSkeletonTreeArgs);
 	SkeletonEditor->Refresh();
 	// add widget delegates (see FDataflowCollectionSpreadSheet)
 	return SkeletonEditor;
@@ -57,21 +48,19 @@ void FDataflowSkeletonView::SetSupportedOutputTypes()
 	GetSupportedOutputTypes().Add("FManagedArrayCollection");
 }
 
-void FDataflowSkeletonView::SetSkeleton(USkeleton* Skeleton)
+void FDataflowSkeletonView::UpdateSkeleton()
 {
-	if (Skeleton)
+	if (!Skeleton)
 	{
-		SkeletalMesh = NewObject<USkeletalMesh>();
-		SkeletalMesh->SetSkeleton(Skeleton);
-		SkeletalMesh->SetRefSkeleton(Skeleton->GetReferenceSkeleton());
-	}
-	else
-	{
-		SkeletalMesh = NewObject<USkeletalMesh>();
-		SkeletalMesh->SetSkeleton(NewObject<USkeleton>(SkeletalMesh, NAME_Name));
+		Skeleton = NewObject<USkeleton>(SkeletalMesh, SkeletonName);
 	}
+	SkeletalMesh->SetSkeleton(Skeleton);
+	SkeletalMesh->SetRefSkeleton(Skeleton->GetReferenceSkeleton());
+
 	if (SkeletonEditor)
 	{
+		SkeletonEditor->GetEditableSkeleton()->RecreateBoneTree(SkeletalMesh);
+		SkeletonEditor->SetSkeletalMesh(SkeletalMesh);
 		SkeletonEditor->Refresh();
 	}
 }
@@ -87,12 +76,13 @@ USkeleton* FDataflowSkeletonView::GetSkeleton()
 
 void FDataflowSkeletonView::UpdateViewData()
 {
-	bool bNeedsDefaultSkeleton = true;
+	bool bClearSkeleton = true;
+	Skeleton = NewObject<USkeleton>(SkeletalMesh, SkeletonName);
 	if (TObjectPtr<UDataflowEdNode> EdNode = GetSelectedNode())
 	{
-		if (GetSelectedNode()->IsBound())
+		if (EdNode->IsBound())
 		{
-			if (TSharedPtr<FDataflowNode> Node = GetSelectedNode()->DataflowGraph->FindBaseNode(GetSelectedNode()->DataflowNodeGuid))
+			if (TSharedPtr<FDataflowNode> Node = EdNode->DataflowGraph->FindBaseNode(EdNode->DataflowNodeGuid))
 			{
 				if (FDataflowOutput* Output = Node->FindOutput(FName("Collection")))
 				{
@@ -103,20 +93,7 @@ void FDataflowSkeletonView::UpdateViewData()
 							FManagedArrayCollection DefaultCollection;
 							const FManagedArrayCollection& Result = Output->GetValue(*Context, DefaultCollection);
 
-							SkeletalMesh = NewObject<USkeletalMesh>();
-							TObjectPtr<USkeleton> Skeleton = NewObject<USkeleton>(SkeletalMesh, Node->Name);
-
 							FGeometryCollectionEngineConversion::ConvertCollectionToSkeleton(Result, Skeleton, CollectionIndexRemap);
-							SkeletalMesh->SetSkeleton(Skeleton);
-							SkeletalMesh->SetRefSkeleton(Skeleton->GetReferenceSkeleton());
-
-							if (SkeletonEditor)
-							{
-								SkeletonEditor->GetEditableSkeleton()->RecreateBoneTree(SkeletalMesh);
-								SkeletonEditor->SetSkeletalMesh(SkeletalMesh);
-								SkeletonEditor->Refresh();
-							}
-							bNeedsDefaultSkeleton = false;
 						}
 					}
 				}
@@ -124,22 +101,7 @@ void FDataflowSkeletonView::UpdateViewData()
 		}
 	}
 
-
-	if (bNeedsDefaultSkeleton)
-	{
-		SetSkeleton(NewObject<USkeleton>(SkeletalMesh, NAME_Name));
-		if (const TObjectPtr<UDataflowSkeletalContent> SkeletalContent = Cast<UDataflowSkeletalContent>(GetEditorContent()))
-		{
-			if (SkeletalContent->GetDataflowAsset())
-			{
-				SetSkeleton(SkeletalContent->GetSkeleton());
-			}
-		}
-		if (SkeletonEditor)
-		{
-			SkeletonEditor->Refresh();
-		}
-	}
+	UpdateSkeleton();
 }
 
 void FDataflowSkeletonView::ConstructionViewSelectionChanged(const TArray<UPrimitiveComponent*>& InSelectedComponents)
@@ -171,5 +133,6 @@ void FDataflowSkeletonView::SkeletonViewSelectionChanged(const TArrayView<TShare
 void FDataflowSkeletonView::AddReferencedObjects(FReferenceCollector& Collector)
 {
 	FDataflowNodeView::AddReferencedObjects(Collector);
+	Collector.AddReferencedObject(Skeleton);
 	Collector.AddReferencedObject(SkeletalMesh);
 }
diff --git a/Engine/Plugins/Experimental/Dataflow/Source/DataflowEditor/Public/Dataflow/DataflowSkeletonView.h b/Engine/Plugins/Experimental/Dataflow/Source/DataflowEditor/Public/Dataflow/DataflowSkeletonView.h
--- a/Engine/Plugins/Experimental/Dataflow/Source/DataflowEditor/Public/Dataflow/DataflowSkeletonView.h
+++ b/Engine/Plugins/Experimental/Dataflow/Source/DataflowEditor/Public/Dataflow/DataflowSkeletonView.h
@@ -33,7 +33,7 @@ public:
 	USkeleton* GetSkeleton();
 
 	/** Update Data*/
-	void SetSkeleton(USkeleton* Skeleton);
+	void UpdateSkeleton();
 
 	/** Update the view */
 	virtual void UpdateViewData() override;
@@ -47,12 +47,16 @@ public:
 	/** Selection View Callbacks */
 	void SkeletonViewSelectionChanged(const TArrayView<TSharedPtr<ISkeletonTreeItem>>& InSelectedItems, ESelectInfo::Type InSelectInfo);
 
+	static const FName SkeletonName;
 
 private:
 	TSharedPtr<ISkeletonTree> SkeletonEditor;
 
 	/* Skeletal Mesh in the SkeletalViewer*/
-	TObjectPtr<USkeletalMesh> SkeletalMesh = nullptr;
+	TObjectPtr<USkeletalMesh> SkeletalMesh;
+
+	/* Visualized Skeleton*/
+	TObjectPtr<USkeleton> Skeleton;
 
 	/* Rempping from the selected node to the SkeletalMesh*/
 	TArray<int32> CollectionIndexRemap;
