Dataflow - Change ButtonText meta to the existing DisplayName one.

#rnx

diff --git a/Engine/Plugins/Experimental/Dataflow/Source/DataflowEditor/Private/Dataflow/DataflowFunctionPropertyCustomization.cpp b/Engine/Plugins/Experimental/Dataflow/Source/DataflowEditor/Private/Dataflow/DataflowFunctionPropertyCustomization.cpp
--- a/Engine/Plugins/Experimental/Dataflow/Source/DataflowEditor/Private/Dataflow/DataflowFunctionPropertyCustomization.cpp
+++ b/Engine/Plugins/Experimental/Dataflow/Source/DataflowEditor/Private/Dataflow/DataflowFunctionPropertyCustomization.cpp
@@ -89,8 +89,8 @@ namespace Dataflow
 							return FReply::Handled();
 						};
 
-					const FText Name = ChildProperty->HasMetaData("ButtonText") ?
-						FText::FromString(ChildProperty->GetMetaData("ButtonText")) :
+					const FText Name = ChildProperty->HasMetaData("DisplayName") ?
+						FText::FromString(ChildProperty->GetMetaData("DisplayName")) :
 						ChildProperty->GetPropertyDisplayName();
 					const FText ToolTip = ChildProperty->GetToolTipText();
 					const FName ButtonImage = *ChildProperty->GetMetaData("ButtonImage");  // e.g. FAppStyle::GetBrush("Persona.ReimportAsset") will be Meta = (ButtonImage = "Persona.ReimportAsset")
diff --git a/Engine/Plugins/Experimental/Dataflow/Source/DataflowNodes/Public/Dataflow/DataflowFunctionProperty.h b/Engine/Plugins/Experimental/Dataflow/Source/DataflowNodes/Public/Dataflow/DataflowFunctionProperty.h
--- a/Engine/Plugins/Experimental/Dataflow/Source/DataflowNodes/Public/Dataflow/DataflowFunctionProperty.h
+++ b/Engine/Plugins/Experimental/Dataflow/Source/DataflowNodes/Public/Dataflow/DataflowFunctionProperty.h
@@ -13,10 +13,10 @@
  * By default the text of the button is the name of the structure property.
  * The tooltip is the property source documentation.
  * Further (but optional) customizations can be achieved by using the following Meta tags where declaring the property:
- *   ButtonText
+ *   DisplayName
  *   ButtonImage
  *
- * Specifying an empty ButtonText string will only display the icon and no text.
+ * Specifying an empty DisplayName string will only display the icon and no text.
  *
  * For example:
  *   UPROPERTY(EditAnywhere, Category = "Functions")
@@ -25,13 +25,13 @@
  *   UPROPERTY(EditAnywhere, Category = "Functions", Meta = (ButtonImage = "Persona.ReimportAsset"))
  *   FDataflowFunctionProperty ReimportAssetTextAndIcon;
  * 
- *   UPROPERTY(EditAnywhere, Category = "Functions", Meta = (ButtonText = "", ButtonImage = "Persona.ReimportAsset"))
+ *   UPROPERTY(EditAnywhere, Category = "Functions", Meta = (DisplayName = "", ButtonImage = "Persona.ReimportAsset"))
  *   FDataflowFunctionProperty ReimportAssetIconOnly;
  *
- *   UPROPERTY(EditAnywhere, Category = "Functions", Meta = (ButtonText = "Reimport Asset"))
+ *   UPROPERTY(EditAnywhere, Category = "Functions", Meta = (DisplayName = "Reimport Asset"))
  *   FDataflowFunctionProperty ReimportAssetOverriddenText;
  *
- *   UPROPERTY(EditAnywhere, Category = "Functions", Meta = (ButtonText = "Reimport Asset", ButtonImage = "Persona.ReimportAsset"))
+ *   UPROPERTY(EditAnywhere, Category = "Functions", Meta = (DisplayName = "Reimport Asset", ButtonImage = "Persona.ReimportAsset"))
  *   FDataflowFunctionProperty ReimportAssetOverriddenTextAndIcon;
  *
  */
