Dataflow editor : Fix Geometry collection nodes not showing in dataflow
-  Remove CreateNodesMap as this was precomputed once and only used to gate if a node was to be Added to the action menu

diff --git a/Engine/Plugins/Experimental/Dataflow/Source/DataflowEditor/Private/Dataflow/DataflowEditorCommands.cpp b/Engine/Plugins/Experimental/Dataflow/Source/DataflowEditor/Private/Dataflow/DataflowEditorCommands.cpp
--- a/Engine/Plugins/Experimental/Dataflow/Source/DataflowEditor/Private/Dataflow/DataflowEditorCommands.cpp
+++ b/Engine/Plugins/Experimental/Dataflow/Source/DataflowEditor/Private/Dataflow/DataflowEditorCommands.cpp
@@ -92,26 +92,6 @@ void FDataflowEditorCommandsImpl::RegisterCommands()
 		SetConstructionViewModeCommands.Add(ViewMode->GetName(), SetViewModeCommand);
 	}
 
-	if (Dataflow::FNodeFactory* Factory = Dataflow::FNodeFactory::GetInstance())
-	{
-		for (Dataflow::FFactoryParameters Parameters : Factory->RegisteredParameters())
-		{
-			TSharedPtr< FUICommandInfo > AddNode;
-			FUICommandInfo::MakeCommandInfo(
-				this->AsShared(),
-				AddNode,
-				Parameters.TypeName,
-				LOCTEXT("DataflowButton", "New Dataflow Node"),
-				LOCTEXT("NewDataflowNodeTooltip", "New Dataflow Node Tooltip"),
-				FSlateIcon(),
-				EUserInterfaceActionType::Button,
-				FInputChord()
-			);
-			CreateNodesMap.Add(Parameters.TypeName, AddNode);
-		}
-	}
-
-
 	Dataflow::FDataflowToolRegistry& ToolRegistry = Dataflow::FDataflowToolRegistry::Get();
 	const TArray<FName> NodeNames = ToolRegistry.GetNodeNames();
 	for (const FName& NodeName : NodeNames)
diff --git a/Engine/Plugins/Experimental/Dataflow/Source/DataflowEditor/Private/Dataflow/DataflowSchema.cpp b/Engine/Plugins/Experimental/Dataflow/Source/DataflowEditor/Private/Dataflow/DataflowSchema.cpp
--- a/Engine/Plugins/Experimental/Dataflow/Source/DataflowEditor/Private/Dataflow/DataflowSchema.cpp
+++ b/Engine/Plugins/Experimental/Dataflow/Source/DataflowEditor/Private/Dataflow/DataflowSchema.cpp
@@ -94,13 +94,10 @@ void UDataflowSchema::GetGraphContextActions(FGraphContextMenuBuilder& ContextMe
 	{
 		for (Dataflow::FFactoryParameters NodeParameters : Factory->RegisteredParameters())
 		{
-			if (FDataflowEditorCommands::Get().CreateNodesMap.Contains(NodeParameters.TypeName))
+			if (TSharedPtr<FAssetSchemaAction_Dataflow_CreateNode_DataflowEdNode> Action =
+				FAssetSchemaAction_Dataflow_CreateNode_DataflowEdNode::CreateAction(ContextMenuBuilder.CurrentGraph, NodeParameters.TypeName, NodeParameters.DisplayName))
 			{
-				if (TSharedPtr<FAssetSchemaAction_Dataflow_CreateNode_DataflowEdNode> Action =
-					FAssetSchemaAction_Dataflow_CreateNode_DataflowEdNode::CreateAction(ContextMenuBuilder.CurrentGraph, NodeParameters.TypeName, NodeParameters.DisplayName))
-				{
-					ContextMenuBuilder.AddAction(Action);
-				}
+				ContextMenuBuilder.AddAction(Action);
 			}
 		}
 	}
diff --git a/Engine/Plugins/Experimental/Dataflow/Source/DataflowEditor/Public/Dataflow/DataflowEditorCommands.h b/Engine/Plugins/Experimental/Dataflow/Source/DataflowEditor/Public/Dataflow/DataflowEditorCommands.h
--- a/Engine/Plugins/Experimental/Dataflow/Source/DataflowEditor/Public/Dataflow/DataflowEditorCommands.h
+++ b/Engine/Plugins/Experimental/Dataflow/Source/DataflowEditor/Public/Dataflow/DataflowEditorCommands.h
@@ -66,8 +66,6 @@ public:
 
 	TMap<FName, TSharedPtr<FUICommandInfo>> SetConstructionViewModeCommands;
 
-	TMap< FName, TSharedPtr<FUICommandInfo> > CreateNodesMap;
-
 	UE_DEPRECATED(5.5, "Dataflow Tool commands are now stored in FDataflowToolRegistry")
 	const static FString BeginWeightMapPaintToolIdentifier;
 	UE_DEPRECATED(5.5, "Dataflow Tool commands are now stored in FDataflowToolRegistry")
