Dataflow : avoid the simulation scene from being recreated over and over every tick when the dataflow has no terminal node
Also avoid spawning the actor if the blueprint class is null

diff --git a/Engine/Plugins/Experimental/Dataflow/Source/DataflowEditor/Private/Dataflow/DataflowSimulationScene.cpp b/Engine/Plugins/Experimental/Dataflow/Source/DataflowEditor/Private/Dataflow/DataflowSimulationScene.cpp
--- a/Engine/Plugins/Experimental/Dataflow/Source/DataflowEditor/Private/Dataflow/DataflowSimulationScene.cpp
+++ b/Engine/Plugins/Experimental/Dataflow/Source/DataflowEditor/Private/Dataflow/DataflowSimulationScene.cpp
@@ -149,7 +149,7 @@ void FDataflowSimulationScene::BindSceneSelection()
 
 void FDataflowSimulationScene::CreateSimulationScene()
 {
-	if(SimulationGenerator && SceneDescription && GetWorld())
+	if(SimulationGenerator && SceneDescription && SceneDescription->BlueprintClass && GetWorld())
 	{
 		SimulationGenerator->SetCacheParams(SceneDescription->CacheParams);
 		SimulationGenerator->SetCacheAsset(SceneDescription->CacheAsset);
diff --git a/Engine/Plugins/Experimental/Dataflow/Source/DataflowSimulation/Private/Dataflow/DataflowSimulationUtils.cpp b/Engine/Plugins/Experimental/Dataflow/Source/DataflowSimulation/Private/Dataflow/DataflowSimulationUtils.cpp
--- a/Engine/Plugins/Experimental/Dataflow/Source/DataflowSimulation/Private/Dataflow/DataflowSimulationUtils.cpp
+++ b/Engine/Plugins/Experimental/Dataflow/Source/DataflowSimulation/Private/Dataflow/DataflowSimulationUtils.cpp
@@ -1,4 +1,4 @@
-﻿// Copyright Epic Games, Inc. All Rights Reserved.
+// Copyright Epic Games, Inc. All Rights Reserved.
 
 #include "Dataflow/DataflowSimulationUtils.h"
 
@@ -26,11 +26,11 @@ namespace Dataflow
 			for(const TSharedPtr<FDataflowNode>& TerminalNode : DataflowGraph->GetFilteredNodes(FDataflowTerminalNode::StaticType()))
 			{
 				MaxTimeStamp.Value = FMath::Max(MaxTimeStamp.Value, TerminalNode->GetTimestamp().Value);
-			}
-			if(MaxTimeStamp.Value > LastTimeStamp.Value)
-			{
-				LastTimeStamp = MaxTimeStamp.Value;
-				return true;
+				if (MaxTimeStamp.Value > LastTimeStamp.Value)
+				{
+					LastTimeStamp = MaxTimeStamp.Value;
+					return true;
+				}
 			}
 		}
 		return false;
