Dataflow Editor: delay loading the editor module to PostDefault to allow other modules to register their nodes, tools, and view modes. Warn if trying to register tools and view modes after editor commands have already been set up.

diff --git a/Engine/Plugins/Experimental/Dataflow/Dataflow.uplugin b/Engine/Plugins/Experimental/Dataflow/Dataflow.uplugin
--- a/Engine/Plugins/Experimental/Dataflow/Dataflow.uplugin
+++ b/Engine/Plugins/Experimental/Dataflow/Dataflow.uplugin
@@ -19,7 +19,7 @@
 		{
 			"Name": "DataflowEditor",
 			"Type": "Editor",
-			"LoadingPhase": "Default"
+			"LoadingPhase": "PostDefault"
 		},
 		{
 			"Name": "DataflowAssetTools",
diff --git a/Engine/Plugins/Experimental/Dataflow/Source/DataflowEditor/Private/Dataflow/DataflowEditorCommands.cpp b/Engine/Plugins/Experimental/Dataflow/Source/DataflowEditor/Private/Dataflow/DataflowEditorCommands.cpp
--- a/Engine/Plugins/Experimental/Dataflow/Source/DataflowEditor/Private/Dataflow/DataflowEditorCommands.cpp
+++ b/Engine/Plugins/Experimental/Dataflow/Source/DataflowEditor/Private/Dataflow/DataflowEditorCommands.cpp
@@ -139,6 +139,12 @@ void FDataflowEditorCommands::Unregister()
 	return FDataflowEditorCommandsImpl::Unregister();
 }
 
+bool FDataflowEditorCommands::IsRegistered()
+{
+	return FDataflowEditorCommandsImpl::IsRegistered();
+}
+
+
 const FDataflowNode* FDataflowEditorCommands::EvaluateNode(Dataflow::FContext& Context, Dataflow::FTimestamp& InOutLastNodeTimestamp,
 	const UDataflow* Dataflow, const FDataflowNode* Node, const FDataflowOutput* Output, const FString& NodeName, UObject* Asset)
 {
diff --git a/Engine/Plugins/Experimental/Dataflow/Source/DataflowEditor/Private/Dataflow/DataflowGraphEditor.cpp b/Engine/Plugins/Experimental/Dataflow/Source/DataflowEditor/Private/Dataflow/DataflowGraphEditor.cpp
--- a/Engine/Plugins/Experimental/Dataflow/Source/DataflowEditor/Private/Dataflow/DataflowGraphEditor.cpp
+++ b/Engine/Plugins/Experimental/Dataflow/Source/DataflowEditor/Private/Dataflow/DataflowGraphEditor.cpp
@@ -37,7 +37,6 @@ void SDataflowGraphEditor::Construct(const FArguments& InArgs, UObject* InAssetO
 	AppearanceInfo.CornerText = FText::FromString("Dataflow");
 
 	FGraphEditorCommands::Register();
-	FDataflowEditorCommands::Register();
 	if (!GraphEditorCommands.IsValid())
 	{
 		GraphEditorCommands = MakeShareable(new FUICommandList);
diff --git a/Engine/Plugins/Experimental/Dataflow/Source/DataflowEditor/Private/Dataflow/DataflowRenderingViewMode.cpp b/Engine/Plugins/Experimental/Dataflow/Source/DataflowEditor/Private/Dataflow/DataflowRenderingViewMode.cpp
--- a/Engine/Plugins/Experimental/Dataflow/Source/DataflowEditor/Private/Dataflow/DataflowRenderingViewMode.cpp
+++ b/Engine/Plugins/Experimental/Dataflow/Source/DataflowEditor/Private/Dataflow/DataflowRenderingViewMode.cpp
@@ -2,6 +2,7 @@
 
 #include "Dataflow/DataflowRenderingViewMode.h"
 #include "Misc/LazySingleton.h"
+#include "Dataflow/DataflowEditorCommands.h"
 
 #define LOCTEXT_NAMESPACE "DataflowRenderingViewMode"
 
@@ -81,6 +82,10 @@ namespace Dataflow
 
 	void FRenderingViewModeFactory::RegisterViewMode(TUniquePtr<IDataflowConstructionViewMode>&& ViewMode)
 	{
+		ensureMsgf(!FDataflowEditorCommands::IsRegistered(), TEXT("FRenderingViewModeFactory: DataflowEditorCommands have already been registered. \
+			Newly registered View Modes may not be available in the Editor. \
+			Ensure that RegisterViewMode is called before the DataflowEditor module is loaded."));
+
 		const FName NewViewModeName = ViewMode->GetName();
 		if (ViewModeMap.Contains(NewViewModeName))
 		{
diff --git a/Engine/Plugins/Experimental/Dataflow/Source/DataflowEditor/Private/Dataflow/DataflowToolRegistry.cpp b/Engine/Plugins/Experimental/Dataflow/Source/DataflowEditor/Private/Dataflow/DataflowToolRegistry.cpp
--- a/Engine/Plugins/Experimental/Dataflow/Source/DataflowEditor/Private/Dataflow/DataflowToolRegistry.cpp
+++ b/Engine/Plugins/Experimental/Dataflow/Source/DataflowEditor/Private/Dataflow/DataflowToolRegistry.cpp
@@ -3,6 +3,7 @@
 #include "Dataflow/DataflowToolRegistry.h"
 #include "Misc/LazySingleton.h"
 #include "Framework/Commands/UICommandInfo.h"
+#include "Dataflow/DataflowEditorCommands.h"
 
 namespace Dataflow
 {
@@ -19,6 +20,10 @@ namespace Dataflow
 
 	void FDataflowToolRegistry::AddNodeToToolMapping(const FName& NodeName, TObjectPtr<UInteractiveToolBuilder> ToolBuilder, const TSharedRef<const IDataflowToolActionCommands>& ToolActionRegistry)
 	{
+		ensureMsgf(!FDataflowEditorCommands::IsRegistered(), TEXT("FDataflowToolRegistry: DataflowEditorCommands have already been registered. \
+				Newly registered Tools may not be available in the Editor. \
+				Ensure that AddNodeToToolMapping is called before the DataflowEditor module is loaded."));
+
 		// FUICommandInfo is uninitialized, it will be created later in FDataflowEditorCommandsImpl::RegisterCommands
 		NodeTypeToToolMap.Add(NodeName, { ToolBuilder, ToolActionRegistry, nullptr });
 	}
diff --git a/Engine/Plugins/Experimental/Dataflow/Source/DataflowEditor/Public/Dataflow/DataflowEditorCommands.h b/Engine/Plugins/Experimental/Dataflow/Source/DataflowEditor/Public/Dataflow/DataflowEditorCommands.h
--- a/Engine/Plugins/Experimental/Dataflow/Source/DataflowEditor/Public/Dataflow/DataflowEditorCommands.h
+++ b/Engine/Plugins/Experimental/Dataflow/Source/DataflowEditor/Public/Dataflow/DataflowEditorCommands.h
@@ -103,6 +103,7 @@ public:
 
 	static void Register();
 	static void Unregister();
+	static bool IsRegistered();
 
 	static const FDataflowEditorCommandsImpl& Get();
 
