Dataflow : remove node actions from the pin context menu

diff --git a/Engine/Plugins/Experimental/Dataflow/Source/DataflowEditor/Private/Dataflow/DataflowSchema.cpp b/Engine/Plugins/Experimental/Dataflow/Source/DataflowEditor/Private/Dataflow/DataflowSchema.cpp
--- a/Engine/Plugins/Experimental/Dataflow/Source/DataflowEditor/Private/Dataflow/DataflowSchema.cpp
+++ b/Engine/Plugins/Experimental/Dataflow/Source/DataflowEditor/Private/Dataflow/DataflowSchema.cpp
@@ -28,7 +28,7 @@ UDataflowSchema::UDataflowSchema()
 
 void UDataflowSchema::GetContextMenuActions(class UToolMenu* Menu, class UGraphNodeContextMenuContext* Context) const
 {
-	if (Context->Node)
+	if (Context->Node && !Context->Pin)
 	{
 		{
 			FToolMenuSection& Section = Menu->AddSection("TestGraphSchemaNodeActions", LOCTEXT("GraphSchemaNodeActions_MenuHeader", "Node Actions"));
