Dataflow editor : fix crash when clicking on nodes that have been disabled because they are not loaded

diff --git a/Engine/Plugins/Experimental/Dataflow/Source/DataflowEditor/Private/Dataflow/DataflowEditorToolkit.cpp b/Engine/Plugins/Experimental/Dataflow/Source/DataflowEditor/Private/Dataflow/DataflowEditorToolkit.cpp
--- a/Engine/Plugins/Experimental/Dataflow/Source/DataflowEditor/Private/Dataflow/DataflowEditorToolkit.cpp
+++ b/Engine/Plugins/Experimental/Dataflow/Source/DataflowEditor/Private/Dataflow/DataflowEditorToolkit.cpp
@@ -942,18 +942,20 @@ void FDataflowEditorToolkit::OnNodeSelectionChanged(const TSet<UObject*>& InNewS
 				TSharedPtr<Dataflow::FEngineContext> DataflowContext = EditorContent->GetDataflowContext();
 				if (PrimarySelection && DataflowContext.IsValid())
 				{
-					PrimarySelection->GetDataflowNode()->OnSelected(*DataflowContext);
-
-					// Update selected Collection in the ContextObject
-					const TSharedPtr<FDataflowNode> DataflowNode = PrimarySelection->GetDataflowNode();
-					for (const FDataflowOutput* const Output : DataflowNode->GetOutputs())
+					if (const TSharedPtr<FDataflowNode> DataflowNode = PrimarySelection->GetDataflowNode())
 					{
-						if (Output->GetType() == FName(TEXT("FManagedArrayCollection")))
+						DataflowNode->OnSelected(*DataflowContext);
+
+						// Update selected Collection in the ContextObject
+						for (const FDataflowOutput* const Output : DataflowNode->GetOutputs())
 						{
-							const FManagedArrayCollection DefaultValue;
-							TSharedRef<FManagedArrayCollection> Collection = MakeShared<FManagedArrayCollection>(Output->GetValue<FManagedArrayCollection>(*DataflowContext, DefaultValue));
-							constexpr bool bCollectionIsInput = false;
-							EditorContent->SetSelectedCollection(Collection, bCollectionIsInput);
+							if (Output->GetType() == FName(TEXT("FManagedArrayCollection")))
+							{
+								const FManagedArrayCollection DefaultValue;
+								TSharedRef<FManagedArrayCollection> Collection = MakeShared<FManagedArrayCollection>(Output->GetValue<FManagedArrayCollection>(*DataflowContext, DefaultValue));
+								constexpr bool bCollectionIsInput = false;
+								EditorContent->SetSelectedCollection(Collection, bCollectionIsInput);
+							}
 						}
 					}
 				}
