Dataflow Editor: restore "Add Weight Map" button functionality in the UI

diff --git a/Engine/Plugins/Experimental/Dataflow/Source/DataflowEditor/Private/Dataflow/DataflowEditorMode.cpp b/Engine/Plugins/Experimental/Dataflow/Source/DataflowEditor/Private/Dataflow/DataflowEditorMode.cpp
--- a/Engine/Plugins/Experimental/Dataflow/Source/DataflowEditor/Private/Dataflow/DataflowEditorMode.cpp
+++ b/Engine/Plugins/Experimental/Dataflow/Source/DataflowEditor/Private/Dataflow/DataflowEditorMode.cpp
@@ -177,7 +177,7 @@ void UDataflowEditorMode::RegisterDataflowTool(TSharedPtr<FUICommandInfo> UIComm
 	);
 }
 
-void UDataflowEditorMode::RegisterAddNodeCommand(TSharedPtr<FUICommandInfo> AddNodeCommand, const FName& NewNodeType, TSharedPtr<FUICommandInfo> StartToolCommand)
+void UDataflowEditorMode::RegisterAddNodeCommand(TSharedPtr<FUICommandInfo> AddNodeCommand, const FName& NewNodeType)
 {
 	auto AddNode = [this](const FName& NewNodeType)
 	{
@@ -205,8 +205,6 @@ void UDataflowEditorMode::RegisterAddNodeCommand(TSharedPtr<FUICommandInfo> AddN
 		FExecuteAction::CreateWeakLambda(this, AddNode, NewNodeType),
 		FCanExecuteAction::CreateWeakLambda(this, CanAddNode, NewNodeType)
 	);
-
-	NodeTypeToToolCommandMap.Add(NewNodeType, StartToolCommand);
 }
 
 void UDataflowEditorMode::RegisterTools()
@@ -233,6 +231,8 @@ void UDataflowEditorMode::RegisterTools()
 		NodeTypeToToolCommandMap.Add(RegisteredNodeName, CommandInfo);
 	}
 
+	// Register "Add Node" commands for buttons in the UI. So far we just have the one
+	RegisterAddNodeCommand(CommandInfos.AddWeightMapNode, FDataflowCollectionAddScalarVertexPropertyNode::StaticType());
 }
 
 bool UDataflowEditorMode::ShouldToolStartBeAllowed(const FString& ToolIdentifier) const
diff --git a/Engine/Plugins/Experimental/Dataflow/Source/DataflowEditor/Public/Dataflow/DataflowEditorMode.h b/Engine/Plugins/Experimental/Dataflow/Source/DataflowEditor/Public/Dataflow/DataflowEditorMode.h
--- a/Engine/Plugins/Experimental/Dataflow/Source/DataflowEditor/Public/Dataflow/DataflowEditorMode.h
+++ b/Engine/Plugins/Experimental/Dataflow/Source/DataflowEditor/Public/Dataflow/DataflowEditorMode.h
@@ -117,7 +117,8 @@ private:
 							  UEditorInteractiveToolsContext* const ToolsContext,
 							  EToolsContextScope ToolScope = EToolsContextScope::Default);
 
-	void RegisterAddNodeCommand(TSharedPtr<FUICommandInfo> AddNodeCommand, const FName& NewNodeType, TSharedPtr<FUICommandInfo> StartToolCommand);
+	// Command for quickly adding a new node to the graph, e.g. triggered by a button press
+	void RegisterAddNodeCommand(TSharedPtr<FUICommandInfo> AddNodeCommand, const FName& NewNodeType);
 
 	void SetSimulationViewportClient(TWeakPtr<FDataflowSimulationViewportClient, ESPMode::ThreadSafe>);
 	void RefocusSimulationViewportClient();
