Dataflow: remove Hide unused pins menu option since it's effectively the same as Hide unconnected pins for us.

Also update text to "Show" rather than "Unhide" option pins.

diff --git a/Engine/Plugins/Experimental/Dataflow/Source/DataflowEditor/Private/Dataflow/DataflowSNode.cpp b/Engine/Plugins/Experimental/Dataflow/Source/DataflowEditor/Private/Dataflow/DataflowSNode.cpp
--- a/Engine/Plugins/Experimental/Dataflow/Source/DataflowEditor/Private/Dataflow/DataflowSNode.cpp
+++ b/Engine/Plugins/Experimental/Dataflow/Source/DataflowEditor/Private/Dataflow/DataflowSNode.cpp
@@ -229,7 +229,7 @@ FReply SDataflowEdNode::OnAddPin()
 				MenuBuilder.AddMenuEntry(LOCTEXT("HideAllInputs", "Hide all"), LOCTEXT("HideAllInputsTooltip", "Hide all hideable input pins"), FSlateIcon(),
 					FUIAction(
 						FExecuteAction::CreateUObject(DataflowGraphNode, &UDataflowEdNode::HideAllInputPins)));
-				MenuBuilder.AddMenuEntry(LOCTEXT("UnhideAllInputs", "Unhide all"), LOCTEXT("UnhideAllInputsTooltip", "Unhide all hideable input pins"), FSlateIcon(),
+				MenuBuilder.AddMenuEntry(LOCTEXT("UnhideAllInputs", "Show all"), LOCTEXT("UnhideAllInputsTooltip", "Show all hideable input pins"), FSlateIcon(),
 					FUIAction(
 						FExecuteAction::CreateUObject(DataflowGraphNode, &UDataflowEdNode::ShowAllInputPins)));
 
diff --git a/Engine/Plugins/Experimental/Dataflow/Source/DataflowEditor/Private/Dataflow/DataflowSchema.cpp b/Engine/Plugins/Experimental/Dataflow/Source/DataflowEditor/Private/Dataflow/DataflowSchema.cpp
--- a/Engine/Plugins/Experimental/Dataflow/Source/DataflowEditor/Private/Dataflow/DataflowSchema.cpp
+++ b/Engine/Plugins/Experimental/Dataflow/Source/DataflowEditor/Private/Dataflow/DataflowSchema.cpp
@@ -79,7 +79,6 @@ void UDataflowSchema::GetContextMenuActions(class UToolMenu* Menu, class UGraphN
 					FToolMenuSection& InSection = PinVisibilityMenu->AddSection("TestGraphSchemaPinVisibility");
 					InSection.AddMenuEntry(FGraphEditorCommands::Get().ShowAllPins);
 					InSection.AddMenuEntry(FGraphEditorCommands::Get().HideNoConnectionPins);
-					InSection.AddMenuEntry(FGraphEditorCommands::Get().HideNoConnectionNoDefaultPins);
 				}));
 			}
 
