Dataflow:
- Added DataflowAnyTypePinTypeColor to DataflowSettings
- Also added some other types
- Clean up

diff --git a/Engine/Plugins/Experimental/Dataflow/Source/DataflowEditor/Private/Dataflow/DataflowSchema.cpp b/Engine/Plugins/Experimental/Dataflow/Source/DataflowEditor/Private/Dataflow/DataflowSchema.cpp
--- a/Engine/Plugins/Experimental/Dataflow/Source/DataflowEditor/Private/Dataflow/DataflowSchema.cpp
+++ b/Engine/Plugins/Experimental/Dataflow/Source/DataflowEditor/Private/Dataflow/DataflowSchema.cpp
@@ -22,6 +22,25 @@
 
 #define LOCTEXT_NAMESPACE "DataflowNode"
 
+namespace Dataflow::Private
+{
+	static const FName ManagedArrayCollectionType = FName("FManagedArrayCollection");
+	static const FName FloatType = FName("float");
+	static const FName DoubleType = FName("double");
+	static const FName Int32Type = FName("int32");
+	static const FName BoolType = FName("bool");
+	static const FName StringType = FName("FString");
+	static const FName NameType = FName("FName");
+	static const FName TextType = FName("FText");
+	static const FName VectorType = FName("FVector");
+	static const FName TransformType = FName("FTransform");
+	static const FName RotatorType = FName("FRotator");
+	static const FName ArrayType = FName("TArray");
+	static const FName BoxType = FName("FBox");
+	static const FName SphereType = FName("FSphere");
+	static const FName DataflowAnyTypeType = FName("FDataflowAnyType");
+} // namespace Dataflow::Private
+
 UDataflowSchema::UDataflowSchema()
 {
 }
@@ -224,42 +243,66 @@ FLinearColor UDataflowSchema::GetTypeColor(const FName& Type)
 	const UGraphEditorSettings* Settings = GetDefault<UGraphEditorSettings>();
 	const UDataflowSettings* DataflowSettings = GetDefault<UDataflowSettings>();
 
-	if (Type == FName("FManagedArrayCollection"))
+	if (Type == Dataflow::Private::ManagedArrayCollectionType)
 	{
 		return DataflowSettings->ManagedArrayCollectionPinTypeColor;
 	}
-	else if (Type == FName("float"))
+	else if (Type == Dataflow::Private::FloatType)
 	{
 		return Settings->FloatPinTypeColor;
 	}
-	else if (Type == FName("int32"))
+	else if (Type == Dataflow::Private::DoubleType)
+	{
+		return Settings->DoublePinTypeColor;
+	}
+	else if (Type == Dataflow::Private::Int32Type)
 	{
 		return Settings->IntPinTypeColor;
 	}
-	else if (Type == FName("bool"))
+	else if (Type == Dataflow::Private::BoolType)
 	{
 		return Settings->BooleanPinTypeColor;
 	}
-	else if (Type == FName("FString"))
+	else if (Type == Dataflow::Private::StringType)
 	{
 		return Settings->StringPinTypeColor;
 	}
-	else if (Type == FName("FVector"))
+	else if (Type == Dataflow::Private::NameType)
+	{
+		return Settings->NamePinTypeColor;
+	}
+	else if (Type == Dataflow::Private::TextType)
+	{
+		return Settings->TextPinTypeColor;
+	}
+	else if (Type == Dataflow::Private::VectorType)
 	{
 		return Settings->VectorPinTypeColor;
 	}
-	else if (Type == FName("TArray"))
+	else if (Type == Dataflow::Private::TransformType)
+	{
+		return Settings->TransformPinTypeColor;
+	}
+	else if (Type == Dataflow::Private::RotatorType)
+	{
+		return Settings->RotatorPinTypeColor;
+	}
+	else if (Type == Dataflow::Private::ArrayType)
 	{
 		return DataflowSettings->ArrayPinTypeColor;
 	}
-	else if (Type == FName("FBox"))
+	else if (Type == Dataflow::Private::BoxType)
 	{
 		return DataflowSettings->BoxPinTypeColor;
 	}
-	else if (Type == FName("FSphere"))
+	else if (Type == Dataflow::Private::SphereType)
 	{
 		return DataflowSettings->SpherePinTypeColor;
 	}
+	else if (Type == Dataflow::Private::DataflowAnyTypeType)
+	{
+		return DataflowSettings->DataflowAnyTypePinTypeColor;
+	}
 
 	return Settings->DefaultPinTypeColor;
 }
diff --git a/Engine/Source/Runtime/Experimental/Dataflow/Core/Private/Dataflow/DataflowSettings.cpp b/Engine/Source/Runtime/Experimental/Dataflow/Core/Private/Dataflow/DataflowSettings.cpp
--- a/Engine/Source/Runtime/Experimental/Dataflow/Core/Private/Dataflow/DataflowSettings.cpp
+++ b/Engine/Source/Runtime/Experimental/Dataflow/Core/Private/Dataflow/DataflowSettings.cpp
@@ -4,13 +4,23 @@
 
 #include UE_INLINE_GENERATED_CPP_BY_NAME(DataflowSettings)
 
+namespace Dataflow::Private
+{
+	static const FLinearColor CManagedArrayCollectionPinTypeColor = FLinearColor(0.353393f, 0.454175f, 1.0f, 1.0f);
+	static const FLinearColor CArrayPinTypeColor = FLinearColor(1.0f, 0.172585f, 0.0f, 1.0f);
+	static const FLinearColor CBoxPinTypeColor = FLinearColor(0.013575f, 0.770000f, 0.429609f, 1.0f);
+	static const FLinearColor CSpherePinTypeColor = FLinearColor(0.2f, 0.6f, 1.f, 1.0f);
+	static const FLinearColor CDataflowAnyTypePinTypeColor = FLinearColor(0.3f, 0.3f, 0.3f, 1.0f);
+} // namespace Dataflow::Private
+
 UDataflowSettings::UDataflowSettings(const FObjectInitializer& ObjectInitlaizer)
 	: Super(ObjectInitlaizer)
 {
-	ManagedArrayCollectionPinTypeColor = CManagedArrayCollectionPinTypeColor;
-	ArrayPinTypeColor = CArrayPinTypeColor;
-	BoxPinTypeColor = CBoxPinTypeColor;
-	SpherePinTypeColor = CSpherePinTypeColor;
+	ManagedArrayCollectionPinTypeColor = Dataflow::Private::CManagedArrayCollectionPinTypeColor;
+	ArrayPinTypeColor = Dataflow::Private::CArrayPinTypeColor;
+	BoxPinTypeColor = Dataflow::Private::CBoxPinTypeColor;
+	SpherePinTypeColor = Dataflow::Private::CSpherePinTypeColor;
+	DataflowAnyTypePinTypeColor = Dataflow::Private::CDataflowAnyTypePinTypeColor;
 }
 
 FName UDataflowSettings::GetCategoryName() const
diff --git a/Engine/Source/Runtime/Experimental/Dataflow/Core/Public/Dataflow/DataflowSettings.h b/Engine/Source/Runtime/Experimental/Dataflow/Core/Public/Dataflow/DataflowSettings.h
--- a/Engine/Source/Runtime/Experimental/Dataflow/Core/Public/Dataflow/DataflowSettings.h
+++ b/Engine/Source/Runtime/Experimental/Dataflow/Core/Public/Dataflow/DataflowSettings.h
@@ -7,11 +7,6 @@
 
 #include "DataflowSettings.generated.h"
 
-static const FLinearColor CManagedArrayCollectionPinTypeColor = FLinearColor(0.353393f, 0.454175f, 1.0f, 1.0f);
-static const FLinearColor CArrayPinTypeColor = FLinearColor(1.0f, 0.172585f, 0.0f, 1.0f);
-static const FLinearColor CBoxPinTypeColor = FLinearColor(0.013575f, 0.770000f, 0.429609f, 1.0f);
-static const FLinearColor CSpherePinTypeColor = FLinearColor(0.2f, 0.6f, 1.f, 1.0f);
-
 typedef TMap<FName, FNodeColors> FNodeColorsMap;
 
 USTRUCT()
@@ -47,6 +42,10 @@ class UDataflowSettings : public UDeveloperSettings
 	UPROPERTY(config, EditAnywhere, Category = PinColors)
 	FLinearColor SpherePinTypeColor;
 
+	/** FDataflowAnyType pin type color. The other pin colors are defined in the general editor settings. */
+	UPROPERTY(config, EditAnywhere, Category = PinColors)
+	FLinearColor DataflowAnyTypePinTypeColor;
+
 	UPROPERTY(config, EditAnywhere, Category = NodeColors)
 	TMap<FName, FNodeColors> NodeColorsMap;
 
