Fix default BP preview saving

diff --git a/Engine/Plugins/Experimental/Dataflow/Source/DataflowEditor/Private/Dataflow/DataflowSimulationScene.cpp b/Engine/Plugins/Experimental/Dataflow/Source/DataflowEditor/Private/Dataflow/DataflowSimulationScene.cpp
--- a/Engine/Plugins/Experimental/Dataflow/Source/DataflowEditor/Private/Dataflow/DataflowSimulationScene.cpp
+++ b/Engine/Plugins/Experimental/Dataflow/Source/DataflowEditor/Private/Dataflow/DataflowSimulationScene.cpp
@@ -42,14 +42,15 @@ FDataflowSimulationScene::FDataflowSimulationScene(FPreviewScene::ConstructionVa
 	if(GetEditorContent())
 	{
 #if WITH_EDITORONLY_DATA
-		if(SceneDescription->BlueprintClass == nullptr)
-		{
-			SceneDescription->BlueprintClass = GetEditorContent()->GetPreviewClass();
-		}
 		if(const UDataflow* DataflowAsset = GetEditorContent()->GetDataflowAsset())
 		{
 			SceneDescription->CacheParams = DataflowAsset->PreviewCacheParams;
 			SceneDescription->CacheAsset = Cast<UChaosCacheCollection>(DataflowAsset->PreviewCacheAsset.Get());
+			SceneDescription->BlueprintClass = DataflowAsset->PreviewBlueprintClass; 
+		}
+		if(SceneDescription->BlueprintClass == nullptr)
+		{
+			SceneDescription->BlueprintClass = GetEditorContent()->GetPreviewClass();
 		}
 #endif
 	}
