Fix dataflow animation lag + Enable visualisation + BP compilation crash + Floor preview fix + BP Transform

diff --git a/Engine/Plugins/Experimental/Dataflow/Source/DataflowEditor/Private/Dataflow/DataflowEditorPreviewSceneBase.cpp b/Engine/Plugins/Experimental/Dataflow/Source/DataflowEditor/Private/Dataflow/DataflowEditorPreviewSceneBase.cpp
--- a/Engine/Plugins/Experimental/Dataflow/Source/DataflowEditor/Private/Dataflow/DataflowEditorPreviewSceneBase.cpp
+++ b/Engine/Plugins/Experimental/Dataflow/Source/DataflowEditor/Private/Dataflow/DataflowEditorPreviewSceneBase.cpp
@@ -10,9 +10,11 @@
 #define LOCTEXT_NAMESPACE "FDataflowPreviewSceneBase"
 
 
-bool bDataflowShowFloorDefault = false;
+bool bDataflowShowFloorDefault = true;
 FAutoConsoleVariableRef CVARDataflowShowFloorDefault(TEXT("p.Dataflow.Editor.ShowFloor"), bDataflowShowFloorDefault, TEXT("Show the floor in the dataflow editor[def:false]"));
 
+bool bDataflowShowEnvironmentDefault = true;
+FAutoConsoleVariableRef CVARDataflowShowEnvironmentDefault(TEXT("p.Dataflow.Editor.ShowEnvironment"), bDataflowShowEnvironmentDefault, TEXT("Show the environment in the dataflow editor[def:false]"));
 
 FDataflowPreviewSceneBase::FDataflowPreviewSceneBase(FPreviewScene::ConstructionValues ConstructionValues, UDataflowEditor* InEditor)
 	: FAdvancedPreviewScene(ConstructionValues)
@@ -21,7 +23,9 @@ FDataflowPreviewSceneBase::FDataflowPreviewSceneBase(FPreviewScene::Construction
 	RootSceneActor = GetWorld()->SpawnActor<AActor>(AActor::StaticClass());
 	
 	check(DataflowEditor);
-	SetFloorVisibility(bDataflowShowFloorDefault, true);
+	
+	SetFloorVisibility(bDataflowShowFloorDefault, false);
+	SetEnvironmentVisibility(bDataflowShowEnvironmentDefault, false);
 }
 
 FDataflowPreviewSceneBase::~FDataflowPreviewSceneBase()
diff --git a/Engine/Plugins/Experimental/Dataflow/Source/DataflowEditor/Private/Dataflow/DataflowEditorToolkit.cpp b/Engine/Plugins/Experimental/Dataflow/Source/DataflowEditor/Private/Dataflow/DataflowEditorToolkit.cpp
--- a/Engine/Plugins/Experimental/Dataflow/Source/DataflowEditor/Private/Dataflow/DataflowEditorToolkit.cpp
+++ b/Engine/Plugins/Experimental/Dataflow/Source/DataflowEditor/Private/Dataflow/DataflowEditorToolkit.cpp
@@ -53,6 +53,7 @@
 #include "Selection.h"
 #include "Styling/SlateStyleRegistry.h"
 #include "Widgets/Docking/SDockTab.h"
+#include "GeometryCache.h"
 
 
 #define LOCTEXT_NAMESPACE "DataflowEditorToolkit"
@@ -539,6 +540,18 @@ void FDataflowEditorToolkit::GetSaveableObjects(TArray<UObject*>& OutObjects) co
 			check(DataflowAsset->IsAsset());
 			OutObjects.AddUnique(DataflowAsset);
 		}
+
+		if(SimulationScene && SimulationScene->GetPreviewSceneDescription())
+		{
+			if(TObjectPtr<UChaosCacheCollection> CacheCollection = SimulationScene->GetPreviewSceneDescription()->CacheAsset)
+			{
+				OutObjects.AddUnique(CacheCollection);
+			}
+			if(TObjectPtr<UGeometryCache> GeometryCache = SimulationScene->GetPreviewSceneDescription()->GeometryCacheAsset)
+			{
+				OutObjects.AddUnique(GeometryCache);
+			}
+		}
 	}
 }
 
diff --git a/Engine/Plugins/Experimental/Dataflow/Source/DataflowEditor/Private/Dataflow/DataflowSimulationControls.cpp b/Engine/Plugins/Experimental/Dataflow/Source/DataflowEditor/Private/Dataflow/DataflowSimulationControls.cpp
--- a/Engine/Plugins/Experimental/Dataflow/Source/DataflowEditor/Private/Dataflow/DataflowSimulationControls.cpp
+++ b/Engine/Plugins/Experimental/Dataflow/Source/DataflowEditor/Private/Dataflow/DataflowSimulationControls.cpp
@@ -39,7 +39,7 @@ namespace UE::Dataflow
 	
 	TObjectPtr<AActor> SpawnSimulatedActor(const TSubclassOf<AActor>& ActorClass,
 		const TObjectPtr<AChaosCacheManager>& CacheManager, const TObjectPtr<UChaosCacheCollection>& CacheCollection,
-		const bool bIsRecording, const TObjectPtr<UDataflowBaseContent>& DataflowContent)
+		const bool bIsRecording, const TObjectPtr<UDataflowBaseContent>& DataflowContent, const FTransform& ActorTransform)
 	{
 		if(CacheManager)
 		{
@@ -56,7 +56,7 @@ namespace UE::Dataflow
 				DataflowContent->SetActorProperties(PreviewActor);
 
 				// Finish spawning
-				PreviewActor->FinishSpawning(FTransform::Identity, true);
+				PreviewActor->FinishSpawning(ActorTransform, true);
 			}
 
 			CacheManager->CacheCollection = CacheCollection;
@@ -198,6 +198,12 @@ namespace UE::Dataflow
 					SkeletalMeshComponent->SetPosition(SimulationTime);
 					SkeletalMeshComponent->TickAnimation(0.f, false /*bNeedsValidRootMotion*/);
 					SkeletalMeshComponent->RefreshBoneTransforms(nullptr /*TickFunction*/);
+
+					SkeletalMeshComponent->RefreshFollowerComponents();
+					SkeletalMeshComponent->UpdateComponentToWorld();
+					SkeletalMeshComponent->FinalizeBoneTransform();
+					SkeletalMeshComponent->MarkRenderTransformDirty();
+					SkeletalMeshComponent->MarkRenderDynamicDataDirty();
 				}
 			}
 		}
diff --git a/Engine/Plugins/Experimental/Dataflow/Source/DataflowEditor/Private/Dataflow/DataflowSimulationGenerator.cpp b/Engine/Plugins/Experimental/Dataflow/Source/DataflowEditor/Private/Dataflow/DataflowSimulationGenerator.cpp
--- a/Engine/Plugins/Experimental/Dataflow/Source/DataflowEditor/Private/Dataflow/DataflowSimulationGenerator.cpp
+++ b/Engine/Plugins/Experimental/Dataflow/Source/DataflowEditor/Private/Dataflow/DataflowSimulationGenerator.cpp
@@ -5,6 +5,7 @@
 #include "Dataflow/DataflowSimulationControls.h"
 #include "Dataflow/DataflowSimulationUtils.h"
 #include "Chaos/CacheManagerActor.h"
+#include "Chaos/CacheCollection.h"
 #include "Engine/World.h"
 #include "Engine/Engine.h"
 #include "Misc/AsyncTaskNotification.h" 
@@ -17,7 +18,7 @@ namespace UE::Dataflow
 {
 void FDataflowSimulationTask::DoWork()
 {
-	const int32 StartFrame = bBackgroundTask ? 1 : 0;
+	const int32 StartFrame = bAsyncCaching ? 1 : 0;
 	const int32 NumFrames = (MaxTime-MinTime) / DeltaTime;
 	for (int32 FrameIndex = StartFrame; FrameIndex < NumFrames; ++FrameIndex)
 	{
@@ -26,7 +27,7 @@ void FDataflowSimulationTask::DoWork()
 			// Compute the simulation time that will be sent to the graph
 			const float SimulationTime = MinTime + (FrameIndex+1) * DeltaTime;
 
-			if(bBackgroundTask)
+			if(bAsyncCaching)
 			{
 				// Compute all the skelmesh animations at the simulation time
 				Dataflow::ComputeSkeletonAnimation(TaskManager->PreviewActor, SimulationTime);
@@ -63,7 +64,8 @@ void FDataflowSimulationTask::DoWork()
 }
 
 bool FDataflowTaskManager::AllocateSimulationResource(const FVector2f& TimeRange, const int32 FrameRate,
-		const TObjectPtr<UChaosCacheCollection>& CacheAsset, const TSubclassOf<AActor>& ActorClass, const TObjectPtr<UDataflowBaseContent>& DataflowContent)
+		const TObjectPtr<UChaosCacheCollection>& CacheAsset, const TSubclassOf<AActor>& ActorClass,
+		const TObjectPtr<UDataflowBaseContent>& DataflowContent, const FTransform& BlueprintTransform)
 {
 	SimulationWorld = UWorld::CreateWorld(EWorldType::Editor, false);
 	SimulationWorld->bPostTickComponentUpdate = false;
@@ -73,7 +75,7 @@ bool FDataflowTaskManager::AllocateSimulationResource(const FVector2f& TimeRange
 	
 	CacheManager = SimulationWorld->SpawnActor<AChaosCacheManager>();
 
-	PreviewActor = Dataflow::SpawnSimulatedActor(ActorClass, CacheManager, CacheAsset, true, DataflowContent);
+	PreviewActor = Dataflow::SpawnSimulatedActor(ActorClass, CacheManager, CacheAsset, true, DataflowContent, BlueprintTransform);
 	Dataflow::SetupSkeletonAnimation(PreviewActor);
 
 	// Init the cache manager
@@ -93,7 +95,7 @@ bool FDataflowTaskManager::AllocateSimulationResource(const FVector2f& TimeRange
 		SimulationTask->GetTask().DeltaTime = (TimeRange[1]-TimeRange[0]) / NumFrames;
 
 		UDataflowSimulationManager* DataflowManager = SimulationWorld->GetSubsystem<UDataflowSimulationManager>();
-		if(SimulationTask->GetTask().bBackgroundTask)
+		if(SimulationTask->GetTask().bAsyncCaching)
 		{
 			// Update all the skelmesh animations at the simulation time
 			Dataflow::UpdateSkeletonAnimation(PreviewActor, SimulationTask->GetTask().MinTime+SimulationTask->GetTask().DeltaTime);
@@ -104,7 +106,7 @@ bool FDataflowTaskManager::AllocateSimulationResource(const FVector2f& TimeRange
 			// Init simulation proxies from interface
 			DataflowManager->InitSimulationInterfaces();
 		}
-		DataflowManager->SetSimulationEnabled(!SimulationTask->GetTask().bBackgroundTask);
+		DataflowManager->SetSimulationEnabled(!SimulationTask->GetTask().bAsyncCaching);
 	}
 
 	return true;
@@ -116,7 +118,7 @@ void FDataflowTaskManager::FreeSimulationResource()
 	{
 		SimulationTask->EnsureCompletion();
 
-		if(SimulationTask->GetTask().bBackgroundTask)
+		if(SimulationTask->GetTask().bAsyncCaching)
 		{
 			UDataflowSimulationManager* DataflowManager = SimulationWorld->GetSubsystem<UDataflowSimulationManager>();
 			DataflowManager->ResetSimulationInterfaces();
@@ -186,11 +188,11 @@ void FDataflowSimulationGenerator::StartGenerateSimulation()
 
 	TaskManager->SimulationTask = MakeUnique<FAsyncTask<FDataflowSimulationTask>>();
 	TaskManager->SimulationTask->GetTask().TaskManager = TaskManager;
-	TaskManager->SimulationTask->GetTask().bBackgroundTask = CacheParams.bBackgroundTask;
+	TaskManager->SimulationTask->GetTask().bAsyncCaching = CacheParams.bAsyncCaching;
 	
-	TaskManager->AllocateSimulationResource(CacheParams.TimeRange, CacheParams.FrameRate, CacheAsset, BlueprintClass, DataflowContent);
+	TaskManager->AllocateSimulationResource(CacheParams.TimeRange, CacheParams.FrameRate, CacheAsset, BlueprintClass, DataflowContent, BlueprintTransform);
 
-	if(CacheParams.bBackgroundTask)
+	if(CacheParams.bAsyncCaching)
 	{
 		TaskManager->SimulationTask->StartBackgroundTask();
 	}
@@ -266,6 +268,11 @@ void FDataflowSimulationGenerator::SetBlueprintClass(const TSubclassOf<AActor>&
 	BlueprintClass = InBlueprintClass;
 }
 
+void FDataflowSimulationGenerator::SetBlueprintTransform(const FTransform& InBlueprintTransform)
+{
+	BlueprintTransform = InBlueprintTransform;
+}
+
 void FDataflowSimulationGenerator::SetDataflowContent(const TObjectPtr<UDataflowBaseContent>& InDataflowContent)
 {
 	DataflowContent = InDataflowContent;
@@ -301,6 +308,7 @@ void FDataflowSimulationGenerator::FreeTaskResource(bool bCancelled)
 		TaskManager->AsyncNotification->SetComplete(true);
 	}
 	TaskManager.Reset();
+	CacheAsset->MarkPackageDirty();
 }
 };
 
diff --git a/Engine/Plugins/Experimental/Dataflow/Source/DataflowEditor/Private/Dataflow/DataflowSimulationScene.cpp b/Engine/Plugins/Experimental/Dataflow/Source/DataflowEditor/Private/Dataflow/DataflowSimulationScene.cpp
--- a/Engine/Plugins/Experimental/Dataflow/Source/DataflowEditor/Private/Dataflow/DataflowSimulationScene.cpp
+++ b/Engine/Plugins/Experimental/Dataflow/Source/DataflowEditor/Private/Dataflow/DataflowSimulationScene.cpp
@@ -10,6 +10,7 @@
 #include "EngineUtils.h"
 #include "Animation/AnimSingleNodeInstance.h"
 #include "AssetEditorModeManager.h"
+#include "AssetViewerSettings.h"
 #include "Engine/Selection.h"
 
 #include "Dataflow/Interfaces/DataflowInterfaceGeometryCachable.h"
@@ -45,8 +46,9 @@ FDataflowSimulationScene::FDataflowSimulationScene(FPreviewScene::ConstructionVa
 		if(const UDataflow* DataflowAsset = GetEditorContent()->GetDataflowAsset())
 		{
 			SceneDescription->CacheParams = DataflowAsset->PreviewCacheParams;
-			SceneDescription->CacheAsset = Cast<UChaosCacheCollection>(DataflowAsset->PreviewCacheAsset.Get());
-			SceneDescription->BlueprintClass = DataflowAsset->PreviewBlueprintClass; 
+			SceneDescription->CacheAsset = Cast<UChaosCacheCollection>(DataflowAsset->PreviewCacheAsset.LoadSynchronous());
+			SceneDescription->BlueprintClass = DataflowAsset->PreviewBlueprintClass;
+			SceneDescription->BlueprintTransform = DataflowAsset->PreviewBlueprintTransform; 
 		}
 		if(SceneDescription->BlueprintClass == nullptr)
 		{
@@ -55,12 +57,31 @@ FDataflowSimulationScene::FDataflowSimulationScene(FPreviewScene::ConstructionVa
 #endif
 	}
 
+#if WITH_EDITOR
+	OnObjectsReinstancedHandle = FCoreUObjectDelegates::OnObjectsReinstanced.AddRaw(this, &FDataflowSimulationScene::OnObjectsReinstanced);
+#endif
+
 	CreateSimulationScene();
 }
 
+void FDataflowSimulationScene::OnObjectsReinstanced(const TMap<UObject*, UObject*>& ObjectsMap)
+{
+	if(UObject* const* InstancedActor = ObjectsMap.Find(PreviewActor))
+	{
+		if(*InstancedActor)
+		{
+			PreviewActor = Cast<AActor>(*InstancedActor);
+		}
+	}
+}
+
 FDataflowSimulationScene::~FDataflowSimulationScene()
 {
 	ResetSimulationScene();
+
+#if WITH_EDITOR
+	FCoreUObjectDelegates::OnObjectsReinstanced.Remove(OnObjectsReinstancedHandle);
+#endif
 }
 
 void FDataflowSimulationScene::UnbindSceneSelection()
@@ -91,9 +112,8 @@ void FDataflowSimulationScene::ResetSimulationScene()
 	// Destroy the spawned root actor
 	if(PreviewActor && GetWorld())
 	{
-		GetWorld()->DestroyActor(PreviewActor);
-
 		GetWorld()->EditorDestroyActor(PreviewActor, true);
+			
 		// Since deletion can be delayed, rename to avoid future name collision
 		// Call UObject::Rename directly on actor to avoid AActor::Rename which unnecessarily sunregister and re-register components
 		PreviewActor->UObject::Rename(nullptr, GetTransientPackage(), REN_DontCreateRedirectors | REN_ForceNoResetLoaders);
@@ -168,13 +188,14 @@ void FDataflowSimulationScene::CreateSimulationScene()
 		SimulationGenerator->SetCacheParams(SceneDescription->CacheParams);
 		SimulationGenerator->SetCacheAsset(SceneDescription->CacheAsset);
 		SimulationGenerator->SetBlueprintClass(SceneDescription->BlueprintClass);
+		SimulationGenerator->SetBlueprintTransform(SceneDescription->BlueprintTransform);
 		SimulationGenerator->SetDataflowContent(GetEditorContent());
 
 		TimeRange = SceneDescription->CacheParams.TimeRange;
 		NumFrames = (TimeRange[1] > TimeRange[0]) ? FMath::Floor((TimeRange[1] - TimeRange[0]) * SceneDescription->CacheParams.FrameRate) : 0;
 		
 		PreviewActor = UE::Dataflow::SpawnSimulatedActor(SceneDescription->BlueprintClass, Cast<AChaosCacheManager>(RootSceneActor),
-			SceneDescription->CacheAsset, false, GetEditorContent());
+			SceneDescription->CacheAsset, false, GetEditorContent(), SceneDescription->BlueprintTransform);
 
 		// Setup all the skelmesh animations
 		UE::Dataflow::SetupSkeletonAnimation(PreviewActor);
@@ -196,8 +217,6 @@ void FDataflowSimulationScene::UpdateSimulationCache()
 
 void FDataflowSimulationScene::TickDataflowScene(const float DeltaSeconds)
 {
-	GetWorld()->Tick(ELevelTick::LEVELTICK_All, DeltaSeconds);
-
 	if(const TObjectPtr<UDataflowBaseContent>& EditorContent = GetEditorContent())
 	{
 		if(UE::Dataflow::ShouldResetWorld(EditorContent->GetDataflowAsset(), GetWorld(), LastTimeStamp) || EditorContent->IsSimulationDirty())
@@ -223,6 +242,7 @@ void FDataflowSimulationScene::TickDataflowScene(const float DeltaSeconds)
 			UE::Dataflow::UpdateSkeletonAnimation(PreviewActor, SimulationTime);
 		}
 	}
+	GetWorld()->Tick(ELevelTick::LEVELTICK_All, DeltaSeconds);
 }
 
 void FDataflowSimulationScene::AddReferencedObjects(FReferenceCollector& Collector)
@@ -255,6 +275,13 @@ void FDataflowSimulationScene::SceneDescriptionPropertyChanged(const FName& Prop
 			SimulationGenerator->SetBlueprintClass(SceneDescription->BlueprintClass);
 		}
 	}
+	else if(PropertyName == GET_MEMBER_NAME_CHECKED(UDataflowSimulationSceneDescription, BlueprintTransform))
+	{
+		if(SimulationGenerator)
+		{
+			SimulationGenerator->SetBlueprintTransform(SceneDescription->BlueprintTransform);
+		}
+	}
 	if(GetEditorContent())
 	{
 		if(UDataflow* DataflowAsset = GetEditorContent()->GetDataflowAsset())
@@ -263,6 +290,9 @@ void FDataflowSimulationScene::SceneDescriptionPropertyChanged(const FName& Prop
 			DataflowAsset->PreviewCacheParams = SceneDescription->CacheParams;
 			DataflowAsset->PreviewCacheAsset = SceneDescription->CacheAsset;
 			DataflowAsset->PreviewBlueprintClass = SceneDescription->BlueprintClass;
+			DataflowAsset->PreviewBlueprintTransform = SceneDescription->BlueprintTransform;
+			
+			DataflowAsset->MarkPackageDirty();
 #endif
 		}
 	}
@@ -294,7 +324,6 @@ void UDataflowSimulationSceneDescription::GenerateGeometryCache()
 		USkeletalMeshComponent* SkeletalComponent = nullptr;
 		for (UPrimitiveComponent* PrimComponent : PrimComponents)
 		{
-			
 			if (!GeometryCachable)
 			{
 				GeometryCachable = Cast<IDataflowGeometryCachable>(PrimComponent);
diff --git a/Engine/Plugins/Experimental/Dataflow/Source/DataflowEditor/Public/Dataflow/DataflowSimulationControls.h b/Engine/Plugins/Experimental/Dataflow/Source/DataflowEditor/Public/Dataflow/DataflowSimulationControls.h
--- a/Engine/Plugins/Experimental/Dataflow/Source/DataflowEditor/Public/Dataflow/DataflowSimulationControls.h
+++ b/Engine/Plugins/Experimental/Dataflow/Source/DataflowEditor/Public/Dataflow/DataflowSimulationControls.h
@@ -19,7 +19,7 @@ namespace UE::Dataflow
 	/** Spawn an actor given a class type and attach it to the cache manager */
 	TObjectPtr<AActor> SpawnSimulatedActor(const TSubclassOf<AActor>& ActorClass,
 		const TObjectPtr<AChaosCacheManager>& CacheManager, const TObjectPtr<UChaosCacheCollection>& CacheCollection,
-		const bool bIsRecording, const TObjectPtr<UDataflowBaseContent>& DataflowContent);
+		const bool bIsRecording, const TObjectPtr<UDataflowBaseContent>& DataflowContent, const FTransform& ActorTransform);
 
 	/** Setup the skelmesh animations to be used in the scene/generator */
 	void SetupSkeletonAnimation(const TObjectPtr<AActor>& PreviewActor);
diff --git a/Engine/Plugins/Experimental/Dataflow/Source/DataflowEditor/Public/Dataflow/DataflowSimulationGenerator.h b/Engine/Plugins/Experimental/Dataflow/Source/DataflowEditor/Public/Dataflow/DataflowSimulationGenerator.h
--- a/Engine/Plugins/Experimental/Dataflow/Source/DataflowEditor/Public/Dataflow/DataflowSimulationGenerator.h
+++ b/Engine/Plugins/Experimental/Dataflow/Source/DataflowEditor/Public/Dataflow/DataflowSimulationGenerator.h
@@ -54,7 +54,7 @@ namespace UE::Dataflow
 		UWorld* SimulationWorld = nullptr;
 
 		/** Boolean to check if we are running the task in the background */
-		bool bBackgroundTask = true;
+		bool bAsyncCaching = true;
 	};
 
 	/** Async simulation resource that will be used while simulating */
@@ -106,7 +106,7 @@ namespace UE::Dataflow
 		/** Allocate the simulation resource from the properties */
         bool AllocateSimulationResource(const FVector2f& TimeRange, const int32 FrameRate,
         	const TObjectPtr<UChaosCacheCollection>& CacheAsset, const TSubclassOf<AActor>& ActorClass,
-        	const TObjectPtr<UDataflowBaseContent>& DataflowContent);
+        	const TObjectPtr<UDataflowBaseContent>& DataflowContent, const FTransform& BlueprintTransform);
 
         /** Free the simulation resource */
         void FreeSimulationResource();
@@ -153,6 +153,9 @@ namespace UE::Dataflow
 
 		/** Set the blueprint class for cache recording */
 		void SetBlueprintClass(const TSubclassOf<AActor>& InBlueprintClass);
+
+		/** Set the blueprint transform for cache recording */
+		void SetBlueprintTransform(const FTransform& InBlueprintTransform);
 		
 		/** Set the cache params for cache recording */
 		void SetCacheParams(const FDataflowPreviewCacheParams& InCacheParams);
@@ -186,6 +189,9 @@ namespace UE::Dataflow
 		/** Blueprint class used to spawn the actor */
 		TSubclassOf<AActor> BlueprintClass = nullptr;
 
+		/** Blueprint transform used to spawn the actor */
+		FTransform BlueprintTransform = FTransform::Identity;
+
 		/** Dataflow content */
 		TObjectPtr<UDataflowBaseContent> DataflowContent;
 
diff --git a/Engine/Plugins/Experimental/Dataflow/Source/DataflowEditor/Public/Dataflow/DataflowSimulationScene.h b/Engine/Plugins/Experimental/Dataflow/Source/DataflowEditor/Public/Dataflow/DataflowSimulationScene.h
--- a/Engine/Plugins/Experimental/Dataflow/Source/DataflowEditor/Public/Dataflow/DataflowSimulationScene.h
+++ b/Engine/Plugins/Experimental/Dataflow/Source/DataflowEditor/Public/Dataflow/DataflowSimulationScene.h
@@ -34,11 +34,15 @@ public:
 	void SetSimulationScene(class FDataflowSimulationScene* SimulationScene);
 
 	/** Caching blueprint actor class to spawn */
-	UPROPERTY(EditAnywhere, Category = "Preview")
+	UPROPERTY(EditAnywhere, Category = "Scene")
 	TSubclassOf<AActor> BlueprintClass = nullptr;
 
+	/** Blueprint actor transform */
+	UPROPERTY(EditAnywhere, Category = "Scene")
+	FTransform BlueprintTransform = FTransform::Identity;
+
 	/** Caching asset to be used to record the simulation  */
-	UPROPERTY(EditAnywhere, Category="Caching")
+	UPROPERTY(EditAnywhere, Category="Caching", DisplayName="CacheCollection")
 	TObjectPtr<UChaosCacheCollection> CacheAsset = nullptr;
 
 	/** Caching params used to record the simulation */
@@ -46,15 +50,15 @@ public:
 	FDataflowPreviewCacheParams CacheParams;
 
 	/** Geometry cache asset used to extract skeletal mesh results from simulation */
-	UPROPERTY(EditAnywhere, Category = "Geometry Cache Conversion")
+	UPROPERTY(EditAnywhere, Category = "Geometry", DisplayName="GeometryCache", meta=(EditCondition = "CacheAsset != nullptr"))
 	TObjectPtr<UGeometryCache> GeometryCacheAsset = nullptr;
 
 	/** Interpolates and saves geometry cache from Chaos cache */
-	UFUNCTION(CallInEditor, Category = "Geometry Cache Conversion")
+	UFUNCTION(CallInEditor, Category = "Geometry")
 	void GenerateGeometryCache();
 
 	/** Creates a new geometry cache file */
-	UFUNCTION(CallInEditor, Category = "Geometry Cache Conversion")
+	UFUNCTION(CallInEditor, Category = "Geometry")
 	void NewGeometryCache();
 
 private:
@@ -84,6 +88,9 @@ public:
 	
 	virtual ~FDataflowSimulationScene();
 
+	/** Functions that will be triggered when objects will be reinstanced (BP compilation) */
+	void OnObjectsReinstanced(const TMap<UObject*, UObject*>& ObjectsMap);
+
 	/** Tick data flow scene */
 	virtual void TickDataflowScene(const float DeltaSeconds) override;
 	
@@ -158,6 +165,9 @@ private:
 
 	/** Preview actor that will will be used to visualize the result of the simulation graph */
 	TObjectPtr<AActor> PreviewActor;
+
+	/** Handle for the delegate */
+	FDelegateHandle OnObjectsReinstancedHandle;
 };
 
 
diff --git a/Engine/Source/Runtime/Experimental/ChaosSolverEngine/Private/Chaos/ChaosDebugDrawSubsystem.cpp b/Engine/Source/Runtime/Experimental/ChaosSolverEngine/Private/Chaos/ChaosDebugDrawSubsystem.cpp
--- a/Engine/Source/Runtime/Experimental/ChaosSolverEngine/Private/Chaos/ChaosDebugDrawSubsystem.cpp
+++ b/Engine/Source/Runtime/Experimental/ChaosSolverEngine/Private/Chaos/ChaosDebugDrawSubsystem.cpp
@@ -260,7 +260,7 @@ bool UChaosDebugDrawSubsystem::ShouldCreateSubsystem(UObject* Outer) const
 #if CHAOS_DEBUG_DRAW
 	if (UWorld* World = Cast<UWorld>(Outer))
 	{
-		const bool bCreateDebugDraw = !IsRunningCommandlet() && !World->IsPreviewWorld();
+		const bool bCreateDebugDraw = !IsRunningCommandlet();
 
 		UE_LOG(LogChaosDD, Log, TEXT("%s Chaos Debug Draw Scene for world %s"), bCreateDebugDraw ? TEXT("Creating") : TEXT("Not creating"), *World->GetName());
 		
diff --git a/Engine/Source/Runtime/Experimental/Dataflow/Engine/Public/Dataflow/DataflowObject.h b/Engine/Source/Runtime/Experimental/Dataflow/Engine/Public/Dataflow/DataflowObject.h
--- a/Engine/Source/Runtime/Experimental/Dataflow/Engine/Public/Dataflow/DataflowObject.h
+++ b/Engine/Source/Runtime/Experimental/Dataflow/Engine/Public/Dataflow/DataflowObject.h
@@ -165,6 +165,10 @@ public:
 	/** Caching blueprint actor class to spawn */
 	UPROPERTY(DuplicateTransient, AssetRegistrySearchable)
 	TSubclassOf<AActor> PreviewBlueprintClass = nullptr;
+	
+	/** Caching blueprint actor transform to spawn */
+	UPROPERTY(DuplicateTransient, AssetRegistrySearchable)
+	FTransform PreviewBlueprintTransform = FTransform::Identity;
 
 #endif
 
diff --git a/Engine/Source/Runtime/Experimental/Dataflow/Engine/Public/Dataflow/DataflowPreview.h b/Engine/Source/Runtime/Experimental/Dataflow/Engine/Public/Dataflow/DataflowPreview.h
--- a/Engine/Source/Runtime/Experimental/Dataflow/Engine/Public/Dataflow/DataflowPreview.h
+++ b/Engine/Source/Runtime/Experimental/Dataflow/Engine/Public/Dataflow/DataflowPreview.h
@@ -28,5 +28,5 @@ public :
 
 	/** Boolean to check if the caching will be done on an async thread (if yes no GT dependency) */
 	UPROPERTY(EditAnywhere, AdvancedDisplay, Category = "Caching")
-	bool bBackgroundTask = true;
+	bool bAsyncCaching = true;
 };
diff --git a/Engine/Source/Runtime/Experimental/Dataflow/Simulation/Private/Dataflow/DataflowSimulationManager.cpp b/Engine/Source/Runtime/Experimental/Dataflow/Simulation/Private/Dataflow/DataflowSimulationManager.cpp
--- a/Engine/Source/Runtime/Experimental/Dataflow/Simulation/Private/Dataflow/DataflowSimulationManager.cpp
+++ b/Engine/Source/Runtime/Experimental/Dataflow/Simulation/Private/Dataflow/DataflowSimulationManager.cpp
@@ -8,6 +8,7 @@
 #include "Async/TaskGraphInterfaces.h"
 #include "Containers/Map.h"
 #include "Dataflow/Interfaces/DataflowPhysicsSolver.h"
+#include "ChaosDebugDraw/ChaosDDContext.h"
 #include "Engine/World.h"
 #include "Kismet/GameplayStatics.h"
 
@@ -39,11 +40,20 @@ namespace UE::Dataflow
 		class FDataflowSimulationTask
 		{
 		public:
+
+#if CHAOS_DEBUG_DRAW
+			FDataflowSimulationTask(const TObjectPtr<UDataflow>& InDataflowAsset, const TSharedPtr<FDataflowSimulationContext>& InSimulationContext,
+				const float InDeltaTime, const float InSimulationTime, const ChaosDD::Private::FChaosDDTaskParentContext& InParentDDContext)
+				: DataflowAsset(InDataflowAsset), SimulationContext(InSimulationContext), DeltaTime(InDeltaTime), SimulationTime(InSimulationTime), ParentDDContext(InParentDDContext)
+			{
+			}
+#else
 			FDataflowSimulationTask(const TObjectPtr<UDataflow>& InDataflowAsset, const TSharedPtr<FDataflowSimulationContext>& InSimulationContext,
-				const float InDeltaTime, const float InSimulationTime)
-				: DataflowAsset(InDataflowAsset), SimulationContext(InSimulationContext), DeltaTime(InDeltaTime), SimulationTime(InSimulationTime)
+							const float InDeltaTime, const float InSimulationTime)
+							: DataflowAsset(InDataflowAsset), SimulationContext(InSimulationContext), DeltaTime(InDeltaTime), SimulationTime(InSimulationTime)
 			{
 			}
+#endif
 
 			TStatId GetStatId() const
 			{
@@ -66,6 +76,9 @@ namespace UE::Dataflow
 
 			void DoTask(ENamedThreads::Type CurrentThread, const FGraphEventRef& MyCompletionGraphEvent)
 			{
+#if CHAOS_DEBUG_DRAW
+				ChaosDD::Private::FChaosDDScopeTaskContext DDTaskContext(ParentDDContext);
+#endif
 				Dataflow::EvaluateSimulationGraph(DataflowAsset, SimulationContext, DeltaTime, SimulationTime);
 			}
 
@@ -81,6 +94,11 @@ namespace UE::Dataflow
 
 			/** World simulation time  */
 			float SimulationTime;
+
+#if CHAOS_DEBUG_DRAW
+			/** Parent debug draw context */
+			const ChaosDD::Private::FChaosDDTaskParentContext ParentDDContext;
+#endif
 		};
 
 		inline void PreSimulationTick(const TObjectPtr<UObject>& SimulationWorld, const float SimulationTime, const float DeltaTime)
@@ -154,7 +172,7 @@ namespace UE::Dataflow
 }
 
 FDelegateHandle UDataflowSimulationManager::OnObjectPropertyChangedHandle;
-FDelegateHandle UDataflowSimulationManager::OnWorldTickEndHandle;
+FDelegateHandle UDataflowSimulationManager::OnWorldPostActorTick;
 FDelegateHandle UDataflowSimulationManager::OnCreatePhysicsStateHandle;
 FDelegateHandle UDataflowSimulationManager::OnDestroyPhysicsStateHandle;
 
@@ -166,7 +184,7 @@ void UDataflowSimulationManager::Tick(float DeltaTime)
 	Super::Tick(DeltaTime);
 	
 	check(IsInGameThread());
-	check(SimulationTasks.IsEmpty());
+	//check(SimulationTasks.IsEmpty());
 	
 	PreProcessSimulation(DeltaTime);
 	
@@ -203,7 +221,7 @@ void UDataflowSimulationManager::Tick(float DeltaTime)
 
 void UDataflowSimulationManager::OnStartup()
 {
-	OnWorldTickEndHandle = FWorldDelegates::OnWorldTickEnd.AddLambda(
+	OnWorldPostActorTick = FWorldDelegates::OnWorldPostActorTick.AddLambda(
 	[](const UWorld* SimulationWorld, ELevelTick LevelTick, const float DeltaSeconds)
 	{
 		if(SimulationWorld)
@@ -259,7 +277,7 @@ void UDataflowSimulationManager::OnShutdown()
 #if WITH_EDITOR
 	FCoreUObjectDelegates::OnObjectPropertyChanged.Remove(OnObjectPropertyChangedHandle);
 #endif
-	FWorldDelegates::OnWorldTickEnd.Remove(OnWorldTickEndHandle);
+	FWorldDelegates::OnWorldPostActorTick.Remove(OnWorldPostActorTick);
 	UActorComponent::GlobalCreatePhysicsDelegate.Remove(OnCreatePhysicsStateHandle);
 	UActorComponent::GlobalDestroyPhysicsDelegate.Remove(OnDestroyPhysicsStateHandle);
 }
@@ -402,15 +420,29 @@ void UDataflowSimulationManager::AdvanceSimulationProxies(const float DeltaTime,
 void UDataflowSimulationManager::StartSimulationTasks(const float DeltaTime, const float SimulationTime)
 {
 	check(IsInGameThread());
+	
+	// Wait until all tasks are complete
+	CompleteSimulationTasks();
+	
 	check(SimulationTasks.IsEmpty());
 
+	// Parent debug draw context
+#if CHAOS_DEBUG_DRAW
+	const ChaosDD::Private::FChaosDDTaskParentContext ParentDDContext;
+#endif
+
 	for(TPair<TObjectPtr<UDataflow>, UE::Dataflow::Private::FDataflowSimulationData>& DataflowData : SimulationData)
 	{
 		if(!DataflowData.Value.IsEmpty())
 		{
 			// Add a simulation task linked to that solver
+#if CHAOS_DEBUG_DRAW
+			SimulationTasks.Add(TGraphTask<UE::Dataflow::Private::FDataflowSimulationTask>::CreateTask(nullptr, ENamedThreads::GameThread).ConstructAndDispatchWhenReady(
+					DataflowData.Key, DataflowData.Value.SimulationContext, DeltaTime, SimulationTime, ParentDDContext));
+#else
 			SimulationTasks.Add(TGraphTask<UE::Dataflow::Private::FDataflowSimulationTask>::CreateTask(nullptr, ENamedThreads::GameThread).ConstructAndDispatchWhenReady(
 					DataflowData.Key, DataflowData.Value.SimulationContext, DeltaTime, SimulationTime));
+#endif
 		}
 	}
 }
diff --git a/Engine/Source/Runtime/Experimental/Dataflow/Simulation/Public/Dataflow/DataflowSimulationManager.h b/Engine/Source/Runtime/Experimental/Dataflow/Simulation/Public/Dataflow/DataflowSimulationManager.h
--- a/Engine/Source/Runtime/Experimental/Dataflow/Simulation/Public/Dataflow/DataflowSimulationManager.h
+++ b/Engine/Source/Runtime/Experimental/Dataflow/Simulation/Public/Dataflow/DataflowSimulationManager.h
@@ -123,8 +123,8 @@ private :
 	/** static delegate for object property changed */
 	static FDelegateHandle OnObjectPropertyChangedHandle;
 
-	/** static delegate for end of frame update */
-	static FDelegateHandle OnWorldTickEndHandle;
+	/** static delegate for post actor tick */
+	static FDelegateHandle OnWorldPostActorTick;
 
 	/** static delegate for physics state creation */
 	static FDelegateHandle OnCreatePhysicsStateHandle;
