Dataflow Editor: when user presses 'F' while a tool is running, use that tool's focus box to orient the camera, if it has one defined

diff --git a/Engine/Plugins/Experimental/Dataflow/Source/DataflowEditor/Private/Dataflow/DataflowConstructionViewport.cpp b/Engine/Plugins/Experimental/Dataflow/Source/DataflowEditor/Private/Dataflow/DataflowConstructionViewport.cpp
--- a/Engine/Plugins/Experimental/Dataflow/Source/DataflowEditor/Private/Dataflow/DataflowConstructionViewport.cpp
+++ b/Engine/Plugins/Experimental/Dataflow/Source/DataflowEditor/Private/Dataflow/DataflowConstructionViewport.cpp
@@ -38,10 +38,15 @@ TSharedPtr<SWidget> SDataflowConstructionViewport::MakeViewportToolbar()
 
 void SDataflowConstructionViewport::OnFocusViewportToSelection()
 {
-	if(const FDataflowPreviewSceneBase* PreviewScene = static_cast<FDataflowPreviewSceneBase*>(Client->GetPreviewScene()))
+	const UDataflowEditorMode* const DataflowEdMode = GetEdMode();
+
+	if (DataflowEdMode)
 	{
-		const FBox SceneBoundingBox = PreviewScene->GetBoundingBox();
-		Client->FocusViewportOnBox(SceneBoundingBox);
+		const FBox BoundingBox = DataflowEdMode->SelectionBoundingBox();
+		if (BoundingBox.IsValid && !(BoundingBox.Min == FVector::Zero() && BoundingBox.Max == FVector::Zero()))
+		{
+			Client->FocusViewportOnBox(BoundingBox);
+		}
 	}
 }
 
diff --git a/Engine/Plugins/Experimental/Dataflow/Source/DataflowEditor/Private/Dataflow/DataflowEditorMode.cpp b/Engine/Plugins/Experimental/Dataflow/Source/DataflowEditor/Private/Dataflow/DataflowEditorMode.cpp
--- a/Engine/Plugins/Experimental/Dataflow/Source/DataflowEditor/Private/Dataflow/DataflowEditorMode.cpp
+++ b/Engine/Plugins/Experimental/Dataflow/Source/DataflowEditor/Private/Dataflow/DataflowEditorMode.cpp
@@ -648,6 +648,17 @@ FBox UDataflowEditorMode::SceneBoundingBox() const
 
 FBox UDataflowEditorMode::SelectionBoundingBox() const
 {
+	// if Tool supports custom Focus box, use that first
+	if (GetToolManager()->HasAnyActiveTool())
+	{
+		UInteractiveTool* const Tool = GetToolManager()->GetActiveTool(EToolSide::Mouse);
+		IInteractiveToolCameraFocusAPI* const FocusAPI = Cast<IInteractiveToolCameraFocusAPI>(Tool);
+		if (FocusAPI && FocusAPI->SupportsWorldSpaceFocusBox())
+		{
+			return FocusAPI->GetWorldSpaceFocusBox();
+		}
+	}
+
 	// If the selection is on the GetBoundingBox is automatically computing the selection one
 	FBox Bounds = ConstructionScene->GetBoundingBox();
 	if (Bounds.IsValid)
