Dataflow:
- Added missing delegate remove

diff --git a/Engine/Plugins/Experimental/Dataflow/Source/DataflowEditor/Private/Dataflow/DataflowCollectionSpreadSheet.cpp b/Engine/Plugins/Experimental/Dataflow/Source/DataflowEditor/Private/Dataflow/DataflowCollectionSpreadSheet.cpp
--- a/Engine/Plugins/Experimental/Dataflow/Source/DataflowEditor/Private/Dataflow/DataflowCollectionSpreadSheet.cpp
+++ b/Engine/Plugins/Experimental/Dataflow/Source/DataflowEditor/Private/Dataflow/DataflowCollectionSpreadSheet.cpp
@@ -88,5 +88,6 @@ FDataflowCollectionSpreadSheet::~FDataflowCollectionSpreadSheet()
 	if (CollectionSpreadSheet)
 	{
 		CollectionSpreadSheet->GetOnPinnedDownChangedDelegate().Remove(OnPinnedDownChangedDelegateHandle);
+		CollectionSpreadSheet->GetOnRefreshLockedChangedDelegate().Remove(OnRefreshLockedChangedDelegateHandle);
 	}
 }
