Dataflow:
- Added display of invalid Dataflow Node Output
- Bug fixes

diff --git a/Engine/Plugins/Experimental/Dataflow/Source/DataflowEditor/Private/Dataflow/DataflowEditorToolkit.cpp b/Engine/Plugins/Experimental/Dataflow/Source/DataflowEditor/Private/Dataflow/DataflowEditorToolkit.cpp
--- a/Engine/Plugins/Experimental/Dataflow/Source/DataflowEditor/Private/Dataflow/DataflowEditorToolkit.cpp
+++ b/Engine/Plugins/Experimental/Dataflow/Source/DataflowEditor/Private/Dataflow/DataflowEditorToolkit.cpp
@@ -1067,6 +1067,11 @@ void FDataflowEditorToolkit::OnConstructionViewSelectionChanged(const TArray<UPr
 	}
 }
 
+void FDataflowEditorToolkit::OnFinishEvaluate()
+{
+	// Refresh graph display to update node output pin display (invalid or valid)
+	GraphEditor->NotifyGraphChanged();
+}
 
 void FDataflowEditorToolkit::Tick(float DeltaTime)
 {
@@ -1167,6 +1172,11 @@ TSharedRef<SDataflowGraphEditor> FDataflowEditorToolkit::CreateGraphEditorWidget
 
 				EditorContent->SetLastModifiedTimestamp(LastNodeTimestamp);
 			}
+
+			//
+			// Graph evaluation done
+			//
+			OnFinishEvaluate();
 		};
 	
 	DataflowEditor->UpdateTerminalContents(Dataflow::FTimestamp::Invalid);
diff --git a/Engine/Plugins/Experimental/Dataflow/Source/DataflowEditor/Private/Dataflow/DataflowSNode.cpp b/Engine/Plugins/Experimental/Dataflow/Source/DataflowEditor/Private/Dataflow/DataflowSNode.cpp
--- a/Engine/Plugins/Experimental/Dataflow/Source/DataflowEditor/Private/Dataflow/DataflowSNode.cpp
+++ b/Engine/Plugins/Experimental/Dataflow/Source/DataflowEditor/Private/Dataflow/DataflowSNode.cpp
@@ -23,6 +23,30 @@
 #include UE_INLINE_GENERATED_CPP_BY_NAME(DataflowSNode)
 
 #define LOCTEXT_NAMESPACE "SDataflowEdNode"
+
+//
+// SDataflowOutputPin
+//
+
+void SDataflowOutputPin::Construct(const FArguments& InArgs, UEdGraphPin* InPin)
+{
+	const bool bIsPinInvalid = InArgs._IsPinInvalid.Get();
+	const FText InvalidPinDisplayText = bIsPinInvalid ? NSLOCTEXT("DataflowGraph", "DataflowOutputPinInvalidText", "*") : NSLOCTEXT("DataflowGraph", "DataflowOutputPinValidText", " ");
+
+	SGraphPin::Construct(SGraphPin::FArguments(), InPin);
+
+	GetLabelAndValue()->AddSlot()
+		.Padding(2.0f, 0.0f, 0.0f, 0.0f)
+		[
+			SNew(STextBlock)
+			.Text_Lambda([InvalidPinDisplayText]()
+				{
+					return InvalidPinDisplayText;
+				})
+		.MinDesiredWidth(5)
+		];
+}
+
 //
 // SDataflowEdNode
 //
@@ -32,6 +56,7 @@ void SDataflowEdNode::Construct(const FArguments& InArgs, UDataflowEdNode* InNod
 	GraphNode = InNode;
 	DataflowGraphNode = Cast<UDataflowEdNode>(InNode);
 	DataflowInterface = InArgs._DataflowInterface;
+
 	UpdateGraphNode();
 
 	
@@ -112,6 +137,45 @@ void SDataflowEdNode::Construct(const FArguments& InArgs, UDataflowEdNode* InNod
 	*/
 }
 
+TSharedPtr<SGraphPin> SDataflowEdNode::CreatePinWidget(UEdGraphPin* Pin) const
+{
+	if (Pin->Direction == EEdGraphPinDirection::EGPD_Output)
+	{
+		if (DataflowGraphNode)
+		{
+			if (TSharedPtr<FDataflowNode> DataflowNode = DataflowGraphNode->GetDataflowNode())
+			{
+				if (FDataflowOutput* Output = DataflowNode->FindOutput(Pin->GetFName()))
+				{
+					if (const TSharedPtr<Dataflow::FEngineContext> DataflowContext = DataflowInterface->GetDataflowContext())
+					{
+						TSet<Dataflow::FContextCacheKey> CacheKeys;
+						const int32 NumKeys = DataflowContext->GetKeys(CacheKeys);
+
+						//
+						// DataStore is empty or 
+						// CacheKey is not in DataStore or
+						// Node's Timestamp is invalid or
+						// Node's Timestamp is greater than CacheKey's Timestamp -> Pin is invalid
+						//
+						const bool bIsOutputInvalid = !NumKeys ||
+							!CacheKeys.Contains(Output->CacheKey()) ||
+PRAGMA_DISABLE_DEPRECATION_WARNINGS  // Until LastModifiedTimestamp becomes private
+							DataflowNode->LastModifiedTimestamp.IsInvalid() || 
+							!DataflowContext->IsCacheEntryAfterTimestamp(Output->CacheKey(), DataflowNode->LastModifiedTimestamp);
+PRAGMA_ENABLE_DEPRECATION_WARNINGS
+
+						return SNew(SDataflowOutputPin, Pin)
+								.IsPinInvalid(bIsOutputInvalid);
+					}
+				}
+			}
+		}
+	}
+
+	return SGraphNode::CreatePinWidget(Pin);
+}
+
 TArray<FOverlayWidgetInfo> SDataflowEdNode::GetOverlayWidgets(bool bSelected, const FVector2D& WidgetSize) const
 {
 	TArray<FOverlayWidgetInfo> Widgets = SGraphNode::GetOverlayWidgets(bSelected, WidgetSize);
diff --git a/Engine/Plugins/Experimental/Dataflow/Source/DataflowEditor/Public/Dataflow/DataflowEditorToolkit.h b/Engine/Plugins/Experimental/Dataflow/Source/DataflowEditor/Public/Dataflow/DataflowEditorToolkit.h
--- a/Engine/Plugins/Experimental/Dataflow/Source/DataflowEditor/Public/Dataflow/DataflowEditorToolkit.h
+++ b/Engine/Plugins/Experimental/Dataflow/Source/DataflowEditor/Public/Dataflow/DataflowEditorToolkit.h
@@ -94,7 +94,8 @@ protected:
 
 	// Node evaluation
 	void EvaluateNode(const FDataflowNode* Node, const FDataflowOutput* Output, Dataflow::FTimestamp& InOutTimestamp);
-	
+	void OnFinishEvaluate();
+
 private:
 	
 	// Spawning of all the additional tabs (viewport,details ones are coming from the base asset toolkit)
diff --git a/Engine/Plugins/Experimental/Dataflow/Source/DataflowEditor/Public/Dataflow/DataflowSNode.h b/Engine/Plugins/Experimental/Dataflow/Source/DataflowEditor/Public/Dataflow/DataflowSNode.h
--- a/Engine/Plugins/Experimental/Dataflow/Source/DataflowEditor/Public/Dataflow/DataflowSNode.h
+++ b/Engine/Plugins/Experimental/Dataflow/Source/DataflowEditor/Public/Dataflow/DataflowSNode.h
@@ -7,12 +7,32 @@
 #include "UObject/GCObject.h"
 #include "Dataflow/DataflowCore.h"
 #include "Dataflow/DataflowSEditorInterface.h"
+#include "SGraphPin.h"
 
 #include "DataflowSNode.generated.h"
 
 class UDataflowEdNode;
 class SCheckBox;
 
+//
+// SDataflowOutputPin
+//
+
+class SDataflowOutputPin : public SGraphPin
+{
+public:
+	SLATE_BEGIN_ARGS(SDataflowOutputPin)
+		: _IsPinInvalid(false)
+	{}
+		SLATE_ATTRIBUTE(bool, IsPinInvalid)
+	SLATE_END_ARGS()
+
+	void Construct(const FArguments& InArgs, UEdGraphPin* InPin);
+
+private:
+	TAttribute<bool> IsPinInvalid;
+};
+
 //
 // SDataflowEdNode
 //
@@ -59,7 +79,7 @@ private:
 	/** Checks whether Add pin button should currently be visible */
 	virtual EVisibility IsAddPinButtonVisible() const;
 
-	
+	virtual TSharedPtr<SGraphPin> CreatePinWidget(UEdGraphPin* Pin) const override;
 
 	TObjectPtr<UDataflowEdNode> DataflowGraphNode = nullptr;	
 
diff --git a/Engine/Source/Runtime/Experimental/Dataflow/Core/Private/Dataflow/DataflowNodeParameters.cpp b/Engine/Source/Runtime/Experimental/Dataflow/Core/Private/Dataflow/DataflowNodeParameters.cpp
--- a/Engine/Source/Runtime/Experimental/Dataflow/Core/Private/Dataflow/DataflowNodeParameters.cpp
+++ b/Engine/Source/Runtime/Experimental/Dataflow/Core/Private/Dataflow/DataflowNodeParameters.cpp
@@ -47,6 +47,25 @@ namespace Dataflow
 #endif
 	}
 
+	bool FContext::IsCacheEntryAfterTimestamp(FContextCacheKey InKey, const FTimestamp InTimestamp)
+	{
+		if (HasData(InKey))
+		{
+			if (TUniquePtr<FContextCacheElementBase>* CacheEntry = GetDataImpl(InKey))
+			{
+				if (*CacheEntry)
+				{
+					if ((*CacheEntry)->GetTimestamp() >= InTimestamp)
+					{
+						return true;
+					}
+				}
+			}
+		}
+
+		return false;
+	}
+
 	FContextScopedCallstack::FContextScopedCallstack(FContext& InContext, const FDataflowConnection* InConnection)
 		: Context(InContext)
 		, Connection(InConnection)
diff --git a/Engine/Source/Runtime/Experimental/Dataflow/Core/Public/Dataflow/DataflowNodeParameters.h b/Engine/Source/Runtime/Experimental/Dataflow/Core/Public/Dataflow/DataflowNodeParameters.h
--- a/Engine/Source/Runtime/Experimental/Dataflow/Core/Public/Dataflow/DataflowNodeParameters.h
+++ b/Engine/Source/Runtime/Experimental/Dataflow/Core/Public/Dataflow/DataflowNodeParameters.h
@@ -270,6 +270,7 @@ namespace Dataflow
 		DATAFLOWCORE_API void PopFromCallstack(const FDataflowConnection* Connection);
 		DATAFLOWCORE_API bool IsInCallstack(const FDataflowConnection* Connection) const;
 
+		DATAFLOWCORE_API bool IsCacheEntryAfterTimestamp(FContextCacheKey InKey, const FTimestamp InTimestamp);
 
 	private:
 #if DATAFLOW_EDITOR_EVALUATION
