Dataflow : Disable dataflow asset editor when opening them out of the context of an asset for 5.5
- Set p.Dataflow.IsEditable to be false by default and display a message to informn the user on how to edit the dataflow

diff --git a/Engine/Plugins/Experimental/Dataflow/Source/DataflowEditor/Private/Dataflow/AssetDefinition_DataflowAsset.cpp b/Engine/Plugins/Experimental/Dataflow/Source/DataflowEditor/Private/Dataflow/AssetDefinition_DataflowAsset.cpp
--- a/Engine/Plugins/Experimental/Dataflow/Source/DataflowEditor/Private/Dataflow/AssetDefinition_DataflowAsset.cpp
+++ b/Engine/Plugins/Experimental/Dataflow/Source/DataflowEditor/Private/Dataflow/AssetDefinition_DataflowAsset.cpp
@@ -17,7 +17,7 @@
 #define LOCTEXT_NAMESPACE "AssetActions_DataflowAsset"
 
 
-bool bCanEditDataflow = true;
+bool bCanEditDataflow = false;
 FAutoConsoleVariableRef CVarDataflowIsEditable(TEXT("p.Dataflow.IsEditable"), bCanEditDataflow, TEXT("Whether to allow edits of the dataflow [def:true]"));
 
 namespace DataflowAssetDefinitionHelpers
@@ -175,34 +175,44 @@ UThumbnailInfo* UAssetDefinition_DataflowAsset::LoadThumbnailInfo(const FAssetDa
 
 FAssetOpenSupport UAssetDefinition_DataflowAsset::GetAssetOpenSupport(const FAssetOpenSupportArgs& OpenSupportArgs) const
 {
-	if (bCanEditDataflow)
-	{
-		return Super::GetAssetOpenSupport(OpenSupportArgs);
-	}
-	return FAssetOpenSupport(EAssetOpenMethod::View, false);
+	return Super::GetAssetOpenSupport(OpenSupportArgs);
 }
 
 
 EAssetCommandResult UAssetDefinition_DataflowAsset::OpenAssets(const FAssetOpenArgs& OpenArgs) const
 {
-	TArray<UDataflow*> DataflowObjects = OpenArgs.LoadObjects<UDataflow>();
-
-	// For now the dataflow editor only works on one asset at a time
-	ensure(DataflowObjects.Num() == 0 || DataflowObjects.Num() == 1);
-
-	if (DataflowObjects.Num() == 1)
+	if (bCanEditDataflow)
 	{
-		UAssetEditorSubsystem* const AssetEditorSubsystem = GEditor->GetEditorSubsystem<UAssetEditorSubsystem>();
-		UDataflowEditor* const AssetEditor = NewObject<UDataflowEditor>(AssetEditorSubsystem, NAME_None, RF_Transient);
+		TArray<UDataflow*> DataflowObjects = OpenArgs.LoadObjects<UDataflow>();
+
+		// For now the dataflow editor only works on one asset at a time
+		ensure(DataflowObjects.Num() == 0 || DataflowObjects.Num() == 1);
 
-		// Validate the asset
-		if (UDataflow* const DataflowAsset = CastChecked<UDataflow>(DataflowObjects[0]))
+		if (DataflowObjects.Num() == 1)
 		{
-			AssetEditor->Initialize({DataflowAsset});
-			return EAssetCommandResult::Handled;
+			UAssetEditorSubsystem* const AssetEditorSubsystem = GEditor->GetEditorSubsystem<UAssetEditorSubsystem>();
+			UDataflowEditor* const AssetEditor = NewObject<UDataflowEditor>(AssetEditorSubsystem, NAME_None, RF_Transient);
+
+			// Validate the asset
+			if (UDataflow* const DataflowAsset = CastChecked<UDataflow>(DataflowObjects[0]))
+			{
+				AssetEditor->Initialize({ DataflowAsset });
+				return EAssetCommandResult::Handled;
+			}
 		}
 	}
-
+	else
+	{
+		TSharedRef<SMessageDialog> MessageDialog = SNew(SMessageDialog)
+			.Title(FText(LOCTEXT("Dataflow_OpenAssetDialog_Title", "Dataflow Asset")))
+			.Message(LOCTEXT("Dataflow_OpenAssetDialog_Text", "Dataflow assets can only be changed while editing assets using them (Cloth, Flesh, Geometry Collection, ...)"))
+			.Buttons({
+				SMessageDialog::FButton(LOCTEXT("Ok", "Ok"))
+				.SetPrimary(true)
+			});
+		MessageDialog->ShowModal();
+		return EAssetCommandResult::Handled;
+	}
 	return EAssetCommandResult::Unhandled;
 }
 
