Cloth Editor: fix crash when changing Dataflow property. We were not passing in a null DataflowEditor when reinitializing the graph editor.

diff --git a/Engine/Plugins/Experimental/Dataflow/Source/DataflowEditor/Public/Dataflow/DataflowGraphEditor.h b/Engine/Plugins/Experimental/Dataflow/Source/DataflowEditor/Public/Dataflow/DataflowGraphEditor.h
--- a/Engine/Plugins/Experimental/Dataflow/Source/DataflowEditor/Public/Dataflow/DataflowGraphEditor.h
+++ b/Engine/Plugins/Experimental/Dataflow/Source/DataflowEditor/Public/Dataflow/DataflowGraphEditor.h
@@ -50,7 +50,7 @@ public:
 	SLATE_ARGUMENT(TSharedPtr<IStructureDetailsView>, DetailsView)
 	SLATE_ARGUMENT(FDataflowEditorCommands::FGraphEvaluationCallback, EvaluateGraph)
 	SLATE_ARGUMENT(FDataflowEditorCommands::FOnDragDropEventCallback, OnDragDropEvent)
-	SLATE_ARGUMENT(UDataflowEditor*, DataflowEditor)
+	SLATE_ARGUMENT_DEFAULT(UDataflowEditor*, DataflowEditor) = nullptr;
 	SLATE_END_ARGS()
 
 	// This delegate exists in SGraphEditor but it is not multicast, and we are going to bind it to OnSelectedNodesChanged().
