Change Dataflow namespace to UE::Dataflow to be ok with coding standard

diff --git a/Engine/Plugins/Experimental/ChaosFlesh/Source/ChaosFleshNodes/Private/Dataflow/ChaosFleshVisualizeFiberFieldNode.cpp b/Engine/Plugins/Experimental/ChaosFlesh/Source/ChaosFleshNodes/Private/Dataflow/ChaosFleshVisualizeFiberFieldNode.cpp
--- a/Engine/Plugins/Experimental/ChaosFlesh/Source/ChaosFleshNodes/Private/Dataflow/ChaosFleshVisualizeFiberFieldNode.cpp
+++ b/Engine/Plugins/Experimental/ChaosFlesh/Source/ChaosFleshNodes/Private/Dataflow/ChaosFleshVisualizeFiberFieldNode.cpp
@@ -7,7 +7,7 @@
 
 #include UE_INLINE_GENERATED_CPP_BY_NAME(ChaosFleshVisualizeFiberFieldNode)
 
-void FVisualizeFiberFieldNode::Evaluate(Dataflow::FContext& Context, const FDataflowOutput* Out) const
+void FVisualizeFiberFieldNode::Evaluate(UE::Dataflow::FContext& Context, const FDataflowOutput* Out) const
 {
 	if (Out->IsA<FFieldCollection>(&VectorField))
 	{
@@ -51,7 +51,7 @@ void FVisualizeFiberFieldNode::Evaluate(Dataflow::FContext& Context, const FData
 	}
 }
 
-void FVisualizePositionTargetsNode::Evaluate(Dataflow::FContext& Context, const FDataflowOutput* Out) const
+void FVisualizePositionTargetsNode::Evaluate(UE::Dataflow::FContext& Context, const FDataflowOutput* Out) const
 {
 	if (Out->IsA<FFieldCollection>(&VectorField))
 	{
@@ -84,7 +84,7 @@ void FVisualizePositionTargetsNode::Evaluate(Dataflow::FContext& Context, const
 	}
 }
 
-void FVisualizeKinematicFacesNode::Evaluate(Dataflow::FContext& Context, const FDataflowOutput* Out) const
+void FVisualizeKinematicFacesNode::Evaluate(UE::Dataflow::FContext& Context, const FDataflowOutput* Out) const
 {
 	if (Out->IsA<FManagedArrayCollection>(&Collection))
 	{
diff --git a/Engine/Plugins/Experimental/ChaosFlesh/Source/ChaosFleshNodes/Private/Dataflow/GeometryCollectionAppendCollectionTransformNode.cpp b/Engine/Plugins/Experimental/ChaosFlesh/Source/ChaosFleshNodes/Private/Dataflow/GeometryCollectionAppendCollectionTransformNode.cpp
--- a/Engine/Plugins/Experimental/ChaosFlesh/Source/ChaosFleshNodes/Private/Dataflow/GeometryCollectionAppendCollectionTransformNode.cpp
+++ b/Engine/Plugins/Experimental/ChaosFlesh/Source/ChaosFleshNodes/Private/Dataflow/GeometryCollectionAppendCollectionTransformNode.cpp
@@ -5,7 +5,7 @@
 #include UE_INLINE_GENERATED_CPP_BY_NAME(GeometryCollectionAppendCollectionTransformNode)
 
 void
-FAppendToCollectionTransformAttributeDataflowNode::Evaluate(Dataflow::FContext& Context, const FDataflowOutput* Out) const
+FAppendToCollectionTransformAttributeDataflowNode::Evaluate(UE::Dataflow::FContext& Context, const FDataflowOutput* Out) const
 {
 	if (Out->IsA<DataType>(&Collection))
 	{
diff --git a/Engine/Plugins/Experimental/ChaosFlesh/Source/ChaosFleshNodes/Public/Dataflow/ChaosFleshAddKinematicParticlesNode.h b/Engine/Plugins/Experimental/ChaosFlesh/Source/ChaosFleshNodes/Public/Dataflow/ChaosFleshAddKinematicParticlesNode.h
--- a/Engine/Plugins/Experimental/ChaosFlesh/Source/ChaosFleshNodes/Public/Dataflow/ChaosFleshAddKinematicParticlesNode.h
+++ b/Engine/Plugins/Experimental/ChaosFlesh/Source/ChaosFleshNodes/Public/Dataflow/ChaosFleshAddKinematicParticlesNode.h
@@ -47,7 +47,7 @@ public:
 
 
 
-	FAddKinematicParticlesDataflowNode(const Dataflow::FNodeParameters& InParam, FGuid InGuid = FGuid::NewGuid())
+	FAddKinematicParticlesDataflowNode(const UE::Dataflow::FNodeParameters& InParam, FGuid InGuid = FGuid::NewGuid())
 		: FDataflowNode(InParam, InGuid)
 	{
 		RegisterInputConnection(&Collection);
@@ -58,6 +58,6 @@ public:
 		RegisterOutputConnection(&TargetIndicesOut);
 	}
 
-	virtual void Evaluate(Dataflow::FContext& Context, const FDataflowOutput* Out) const override;
+	virtual void Evaluate(UE::Dataflow::FContext& Context, const FDataflowOutput* Out) const override;
 };
 
diff --git a/Engine/Plugins/Experimental/ChaosFlesh/Source/ChaosFleshNodes/Public/Dataflow/ChaosFleshAppendTetrahedralCollectionNode.h b/Engine/Plugins/Experimental/ChaosFlesh/Source/ChaosFleshNodes/Public/Dataflow/ChaosFleshAppendTetrahedralCollectionNode.h
--- a/Engine/Plugins/Experimental/ChaosFlesh/Source/ChaosFleshNodes/Public/Dataflow/ChaosFleshAppendTetrahedralCollectionNode.h
+++ b/Engine/Plugins/Experimental/ChaosFlesh/Source/ChaosFleshNodes/Public/Dataflow/ChaosFleshAppendTetrahedralCollectionNode.h
@@ -31,7 +31,7 @@ public:
 	UPROPERTY(meta = (DataflowOutput, DisplayName = "GeometryGroupIndicesOut2"))
 		TArray<FString> GeometryGroupGuidsOut2;
 
-	FAppendTetrahedralCollectionDataflowNode(const Dataflow::FNodeParameters& InParam, FGuid InGuid = FGuid::NewGuid())
+	FAppendTetrahedralCollectionDataflowNode(const UE::Dataflow::FNodeParameters& InParam, FGuid InGuid = FGuid::NewGuid())
 		: FDataflowNode(InParam, InGuid)
 	{
 		RegisterInputConnection(&Collection1);
@@ -41,5 +41,5 @@ public:
 		RegisterOutputConnection(&GeometryGroupGuidsOut2);
 	}
 
-	virtual void Evaluate(Dataflow::FContext& Context, const FDataflowOutput* Out) const override;
+	virtual void Evaluate(UE::Dataflow::FContext& Context, const FDataflowOutput* Out) const override;
 };
diff --git a/Engine/Plugins/Experimental/ChaosFlesh/Source/ChaosFleshNodes/Public/Dataflow/ChaosFleshAuthorSceneCollisionCandidatesNode.h b/Engine/Plugins/Experimental/ChaosFlesh/Source/ChaosFleshNodes/Public/Dataflow/ChaosFleshAuthorSceneCollisionCandidatesNode.h
--- a/Engine/Plugins/Experimental/ChaosFlesh/Source/ChaosFleshNodes/Public/Dataflow/ChaosFleshAuthorSceneCollisionCandidatesNode.h
+++ b/Engine/Plugins/Experimental/ChaosFlesh/Source/ChaosFleshNodes/Public/Dataflow/ChaosFleshAuthorSceneCollisionCandidatesNode.h
@@ -39,7 +39,7 @@ public:
 	UPROPERTY(EditAnywhere, Category = "Dataflow", meta = (DataflowInput, DisplayName = "Raycast Origin Bone Index"))
 	int32 OriginBoneIndex = 0;
 
-	FAuthorSceneCollisionCandidates(const Dataflow::FNodeParameters& InParam, FGuid InGuid = FGuid::NewGuid())
+	FAuthorSceneCollisionCandidates(const UE::Dataflow::FNodeParameters& InParam, FGuid InGuid = FGuid::NewGuid())
 		: FDataflowNode(InParam, InGuid)
 	{
 		RegisterInputConnection(&Collection);
@@ -48,5 +48,5 @@ public:
 		RegisterInputConnection(&OriginBoneIndex);
 	}
 
-	virtual void Evaluate(Dataflow::FContext& Context, const FDataflowOutput* Out) const override;
+	virtual void Evaluate(UE::Dataflow::FContext& Context, const FDataflowOutput* Out) const override;
 };
diff --git a/Engine/Plugins/Experimental/ChaosFlesh/Source/ChaosFleshNodes/Public/Dataflow/ChaosFleshCalculateTetrehedralMetricsNode.h b/Engine/Plugins/Experimental/ChaosFlesh/Source/ChaosFleshNodes/Public/Dataflow/ChaosFleshCalculateTetrehedralMetricsNode.h
--- a/Engine/Plugins/Experimental/ChaosFlesh/Source/ChaosFleshNodes/Public/Dataflow/ChaosFleshCalculateTetrehedralMetricsNode.h
+++ b/Engine/Plugins/Experimental/ChaosFlesh/Source/ChaosFleshNodes/Public/Dataflow/ChaosFleshCalculateTetrehedralMetricsNode.h
@@ -23,14 +23,14 @@ public:
 	UPROPERTY(meta = (DataflowInput, DataflowOutput, DataflowPassthrough = "Collection", DisplayName = "Collection"))
 		FManagedArrayCollection Collection;
 
-	FCalculateTetMetrics(const Dataflow::FNodeParameters& InParam, FGuid InGuid = FGuid::NewGuid())
+	FCalculateTetMetrics(const UE::Dataflow::FNodeParameters& InParam, FGuid InGuid = FGuid::NewGuid())
 		: FDataflowNode(InParam, InGuid)
 	{
 		RegisterInputConnection(&Collection);
 		RegisterOutputConnection(&Collection, &Collection);
 	}
 
-	virtual void Evaluate(Dataflow::FContext& Context, const FDataflowOutput* Out) const override;
+	virtual void Evaluate(UE::Dataflow::FContext& Context, const FDataflowOutput* Out) const override;
 };
 
 
diff --git a/Engine/Plugins/Experimental/ChaosFlesh/Source/ChaosFleshNodes/Public/Dataflow/ChaosFleshCollisionBodyConstraintNode.h b/Engine/Plugins/Experimental/ChaosFlesh/Source/ChaosFleshNodes/Public/Dataflow/ChaosFleshCollisionBodyConstraintNode.h
--- a/Engine/Plugins/Experimental/ChaosFlesh/Source/ChaosFleshNodes/Public/Dataflow/ChaosFleshCollisionBodyConstraintNode.h
+++ b/Engine/Plugins/Experimental/ChaosFlesh/Source/ChaosFleshNodes/Public/Dataflow/ChaosFleshCollisionBodyConstraintNode.h
@@ -30,7 +30,7 @@ public:
 	UPROPERTY(meta = (DataflowInput, DisplayName = "SkeletalMesh"))
 		TObjectPtr<USkeletalMesh> SkeletalMeshIn;
 
-	FKinematicBodySetupInitializationDataflowNode(const Dataflow::FNodeParameters& InParam, FGuid InGuid = FGuid::NewGuid())
+	FKinematicBodySetupInitializationDataflowNode(const UE::Dataflow::FNodeParameters& InParam, FGuid InGuid = FGuid::NewGuid())
 		: FDataflowNode(InParam, InGuid)
 	{
 		RegisterInputConnection(&Collection);
@@ -38,6 +38,6 @@ public:
 		RegisterInputConnection(&SkeletalMeshIn);
 	}
 
-	virtual void Evaluate(Dataflow::FContext& Context, const FDataflowOutput* Out) const override;
+	virtual void Evaluate(UE::Dataflow::FContext& Context, const FDataflowOutput* Out) const override;
 };
 
diff --git a/Engine/Plugins/Experimental/ChaosFlesh/Source/ChaosFleshNodes/Public/Dataflow/ChaosFleshComputeFiberFieldNode.h b/Engine/Plugins/Experimental/ChaosFlesh/Source/ChaosFleshNodes/Public/Dataflow/ChaosFleshComputeFiberFieldNode.h
--- a/Engine/Plugins/Experimental/ChaosFlesh/Source/ChaosFleshNodes/Public/Dataflow/ChaosFleshComputeFiberFieldNode.h
+++ b/Engine/Plugins/Experimental/ChaosFlesh/Source/ChaosFleshNodes/Public/Dataflow/ChaosFleshComputeFiberFieldNode.h
@@ -50,7 +50,7 @@ public:
 	UPROPERTY(EditAnywhere, Category = "Dataflow")
 	float Tolerance = 1.0e-7;
 
-	FComputeFiberFieldNode(const Dataflow::FNodeParameters& InParam, FGuid InGuid = FGuid::NewGuid())
+	FComputeFiberFieldNode(const UE::Dataflow::FNodeParameters& InParam, FGuid InGuid = FGuid::NewGuid())
 		: FDataflowNode(InParam, InGuid)
 	{
 		RegisterInputConnection(&Collection);
@@ -59,7 +59,7 @@ public:
 		RegisterOutputConnection(&Collection, &Collection);
 	}
 
-	virtual void Evaluate(Dataflow::FContext& Context, const FDataflowOutput* Out) const override;
+	virtual void Evaluate(UE::Dataflow::FContext& Context, const FDataflowOutput* Out) const override;
 
 	TArray<int32> GetNonZeroIndices(const TArray<uint8>& Map) const;
 
@@ -117,7 +117,7 @@ public:
 	UPROPERTY(meta = (DataflowOutput, DisplayName = "VectorField"))
 	FFieldCollection VectorField;
 
-	FComputeFiberStreamlineNode(const Dataflow::FNodeParameters& InParam, FGuid InGuid = FGuid::NewGuid())
+	FComputeFiberStreamlineNode(const UE::Dataflow::FNodeParameters& InParam, FGuid InGuid = FGuid::NewGuid())
 		: FDataflowNode(InParam, InGuid)
 	{
 		RegisterInputConnection(&Collection);
@@ -127,5 +127,5 @@ public:
 		RegisterOutputConnection(&VectorField);
 	}
 
-	virtual void Evaluate(Dataflow::FContext& Context, const FDataflowOutput* Out) const override;
+	virtual void Evaluate(UE::Dataflow::FContext& Context, const FDataflowOutput* Out) const override;
 };
\ No newline at end of file
diff --git a/Engine/Plugins/Experimental/ChaosFlesh/Source/ChaosFleshNodes/Public/Dataflow/ChaosFleshComputeIslandsNode.h b/Engine/Plugins/Experimental/ChaosFlesh/Source/ChaosFleshNodes/Public/Dataflow/ChaosFleshComputeIslandsNode.h
--- a/Engine/Plugins/Experimental/ChaosFlesh/Source/ChaosFleshNodes/Public/Dataflow/ChaosFleshComputeIslandsNode.h
+++ b/Engine/Plugins/Experimental/ChaosFlesh/Source/ChaosFleshNodes/Public/Dataflow/ChaosFleshComputeIslandsNode.h
@@ -21,12 +21,12 @@ public:
 	UPROPERTY(meta = (DataflowInput, DataflowOutput, DisplayName = "Collection", DataflowPassthrough = "Collection"))
 	FManagedArrayCollection Collection;
 
-	FComputeIslandsNode(const Dataflow::FNodeParameters& InParam, FGuid InGuid = FGuid::NewGuid())
+	FComputeIslandsNode(const UE::Dataflow::FNodeParameters& InParam, FGuid InGuid = FGuid::NewGuid())
 		: FDataflowNode(InParam, InGuid)
 	{
 		RegisterInputConnection(&Collection);
 		RegisterOutputConnection(&Collection, &Collection);
 	}
 
-	virtual void Evaluate(Dataflow::FContext& Context, const FDataflowOutput* Out) const override;
+	virtual void Evaluate(UE::Dataflow::FContext& Context, const FDataflowOutput* Out) const override;
 };
diff --git a/Engine/Plugins/Experimental/ChaosFlesh/Source/ChaosFleshNodes/Public/Dataflow/ChaosFleshComputeMuscleActivationNode.h b/Engine/Plugins/Experimental/ChaosFlesh/Source/ChaosFleshNodes/Public/Dataflow/ChaosFleshComputeMuscleActivationNode.h
--- a/Engine/Plugins/Experimental/ChaosFlesh/Source/ChaosFleshNodes/Public/Dataflow/ChaosFleshComputeMuscleActivationNode.h
+++ b/Engine/Plugins/Experimental/ChaosFlesh/Source/ChaosFleshNodes/Public/Dataflow/ChaosFleshComputeMuscleActivationNode.h
@@ -31,7 +31,7 @@ public:
 	UPROPERTY(EditAnywhere, Category = "Dataflow")
 		float ContractionVolumeScale = 1.f;
 
-	FComputeMuscleActivationDataNode(const Dataflow::FNodeParameters& InParam, FGuid InGuid = FGuid::NewGuid())
+	FComputeMuscleActivationDataNode(const UE::Dataflow::FNodeParameters& InParam, FGuid InGuid = FGuid::NewGuid())
 		: FDataflowNode(InParam, InGuid)
 	{
 		RegisterInputConnection(&Collection);
@@ -40,5 +40,5 @@ public:
 		RegisterOutputConnection(&Collection, &Collection);
 	}
 
-	virtual void Evaluate(Dataflow::FContext& Context, const FDataflowOutput* Out) const override;
+	virtual void Evaluate(UE::Dataflow::FContext& Context, const FDataflowOutput* Out) const override;
 };
diff --git a/Engine/Plugins/Experimental/ChaosFlesh/Source/ChaosFleshNodes/Public/Dataflow/ChaosFleshCreateTetrahedronNode.h b/Engine/Plugins/Experimental/ChaosFlesh/Source/ChaosFleshNodes/Public/Dataflow/ChaosFleshCreateTetrahedronNode.h
--- a/Engine/Plugins/Experimental/ChaosFlesh/Source/ChaosFleshNodes/Public/Dataflow/ChaosFleshCreateTetrahedronNode.h
+++ b/Engine/Plugins/Experimental/ChaosFlesh/Source/ChaosFleshNodes/Public/Dataflow/ChaosFleshCreateTetrahedronNode.h
@@ -100,7 +100,7 @@ public:
 	UPROPERTY(meta = (DataflowInput, DisplayName = "SourceCollection"))
 	FManagedArrayCollection SourceCollection;
 
-	FCreateTetrahedronDataflowNode(const Dataflow::FNodeParameters& InParam, FGuid InGuid = FGuid::NewGuid())
+	FCreateTetrahedronDataflowNode(const UE::Dataflow::FNodeParameters& InParam, FGuid InGuid = FGuid::NewGuid())
 		: FDataflowNode(InParam, InGuid)
 	{
 		RegisterInputConnection(&Collection);
@@ -108,11 +108,11 @@ public:
 		RegisterInputConnection(&SourceCollection);
 	}
 
-	virtual void Evaluate(Dataflow::FContext& Context, const FDataflowOutput* Out) const override;
+	virtual void Evaluate(UE::Dataflow::FContext& Context, const FDataflowOutput* Out) const override;
 
 protected:
-	void EvaluateIsoStuffing(Dataflow::FContext& Context, FFleshCollection& InCollection, const UE::Geometry::FDynamicMesh3& DynamicMesh) const;
-	void EvaluateTetWild(Dataflow::FContext& Context, FFleshCollection& InCollection, const UE::Geometry::FDynamicMesh3& DynamicMesh) const;
+	void EvaluateIsoStuffing(UE::Dataflow::FContext& Context, FFleshCollection& InCollection, const UE::Geometry::FDynamicMesh3& DynamicMesh) const;
+	void EvaluateTetWild(UE::Dataflow::FContext& Context, FFleshCollection& InCollection, const UE::Geometry::FDynamicMesh3& DynamicMesh) const;
 };
 
 
diff --git a/Engine/Plugins/Experimental/ChaosFlesh/Source/ChaosFleshNodes/Public/Dataflow/ChaosFleshFleshAssetTerminalNode.h b/Engine/Plugins/Experimental/ChaosFlesh/Source/ChaosFleshNodes/Public/Dataflow/ChaosFleshFleshAssetTerminalNode.h
--- a/Engine/Plugins/Experimental/ChaosFlesh/Source/ChaosFleshNodes/Public/Dataflow/ChaosFleshFleshAssetTerminalNode.h
+++ b/Engine/Plugins/Experimental/ChaosFlesh/Source/ChaosFleshNodes/Public/Dataflow/ChaosFleshFleshAssetTerminalNode.h
@@ -22,7 +22,7 @@ public:
 	FManagedArrayCollection Collection;
 
 
-	FFleshAssetTerminalDataflowNode(const Dataflow::FNodeParameters& InParam, FGuid InGuid = FGuid::NewGuid())
+	FFleshAssetTerminalDataflowNode(const UE::Dataflow::FNodeParameters& InParam, FGuid InGuid = FGuid::NewGuid())
 	: FDataflowTerminalNode(InParam, InGuid)
 	{
 		RegisterInputConnection(&Collection);
@@ -35,8 +35,8 @@ public:
 	/** Return the terminal asset */
 	virtual TObjectPtr<UObject> GetTerminalAsset() const override {return FleshAsset;}
 
-	virtual void SetAssetValue(TObjectPtr<UObject> Asset, Dataflow::FContext& Context) const override;
-	virtual void Evaluate(Dataflow::FContext& Context, const FDataflowOutput* Out) const override;
+	virtual void SetAssetValue(TObjectPtr<UObject> Asset, UE::Dataflow::FContext& Context) const override;
+	virtual void Evaluate(UE::Dataflow::FContext& Context, const FDataflowOutput* Out) const override;
 };
 
 
diff --git a/Engine/Plugins/Experimental/ChaosFlesh/Source/ChaosFleshNodes/Public/Dataflow/ChaosFleshGenerateOriginInsertionNode.h b/Engine/Plugins/Experimental/ChaosFlesh/Source/ChaosFleshNodes/Public/Dataflow/ChaosFleshGenerateOriginInsertionNode.h
--- a/Engine/Plugins/Experimental/ChaosFlesh/Source/ChaosFleshNodes/Public/Dataflow/ChaosFleshGenerateOriginInsertionNode.h
+++ b/Engine/Plugins/Experimental/ChaosFlesh/Source/ChaosFleshNodes/Public/Dataflow/ChaosFleshGenerateOriginInsertionNode.h
@@ -37,7 +37,7 @@ public:
 	UPROPERTY(EditAnywhere, Category = "Dataflow")
 	float Radius = float(1);
 
-	FGenerateOriginInsertionNode(const Dataflow::FNodeParameters& InParam, FGuid InGuid = FGuid::NewGuid())
+	FGenerateOriginInsertionNode(const UE::Dataflow::FNodeParameters& InParam, FGuid InGuid = FGuid::NewGuid())
 		: FDataflowNode(InParam, InGuid)
 	{
 		RegisterInputConnection(&Collection);
@@ -48,5 +48,5 @@ public:
 		RegisterOutputConnection(&InsertionIndicesOut);
 	}
 
-	virtual void Evaluate(Dataflow::FContext& Context, const FDataflowOutput* Out) const override;
+	virtual void Evaluate(UE::Dataflow::FContext& Context, const FDataflowOutput* Out) const override;
 };
diff --git a/Engine/Plugins/Experimental/ChaosFlesh/Source/ChaosFleshNodes/Public/Dataflow/ChaosFleshGenerateSkeletalBindingsNode.h b/Engine/Plugins/Experimental/ChaosFlesh/Source/ChaosFleshNodes/Public/Dataflow/ChaosFleshGenerateSkeletalBindingsNode.h
--- a/Engine/Plugins/Experimental/ChaosFlesh/Source/ChaosFleshNodes/Public/Dataflow/ChaosFleshGenerateSkeletalBindingsNode.h
+++ b/Engine/Plugins/Experimental/ChaosFlesh/Source/ChaosFleshNodes/Public/Dataflow/ChaosFleshGenerateSkeletalBindingsNode.h
@@ -35,7 +35,7 @@ public:
 	UPROPERTY(EditAnywhere, Category = "Dataflow", meta = (DataflowInput, DisplayName = "SkeletalMesh"))
 	TObjectPtr<const USkeletalMesh> SkeletalMeshIn = nullptr;
 
-	FGenerateSkeletalBindings(const Dataflow::FNodeParameters& InParam, FGuid InGuid = FGuid::NewGuid())
+	FGenerateSkeletalBindings(const UE::Dataflow::FNodeParameters& InParam, FGuid InGuid = FGuid::NewGuid())
 		: FDataflowNode(InParam, InGuid)
 	{
 		RegisterInputConnection(&Collection);
@@ -44,5 +44,5 @@ public:
 		RegisterInputConnection(&SkeletalMeshIn);
 	}
 
-	virtual void Evaluate(Dataflow::FContext& Context, const FDataflowOutput* Out) const override;
+	virtual void Evaluate(UE::Dataflow::FContext& Context, const FDataflowOutput* Out) const override;
 };
\ No newline at end of file
diff --git a/Engine/Plugins/Experimental/ChaosFlesh/Source/ChaosFleshNodes/Public/Dataflow/ChaosFleshGenerateSurfaceBindingsNode.h b/Engine/Plugins/Experimental/ChaosFlesh/Source/ChaosFleshNodes/Public/Dataflow/ChaosFleshGenerateSurfaceBindingsNode.h
--- a/Engine/Plugins/Experimental/ChaosFlesh/Source/ChaosFleshNodes/Public/Dataflow/ChaosFleshGenerateSurfaceBindingsNode.h
+++ b/Engine/Plugins/Experimental/ChaosFlesh/Source/ChaosFleshNodes/Public/Dataflow/ChaosFleshGenerateSurfaceBindingsNode.h
@@ -60,7 +60,7 @@ public:
 	UPROPERTY(EditAnywhere, Category = "Dataflow", meta = (DisplayName = "DoOrphanReparenting"))
 	bool bDoOrphanReparenting = true;
 
-	FGenerateSurfaceBindings(const Dataflow::FNodeParameters& InParam, FGuid InGuid = FGuid::NewGuid())
+	FGenerateSurfaceBindings(const UE::Dataflow::FNodeParameters& InParam, FGuid InGuid = FGuid::NewGuid())
 		: FDataflowNode(InParam, InGuid)
 	{
 		RegisterInputConnection(&Collection);
@@ -71,5 +71,5 @@ public:
 		RegisterInputConnection(&GeometryGroupGuidsIn);
 	}
 
-	virtual void Evaluate(Dataflow::FContext& Context, const FDataflowOutput* Out) const override;
+	virtual void Evaluate(UE::Dataflow::FContext& Context, const FDataflowOutput* Out) const override;
 };
\ No newline at end of file
diff --git a/Engine/Plugins/Experimental/ChaosFlesh/Source/ChaosFleshNodes/Public/Dataflow/ChaosFleshGetFleshAssetNode.h b/Engine/Plugins/Experimental/ChaosFlesh/Source/ChaosFleshNodes/Public/Dataflow/ChaosFleshGetFleshAssetNode.h
--- a/Engine/Plugins/Experimental/ChaosFlesh/Source/ChaosFleshNodes/Public/Dataflow/ChaosFleshGetFleshAssetNode.h
+++ b/Engine/Plugins/Experimental/ChaosFlesh/Source/ChaosFleshNodes/Public/Dataflow/ChaosFleshGetFleshAssetNode.h
@@ -24,11 +24,11 @@ public:
 	UPROPERTY(meta = (DataflowOutput, DisplayName = "Collection"))
 	FManagedArrayCollection Output;
 
-	FGetFleshAssetDataflowNode(const Dataflow::FNodeParameters& InParam, FGuid InGuid = FGuid::NewGuid())
+	FGetFleshAssetDataflowNode(const UE::Dataflow::FNodeParameters& InParam, FGuid InGuid = FGuid::NewGuid())
 		: FDataflowNode(InParam, InGuid)
 	{
 		RegisterOutputConnection(&Output);
 	}
 
-	virtual void Evaluate(Dataflow::FContext& Context, const FDataflowOutput* Out) const override;
+	virtual void Evaluate(UE::Dataflow::FContext& Context, const FDataflowOutput* Out) const override;
 };
diff --git a/Engine/Plugins/Experimental/ChaosFlesh/Source/ChaosFleshNodes/Public/Dataflow/ChaosFleshGetSurfaceIndicesNode.h b/Engine/Plugins/Experimental/ChaosFlesh/Source/ChaosFleshNodes/Public/Dataflow/ChaosFleshGetSurfaceIndicesNode.h
--- a/Engine/Plugins/Experimental/ChaosFlesh/Source/ChaosFleshNodes/Public/Dataflow/ChaosFleshGetSurfaceIndicesNode.h
+++ b/Engine/Plugins/Experimental/ChaosFlesh/Source/ChaosFleshNodes/Public/Dataflow/ChaosFleshGetSurfaceIndicesNode.h
@@ -25,7 +25,7 @@ public:
 	UPROPERTY(meta = (DataflowOutput, DisplayName = "SurfaceVertexSelection"))
 	FDataflowVertexSelection SurfaceVertexSelection;
 
-	FGetSurfaceIndicesNode(const Dataflow::FNodeParameters& InParam, FGuid InGuid = FGuid::NewGuid())
+	FGetSurfaceIndicesNode(const UE::Dataflow::FNodeParameters& InParam, FGuid InGuid = FGuid::NewGuid())
 		: FDataflowNode(InParam, InGuid)
 	{
 		RegisterInputConnection(&Collection);
@@ -33,10 +33,10 @@ public:
 		RegisterOutputConnection(&SurfaceVertexSelection);
 	}
 
-	virtual void Evaluate(Dataflow::FContext& Context, const FDataflowOutput* Out) const override;
+	virtual void Evaluate(UE::Dataflow::FContext& Context, const FDataflowOutput* Out) const override;
 };
 
-namespace Dataflow
+namespace UE::Dataflow
 {
 	void RegisterChaosFleshEngineAssetNodes();
 }
diff --git a/Engine/Plugins/Experimental/ChaosFlesh/Source/ChaosFleshNodes/Public/Dataflow/ChaosFleshIsolateComponentNode.h b/Engine/Plugins/Experimental/ChaosFlesh/Source/ChaosFleshNodes/Public/Dataflow/ChaosFleshIsolateComponentNode.h
--- a/Engine/Plugins/Experimental/ChaosFlesh/Source/ChaosFleshNodes/Public/Dataflow/ChaosFleshIsolateComponentNode.h
+++ b/Engine/Plugins/Experimental/ChaosFlesh/Source/ChaosFleshNodes/Public/Dataflow/ChaosFleshIsolateComponentNode.h
@@ -28,12 +28,12 @@ public:
 	UPROPERTY(EditAnywhere, Category = "Dataflow")
 	FString TargetGeometryIndex = "";
 
-	FIsolateComponentNode(const Dataflow::FNodeParameters& InParam, FGuid InGuid = FGuid::NewGuid())
+	FIsolateComponentNode(const UE::Dataflow::FNodeParameters& InParam, FGuid InGuid = FGuid::NewGuid())
 		: FDataflowNode(InParam, InGuid)
 	{
 		RegisterInputConnection(&Collection);
 		RegisterOutputConnection(&Collection, &Collection);
 	}
 
-	virtual void Evaluate(Dataflow::FContext& Context, const FDataflowOutput* Out) const override;
+	virtual void Evaluate(UE::Dataflow::FContext& Context, const FDataflowOutput* Out) const override;
 };
diff --git a/Engine/Plugins/Experimental/ChaosFlesh/Source/ChaosFleshNodes/Public/Dataflow/ChaosFleshKinematicConstraintNode.h b/Engine/Plugins/Experimental/ChaosFlesh/Source/ChaosFleshNodes/Public/Dataflow/ChaosFleshKinematicConstraintNode.h
--- a/Engine/Plugins/Experimental/ChaosFlesh/Source/ChaosFleshNodes/Public/Dataflow/ChaosFleshKinematicConstraintNode.h
+++ b/Engine/Plugins/Experimental/ChaosFlesh/Source/ChaosFleshNodes/Public/Dataflow/ChaosFleshKinematicConstraintNode.h
@@ -40,7 +40,7 @@ public:
 	int32 BoneIndexIn = INDEX_NONE;
 	
 
-	FKinematicInitializationDataflowNode(const Dataflow::FNodeParameters& InParam, FGuid InGuid = FGuid::NewGuid())
+	FKinematicInitializationDataflowNode(const UE::Dataflow::FNodeParameters& InParam, FGuid InGuid = FGuid::NewGuid())
 		: FDataflowNode(InParam, InGuid)
 	{
 		RegisterInputConnection(&Collection);
@@ -50,5 +50,5 @@ public:
 		RegisterInputConnection(&BoneIndexIn);
 	}
 
-	virtual void Evaluate(Dataflow::FContext& Context, const FDataflowOutput* Out) const override;
+	virtual void Evaluate(UE::Dataflow::FContext& Context, const FDataflowOutput* Out) const override;
 };
diff --git a/Engine/Plugins/Experimental/ChaosFlesh/Source/ChaosFleshNodes/Public/Dataflow/ChaosFleshKinematicMuscleAttachmentsNode.h b/Engine/Plugins/Experimental/ChaosFlesh/Source/ChaosFleshNodes/Public/Dataflow/ChaosFleshKinematicMuscleAttachmentsNode.h
--- a/Engine/Plugins/Experimental/ChaosFlesh/Source/ChaosFleshNodes/Public/Dataflow/ChaosFleshKinematicMuscleAttachmentsNode.h
+++ b/Engine/Plugins/Experimental/ChaosFlesh/Source/ChaosFleshNodes/Public/Dataflow/ChaosFleshKinematicMuscleAttachmentsNode.h
@@ -32,7 +32,7 @@ public:
 		TObjectPtr<USkeletalMesh> SkeletalMeshIn;
 
 
-	FKinematicMuscleAttachmentsDataflowNode(const Dataflow::FNodeParameters& InParam, FGuid InGuid = FGuid::NewGuid())
+	FKinematicMuscleAttachmentsDataflowNode(const UE::Dataflow::FNodeParameters& InParam, FGuid InGuid = FGuid::NewGuid())
 		: FDataflowNode(InParam, InGuid)
 	{
 		RegisterInputConnection(&Collection);
@@ -42,6 +42,6 @@ public:
 		RegisterInputConnection(&SkeletalMeshIn);
 	}
 
-	virtual void Evaluate(Dataflow::FContext& Context, const FDataflowOutput* Out) const override;
+	virtual void Evaluate(UE::Dataflow::FContext& Context, const FDataflowOutput* Out) const override;
 };
 
diff --git a/Engine/Plugins/Experimental/ChaosFlesh/Source/ChaosFleshNodes/Public/Dataflow/ChaosFleshKinematicSkeletonConstraintNode.h b/Engine/Plugins/Experimental/ChaosFlesh/Source/ChaosFleshNodes/Public/Dataflow/ChaosFleshKinematicSkeletonConstraintNode.h
--- a/Engine/Plugins/Experimental/ChaosFlesh/Source/ChaosFleshNodes/Public/Dataflow/ChaosFleshKinematicSkeletonConstraintNode.h
+++ b/Engine/Plugins/Experimental/ChaosFlesh/Source/ChaosFleshNodes/Public/Dataflow/ChaosFleshKinematicSkeletonConstraintNode.h
@@ -32,7 +32,7 @@ public:
 	UPROPERTY(EditAnywhere, Category = "Dataflow")
 	FString ExclusionList = "";
 
-	FKinematicSkeletonConstraintDataflowNode(const Dataflow::FNodeParameters& InParam, FGuid InGuid = FGuid::NewGuid())
+	FKinematicSkeletonConstraintDataflowNode(const UE::Dataflow::FNodeParameters& InParam, FGuid InGuid = FGuid::NewGuid())
 		: FDataflowNode(InParam, InGuid)
 	{
 		RegisterInputConnection(&Collection);
@@ -40,6 +40,6 @@ public:
 		RegisterInputConnection(&SkeletonIn);
 	}
 
-	virtual void Evaluate(Dataflow::FContext& Context, const FDataflowOutput* Out) const override;
+	virtual void Evaluate(UE::Dataflow::FContext& Context, const FDataflowOutput* Out) const override;
 };
 
diff --git a/Engine/Plugins/Experimental/ChaosFlesh/Source/ChaosFleshNodes/Public/Dataflow/ChaosFleshNodesUtility.h b/Engine/Plugins/Experimental/ChaosFlesh/Source/ChaosFleshNodes/Public/Dataflow/ChaosFleshNodesUtility.h
--- a/Engine/Plugins/Experimental/ChaosFlesh/Source/ChaosFleshNodes/Public/Dataflow/ChaosFleshNodesUtility.h
+++ b/Engine/Plugins/Experimental/ChaosFlesh/Source/ChaosFleshNodes/Public/Dataflow/ChaosFleshNodesUtility.h
@@ -24,7 +24,7 @@ enum class ESkeletalSeletionMode : uint8
 	Chaos_Max UMETA(Hidden)
 };
 
-namespace Dataflow
+namespace UE::Dataflow
 {
 	TArray<FIntVector3> CHAOSFLESHNODES_API GetSurfaceTriangles(const TArray<FIntVector4>& Tets, const bool bKeepInterior);
 
diff --git a/Engine/Plugins/Experimental/ChaosFlesh/Source/ChaosFleshNodes/Public/Dataflow/ChaosFleshRadialTetrahedronNode.h b/Engine/Plugins/Experimental/ChaosFlesh/Source/ChaosFleshNodes/Public/Dataflow/ChaosFleshRadialTetrahedronNode.h
--- a/Engine/Plugins/Experimental/ChaosFlesh/Source/ChaosFleshNodes/Public/Dataflow/ChaosFleshRadialTetrahedronNode.h
+++ b/Engine/Plugins/Experimental/ChaosFlesh/Source/ChaosFleshNodes/Public/Dataflow/ChaosFleshRadialTetrahedronNode.h
@@ -49,14 +49,14 @@ public:
 	UPROPERTY(EditAnywhere, Category = "Dataflow")
 	bool bDiscardInteriorTriangles = true;
 
-	FRadialTetrahedronDataflowNodes(const Dataflow::FNodeParameters& InParam, FGuid InGuid = FGuid::NewGuid())
+	FRadialTetrahedronDataflowNodes(const UE::Dataflow::FNodeParameters& InParam, FGuid InGuid = FGuid::NewGuid())
 		: FDataflowNode(InParam, InGuid)
 	{
 		RegisterInputConnection(&Collection);
 		RegisterOutputConnection(&Collection);
 	}
 
-	virtual void Evaluate(Dataflow::FContext& Context, const FDataflowOutput* Out) const override;
+	virtual void Evaluate(UE::Dataflow::FContext& Context, const FDataflowOutput* Out) const override;
 };
 
 
diff --git a/Engine/Plugins/Experimental/ChaosFlesh/Source/ChaosFleshNodes/Public/Dataflow/ChaosFleshSetFleshBonePositionTargetBindingNode.h b/Engine/Plugins/Experimental/ChaosFlesh/Source/ChaosFleshNodes/Public/Dataflow/ChaosFleshSetFleshBonePositionTargetBindingNode.h
--- a/Engine/Plugins/Experimental/ChaosFlesh/Source/ChaosFleshNodes/Public/Dataflow/ChaosFleshSetFleshBonePositionTargetBindingNode.h
+++ b/Engine/Plugins/Experimental/ChaosFlesh/Source/ChaosFleshNodes/Public/Dataflow/ChaosFleshSetFleshBonePositionTargetBindingNode.h
@@ -46,7 +46,7 @@ public:
 		float VertexRadiusRatio = .001f;
 
 
-	FSetFleshBonePositionTargetBindingDataflowNode(const Dataflow::FNodeParameters& InParam, FGuid InGuid = FGuid::NewGuid())
+	FSetFleshBonePositionTargetBindingDataflowNode(const UE::Dataflow::FNodeParameters& InParam, FGuid InGuid = FGuid::NewGuid())
 		: FDataflowNode(InParam, InGuid)
 	{
 		RegisterInputConnection(&Collection);
@@ -54,10 +54,10 @@ public:
 		RegisterOutputConnection(&Collection, &Collection);
 	}
 
-	virtual void Evaluate(Dataflow::FContext& Context, const FDataflowOutput* Out) const override;
+	virtual void Evaluate(UE::Dataflow::FContext& Context, const FDataflowOutput* Out) const override;
 };
 
-namespace Dataflow
+namespace UE::Dataflow
 {
 	void RegisterChaosFleshPositionTargetInitializationNodes();
 }
diff --git a/Engine/Plugins/Experimental/ChaosFlesh/Source/ChaosFleshNodes/Public/Dataflow/ChaosFleshSetFleshDefaultPropertiesNode.h b/Engine/Plugins/Experimental/ChaosFlesh/Source/ChaosFleshNodes/Public/Dataflow/ChaosFleshSetFleshDefaultPropertiesNode.h
--- a/Engine/Plugins/Experimental/ChaosFlesh/Source/ChaosFleshNodes/Public/Dataflow/ChaosFleshSetFleshDefaultPropertiesNode.h
+++ b/Engine/Plugins/Experimental/ChaosFlesh/Source/ChaosFleshNodes/Public/Dataflow/ChaosFleshSetFleshDefaultPropertiesNode.h
@@ -41,12 +41,12 @@ public:
 	UPROPERTY(meta = (DataflowInput, DataflowOutput, DisplayName = "Collection", DataflowPassthrough  = "Collection"))
 	FManagedArrayCollection Collection;
 
-	FSetFleshDefaultPropertiesNode(const Dataflow::FNodeParameters& InParam, FGuid InGuid = FGuid::NewGuid())
+	FSetFleshDefaultPropertiesNode(const UE::Dataflow::FNodeParameters& InParam, FGuid InGuid = FGuid::NewGuid())
 		: FDataflowNode(InParam, InGuid)
 	{
 		RegisterInputConnection(&Collection);
 		RegisterOutputConnection(&Collection, &Collection);
 	}
 
-	virtual void Evaluate(Dataflow::FContext& Context, const FDataflowOutput* Out) const override;
+	virtual void Evaluate(UE::Dataflow::FContext& Context, const FDataflowOutput* Out) const override;
 };
diff --git a/Engine/Plugins/Experimental/ChaosFlesh/Source/ChaosFleshNodes/Public/Dataflow/ChaosFleshSetVertexTetrahedraPositionTargetBindingNode.h b/Engine/Plugins/Experimental/ChaosFlesh/Source/ChaosFleshNodes/Public/Dataflow/ChaosFleshSetVertexTetrahedraPositionTargetBindingNode.h
--- a/Engine/Plugins/Experimental/ChaosFlesh/Source/ChaosFleshNodes/Public/Dataflow/ChaosFleshSetVertexTetrahedraPositionTargetBindingNode.h
+++ b/Engine/Plugins/Experimental/ChaosFlesh/Source/ChaosFleshNodes/Public/Dataflow/ChaosFleshSetVertexTetrahedraPositionTargetBindingNode.h
@@ -33,7 +33,7 @@ public:
 
 
 
-	FSetVertexTetrahedraPositionTargetBindingDataflowNode(const Dataflow::FNodeParameters& InParam, FGuid InGuid = FGuid::NewGuid())
+	FSetVertexTetrahedraPositionTargetBindingDataflowNode(const UE::Dataflow::FNodeParameters& InParam, FGuid InGuid = FGuid::NewGuid())
 		: FDataflowNode(InParam, InGuid)
 	{
 		RegisterInputConnection(&Collection);
@@ -42,5 +42,5 @@ public:
 		RegisterInputConnection(&GeometryGroupGuidsIn);
 	}
 
-	virtual void Evaluate(Dataflow::FContext& Context, const FDataflowOutput* Out) const override;
+	virtual void Evaluate(UE::Dataflow::FContext& Context, const FDataflowOutput* Out) const override;
 };
diff --git a/Engine/Plugins/Experimental/ChaosFlesh/Source/ChaosFleshNodes/Public/Dataflow/ChaosFleshSetVertexTrianglePositionTargetBindingNode.h b/Engine/Plugins/Experimental/ChaosFlesh/Source/ChaosFleshNodes/Public/Dataflow/ChaosFleshSetVertexTrianglePositionTargetBindingNode.h
--- a/Engine/Plugins/Experimental/ChaosFlesh/Source/ChaosFleshNodes/Public/Dataflow/ChaosFleshSetVertexTrianglePositionTargetBindingNode.h
+++ b/Engine/Plugins/Experimental/ChaosFlesh/Source/ChaosFleshNodes/Public/Dataflow/ChaosFleshSetVertexTrianglePositionTargetBindingNode.h
@@ -33,7 +33,7 @@ public:
 		float VertexRadiusRatio = .001f;
 
 
-	FSetVertexTrianglePositionTargetBindingDataflowNode(const Dataflow::FNodeParameters& InParam, FGuid InGuid = FGuid::NewGuid())
+	FSetVertexTrianglePositionTargetBindingDataflowNode(const UE::Dataflow::FNodeParameters& InParam, FGuid InGuid = FGuid::NewGuid())
 		: FDataflowNode(InParam, InGuid)
 	{
 		RegisterInputConnection(&Collection);
@@ -41,5 +41,5 @@ public:
 		RegisterOutputConnection(&Collection, &Collection);
 	}
 
-	virtual void Evaluate(Dataflow::FContext& Context, const FDataflowOutput* Out) const override;
+	virtual void Evaluate(UE::Dataflow::FContext& Context, const FDataflowOutput* Out) const override;
 };
diff --git a/Engine/Plugins/Experimental/ChaosFlesh/Source/ChaosFleshNodes/Public/Dataflow/ChaosFleshSetVertexVertexPositionTargetBindingNode.h b/Engine/Plugins/Experimental/ChaosFlesh/Source/ChaosFleshNodes/Public/Dataflow/ChaosFleshSetVertexVertexPositionTargetBindingNode.h
--- a/Engine/Plugins/Experimental/ChaosFlesh/Source/ChaosFleshNodes/Public/Dataflow/ChaosFleshSetVertexVertexPositionTargetBindingNode.h
+++ b/Engine/Plugins/Experimental/ChaosFlesh/Source/ChaosFleshNodes/Public/Dataflow/ChaosFleshSetVertexVertexPositionTargetBindingNode.h
@@ -33,7 +33,7 @@ public:
 
 
 
-	FSetVertexVertexPositionTargetBindingDataflowNode(const Dataflow::FNodeParameters& InParam, FGuid InGuid = FGuid::NewGuid())
+	FSetVertexVertexPositionTargetBindingDataflowNode(const UE::Dataflow::FNodeParameters& InParam, FGuid InGuid = FGuid::NewGuid())
 		: FDataflowNode(InParam, InGuid)
 	{
 		RegisterInputConnection(&Collection);
@@ -41,5 +41,5 @@ public:
 		RegisterInputConnection(&TargetIndicesIn);
 	}
 
-	virtual void Evaluate(Dataflow::FContext& Context, const FDataflowOutput* Out) const override;
+	virtual void Evaluate(UE::Dataflow::FContext& Context, const FDataflowOutput* Out) const override;
 };
diff --git a/Engine/Plugins/Experimental/ChaosFlesh/Source/ChaosFleshNodes/Public/Dataflow/ChaosFleshSkeletalMeshConstraintNode.h b/Engine/Plugins/Experimental/ChaosFlesh/Source/ChaosFleshNodes/Public/Dataflow/ChaosFleshSkeletalMeshConstraintNode.h
--- a/Engine/Plugins/Experimental/ChaosFlesh/Source/ChaosFleshNodes/Public/Dataflow/ChaosFleshSkeletalMeshConstraintNode.h
+++ b/Engine/Plugins/Experimental/ChaosFlesh/Source/ChaosFleshNodes/Public/Dataflow/ChaosFleshSkeletalMeshConstraintNode.h
@@ -29,7 +29,7 @@ public:
 	UPROPERTY(meta = (DataflowOutput, DisplayName = "SelectionSet"))
 		TArray<int32> IndicesOut;
 
-	FKinematicSkeletalMeshInitializationDataflowNode(const Dataflow::FNodeParameters& InParam, FGuid InGuid = FGuid::NewGuid())
+	FKinematicSkeletalMeshInitializationDataflowNode(const UE::Dataflow::FNodeParameters& InParam, FGuid InGuid = FGuid::NewGuid())
 		: FDataflowNode(InParam, InGuid)
 	{
 		RegisterInputConnection(&Collection);
@@ -38,6 +38,6 @@ public:
 		RegisterOutputConnection(&IndicesOut);
 	}
 
-	virtual void Evaluate(Dataflow::FContext& Context, const FDataflowOutput* Out) const override;
+	virtual void Evaluate(UE::Dataflow::FContext& Context, const FDataflowOutput* Out) const override;
 };
 
diff --git a/Engine/Plugins/Experimental/ChaosFlesh/Source/ChaosFleshNodes/Public/Dataflow/ChaosFleshSkinSimulationPropertiesNode.h b/Engine/Plugins/Experimental/ChaosFlesh/Source/ChaosFleshNodes/Public/Dataflow/ChaosFleshSkinSimulationPropertiesNode.h
--- a/Engine/Plugins/Experimental/ChaosFlesh/Source/ChaosFleshNodes/Public/Dataflow/ChaosFleshSkinSimulationPropertiesNode.h
+++ b/Engine/Plugins/Experimental/ChaosFlesh/Source/ChaosFleshNodes/Public/Dataflow/ChaosFleshSkinSimulationPropertiesNode.h
@@ -30,13 +30,13 @@ public:
 	FManagedArrayCollection Collection;
 
 
-	FSkinSimulationPropertiesDataflowNodes(const Dataflow::FNodeParameters& InParam, FGuid InGuid = FGuid::NewGuid())
+	FSkinSimulationPropertiesDataflowNodes(const UE::Dataflow::FNodeParameters& InParam, FGuid InGuid = FGuid::NewGuid())
 		: FDataflowNode(InParam, InGuid)
 	{
 		RegisterInputConnection(&Collection);
 		RegisterOutputConnection(&Collection);
 	}
 
-	virtual void Evaluate(Dataflow::FContext& Context, const FDataflowOutput* Out) const override;
+	virtual void Evaluate(UE::Dataflow::FContext& Context, const FDataflowOutput* Out) const override;
 
 };
diff --git a/Engine/Plugins/Experimental/ChaosFlesh/Source/ChaosFleshNodes/Public/Dataflow/ChaosFleshTriangleMeshSimulationPropertiesNode.h b/Engine/Plugins/Experimental/ChaosFlesh/Source/ChaosFleshNodes/Public/Dataflow/ChaosFleshTriangleMeshSimulationPropertiesNode.h
--- a/Engine/Plugins/Experimental/ChaosFlesh/Source/ChaosFleshNodes/Public/Dataflow/ChaosFleshTriangleMeshSimulationPropertiesNode.h
+++ b/Engine/Plugins/Experimental/ChaosFlesh/Source/ChaosFleshNodes/Public/Dataflow/ChaosFleshTriangleMeshSimulationPropertiesNode.h
@@ -40,14 +40,14 @@ public:
 	FManagedArrayCollection Collection;
 
 
-	FTriangleMeshSimulationPropertiesDataflowNodes(const Dataflow::FNodeParameters& InParam, FGuid InGuid = FGuid::NewGuid())
+	FTriangleMeshSimulationPropertiesDataflowNodes(const UE::Dataflow::FNodeParameters& InParam, FGuid InGuid = FGuid::NewGuid())
 		: FDataflowNode(InParam, InGuid)
 	{
 		RegisterInputConnection(&Collection);
 		RegisterOutputConnection(&Collection);
 	}
 
-	virtual void Evaluate(Dataflow::FContext& Context, const FDataflowOutput* Out) const override;
+	virtual void Evaluate(UE::Dataflow::FContext& Context, const FDataflowOutput* Out) const override;
 
 
 };
diff --git a/Engine/Plugins/Experimental/ChaosFlesh/Source/ChaosFleshNodes/Public/Dataflow/ChaosFleshVertexConstraintNode.h b/Engine/Plugins/Experimental/ChaosFlesh/Source/ChaosFleshNodes/Public/Dataflow/ChaosFleshVertexConstraintNode.h
--- a/Engine/Plugins/Experimental/ChaosFlesh/Source/ChaosFleshNodes/Public/Dataflow/ChaosFleshVertexConstraintNode.h
+++ b/Engine/Plugins/Experimental/ChaosFlesh/Source/ChaosFleshNodes/Public/Dataflow/ChaosFleshVertexConstraintNode.h
@@ -23,7 +23,7 @@ public:
 	UPROPERTY(meta = (DataflowInput, DisplayName = "SelectionSet"))
 		TArray<int32> VertexIndicesIn;
 
-	FSetVerticesKinematicDataflowNode(const Dataflow::FNodeParameters& InParam, FGuid InGuid = FGuid::NewGuid())
+	FSetVerticesKinematicDataflowNode(const UE::Dataflow::FNodeParameters& InParam, FGuid InGuid = FGuid::NewGuid())
 		: FDataflowNode(InParam, InGuid)
 	{
 		RegisterInputConnection(&Collection);
@@ -31,7 +31,7 @@ public:
 		RegisterInputConnection(&VertexIndicesIn);
 	}
 
-	virtual void Evaluate(Dataflow::FContext& Context, const FDataflowOutput* Out) const override;
+	virtual void Evaluate(UE::Dataflow::FContext& Context, const FDataflowOutput* Out) const override;
 };
 
 
diff --git a/Engine/Plugins/Experimental/ChaosFlesh/Source/ChaosFleshNodes/Public/Dataflow/ChaosFleshVisualizeFiberFieldNode.h b/Engine/Plugins/Experimental/ChaosFlesh/Source/ChaosFleshNodes/Public/Dataflow/ChaosFleshVisualizeFiberFieldNode.h
--- a/Engine/Plugins/Experimental/ChaosFlesh/Source/ChaosFleshNodes/Public/Dataflow/ChaosFleshVisualizeFiberFieldNode.h
+++ b/Engine/Plugins/Experimental/ChaosFlesh/Source/ChaosFleshNodes/Public/Dataflow/ChaosFleshVisualizeFiberFieldNode.h
@@ -29,14 +29,14 @@ public:
 	UPROPERTY(meta = (DataflowOutput, DisplayName = "VectorField"))
 	FFieldCollection VectorField;
 	
-	FVisualizeFiberFieldNode(const Dataflow::FNodeParameters& InParam, FGuid InGuid = FGuid::NewGuid())
+	FVisualizeFiberFieldNode(const UE::Dataflow::FNodeParameters& InParam, FGuid InGuid = FGuid::NewGuid())
 		: FDataflowNode(InParam, InGuid)
 	{
 		RegisterInputConnection(&Collection);
 		RegisterOutputConnection(&VectorField);
 	}
 
-	virtual void Evaluate(Dataflow::FContext& Context, const FDataflowOutput* Out) const override;
+	virtual void Evaluate(UE::Dataflow::FContext& Context, const FDataflowOutput* Out) const override;
 };
 
 /**
@@ -56,14 +56,14 @@ public:
 	UPROPERTY(meta = (DataflowOutput, DisplayName = "VectorField"))
 	FFieldCollection VectorField;
 
-	FVisualizePositionTargetsNode(const Dataflow::FNodeParameters& InParam, FGuid InGuid = FGuid::NewGuid())
+	FVisualizePositionTargetsNode(const UE::Dataflow::FNodeParameters& InParam, FGuid InGuid = FGuid::NewGuid())
 		: FDataflowNode(InParam, InGuid)
 	{
 		RegisterInputConnection(&Collection);
 		RegisterOutputConnection(&VectorField);
 	}
 
-	virtual void Evaluate(Dataflow::FContext& Context, const FDataflowOutput* Out) const override;
+	virtual void Evaluate(UE::Dataflow::FContext& Context, const FDataflowOutput* Out) const override;
 };
 
 /**
@@ -80,12 +80,12 @@ public:
 	UPROPERTY(meta = (DataflowInput, DataflowOutput, DisplayName = "Collection", DataflowPassthrough = "Collection"))
 	FManagedArrayCollection Collection;
 
-	FVisualizeKinematicFacesNode(const Dataflow::FNodeParameters& InParam, FGuid InGuid = FGuid::NewGuid())
+	FVisualizeKinematicFacesNode(const UE::Dataflow::FNodeParameters& InParam, FGuid InGuid = FGuid::NewGuid())
 		: FDataflowNode(InParam, InGuid)
 	{
 		RegisterInputConnection(&Collection);
 		RegisterOutputConnection(&Collection, &Collection);
 	}
 
-	virtual void Evaluate(Dataflow::FContext& Context, const FDataflowOutput* Out) const override;
+	virtual void Evaluate(UE::Dataflow::FContext& Context, const FDataflowOutput* Out) const override;
 };
\ No newline at end of file
diff --git a/Engine/Plugins/Experimental/ChaosFlesh/Source/ChaosFleshNodes/Public/Dataflow/GeometryCollectionAppendCollectionTransformNode.h b/Engine/Plugins/Experimental/ChaosFlesh/Source/ChaosFleshNodes/Public/Dataflow/GeometryCollectionAppendCollectionTransformNode.h
--- a/Engine/Plugins/Experimental/ChaosFlesh/Source/ChaosFleshNodes/Public/Dataflow/GeometryCollectionAppendCollectionTransformNode.h
+++ b/Engine/Plugins/Experimental/ChaosFlesh/Source/ChaosFleshNodes/Public/Dataflow/GeometryCollectionAppendCollectionTransformNode.h
@@ -30,7 +30,7 @@ public:
 	UPROPERTY(EditAnywhere, Category = "Dataflow")
 		FString GroupName = FString("ComponentTransformGroup");
 
-	FAppendToCollectionTransformAttributeDataflowNode(const Dataflow::FNodeParameters& InParam, FGuid InGuid = FGuid::NewGuid())
+	FAppendToCollectionTransformAttributeDataflowNode(const UE::Dataflow::FNodeParameters& InParam, FGuid InGuid = FGuid::NewGuid())
 		: FDataflowNode(InParam, InGuid)
 	{
 		RegisterInputConnection(&TransformIn);
@@ -38,10 +38,10 @@ public:
 		RegisterOutputConnection(&Collection, &Collection);
 	}
 
-	virtual void Evaluate(Dataflow::FContext& Context, const FDataflowOutput* Out) const override;
+	virtual void Evaluate(UE::Dataflow::FContext& Context, const FDataflowOutput* Out) const override;
 };
 
-namespace Dataflow
+namespace UE::Dataflow
 {
 	void RegisterChaosFleshKinematicInitializationNodes();
 }
diff --git a/Engine/Plugins/Experimental/Dataflow/Source/DataflowAssetTools/Private/Dataflow/CollectionRenderingPatternUtility.cpp b/Engine/Plugins/Experimental/Dataflow/Source/DataflowAssetTools/Private/Dataflow/CollectionRenderingPatternUtility.cpp
--- a/Engine/Plugins/Experimental/Dataflow/Source/DataflowAssetTools/Private/Dataflow/CollectionRenderingPatternUtility.cpp
+++ b/Engine/Plugins/Experimental/Dataflow/Source/DataflowAssetTools/Private/Dataflow/CollectionRenderingPatternUtility.cpp
@@ -13,7 +13,7 @@
 
 using namespace UE::Geometry;
 
-namespace Dataflow
+namespace UE::Dataflow
 {
 	namespace Conversion
 	{
@@ -117,4 +117,4 @@ namespace Dataflow
 		}
 	}
 
-}	// namespace Dataflow
\ No newline at end of file
+}	// namespace UE::Dataflow
\ No newline at end of file
diff --git a/Engine/Plugins/Experimental/Dataflow/Source/DataflowAssetTools/Private/Dataflow/DataflowAssetToolsModule.cpp b/Engine/Plugins/Experimental/Dataflow/Source/DataflowAssetTools/Private/Dataflow/DataflowAssetToolsModule.cpp
--- a/Engine/Plugins/Experimental/Dataflow/Source/DataflowAssetTools/Private/Dataflow/DataflowAssetToolsModule.cpp
+++ b/Engine/Plugins/Experimental/Dataflow/Source/DataflowAssetTools/Private/Dataflow/DataflowAssetToolsModule.cpp
@@ -3,7 +3,7 @@
 #include "Features/IModularFeatures.h"
 #include "Modules/ModuleManager.h"
 
-namespace Dataflow
+namespace UE::Dataflow
 {
 	class FDataflowAssetToolsModule : public IModuleInterface
 	{
@@ -19,4 +19,4 @@ namespace Dataflow
 	};
 }
 
-IMPLEMENT_MODULE(Dataflow::FDataflowAssetToolsModule, DataflowAssetTools);
+IMPLEMENT_MODULE(UE::Dataflow::FDataflowAssetToolsModule, DataflowAssetTools);
diff --git a/Engine/Plugins/Experimental/Dataflow/Source/DataflowAssetTools/Public/Dataflow/CollectionRenderingPatternUtility.h b/Engine/Plugins/Experimental/Dataflow/Source/DataflowAssetTools/Public/Dataflow/CollectionRenderingPatternUtility.h
--- a/Engine/Plugins/Experimental/Dataflow/Source/DataflowAssetTools/Public/Dataflow/CollectionRenderingPatternUtility.h
+++ b/Engine/Plugins/Experimental/Dataflow/Source/DataflowAssetTools/Public/Dataflow/CollectionRenderingPatternUtility.h
@@ -10,7 +10,7 @@ namespace GeometryCollection::Facades { class FRenderingFacade; }
 class UDataflow;
 class UObject;
 
-namespace Dataflow
+namespace UE::Dataflow
 {
 	namespace Conversion
 	{
@@ -21,4 +21,4 @@ namespace Dataflow
 		void DATAFLOWASSETTOOLS_API DynamicMeshToRenderingFacade(const UE::Geometry::FDynamicMesh3& DynamicMesh, GeometryCollection::Facades::FRenderingFacade& Facade);
 	}
 
-}	// namespace Dataflow
\ No newline at end of file
+}	// namespace UE::Dataflow
\ No newline at end of file
diff --git a/Engine/Plugins/Experimental/Dataflow/Source/DataflowEditor/Private/Dataflow/AssetDefinition_DataflowAsset.cpp b/Engine/Plugins/Experimental/Dataflow/Source/DataflowEditor/Private/Dataflow/AssetDefinition_DataflowAsset.cpp
--- a/Engine/Plugins/Experimental/Dataflow/Source/DataflowEditor/Private/Dataflow/AssetDefinition_DataflowAsset.cpp
+++ b/Engine/Plugins/Experimental/Dataflow/Source/DataflowEditor/Private/Dataflow/AssetDefinition_DataflowAsset.cpp
@@ -20,7 +20,7 @@
 bool bCanEditDataflow = false;
 FAutoConsoleVariableRef CVarDataflowIsEditable(TEXT("p.Dataflow.IsEditable"), bCanEditDataflow, TEXT("Whether to allow edits of the dataflow [def:true]"));
 
-namespace DataflowAssetDefinitionHelpers
+namespace UE::DataflowAssetDefinitionHelpers
 {
 	// Return true if we should proceed, false if we should re-open the dialog
 	bool CreateNewDataflowAsset(const UObject* Asset, UObject*& OutDataflowAsset)
diff --git a/Engine/Plugins/Experimental/Dataflow/Source/DataflowEditor/Private/Dataflow/DataflowCollectionSpreadSheet.cpp b/Engine/Plugins/Experimental/Dataflow/Source/DataflowEditor/Private/Dataflow/DataflowCollectionSpreadSheet.cpp
--- a/Engine/Plugins/Experimental/Dataflow/Source/DataflowEditor/Private/Dataflow/DataflowCollectionSpreadSheet.cpp
+++ b/Engine/Plugins/Experimental/Dataflow/Source/DataflowEditor/Private/Dataflow/DataflowCollectionSpreadSheet.cpp
@@ -38,7 +38,7 @@ void FDataflowCollectionSpreadSheet::UpdateViewData()
 					TArray<FDataflowOutput*> Outputs = DataflowNode->GetOutputs();
 					if (const TObjectPtr<UDataflowBaseContent> Content = GetEditorContent())
 					{
-						if (TSharedPtr<Dataflow::FEngineContext> Context = Content->GetDataflowContext())
+						if (TSharedPtr<UE::Dataflow::FEngineContext> Context = Content->GetDataflowContext())
 						{
 							for (FDataflowOutput* Output : Outputs)
 							{
diff --git a/Engine/Plugins/Experimental/Dataflow/Source/DataflowEditor/Private/Dataflow/DataflowConstructionScene.cpp b/Engine/Plugins/Experimental/Dataflow/Source/DataflowEditor/Private/Dataflow/DataflowConstructionScene.cpp
--- a/Engine/Plugins/Experimental/Dataflow/Source/DataflowEditor/Private/Dataflow/DataflowConstructionScene.cpp
+++ b/Engine/Plugins/Experimental/Dataflow/Source/DataflowEditor/Private/Dataflow/DataflowConstructionScene.cpp
@@ -83,7 +83,7 @@ void FDataflowConstructionScene::TickDataflowScene(const float DeltaSeconds)
 		{
 			if (Dataflow->GetDataflow())
 			{
-				Dataflow::FTimestamp SystemTimestamp = Dataflow::FTimestamp::Invalid;
+				UE::Dataflow::FTimestamp SystemTimestamp = UE::Dataflow::FTimestamp::Invalid;
 				bool bMustUpdateConstructionScene = false;
 				for (TObjectPtr<const UDataflowBaseContent> DataflowBaseContent : GetTerminalContents())
 				{
@@ -139,7 +139,7 @@ void FDataflowConstructionScene::UpdateDynamicMeshComponents()
 	if (const TObjectPtr<UDataflowBaseContent>& EditorContent = GetEditorContent())
 	{
 		const TObjectPtr<UDataflow>& DataflowAsset = EditorContent->GetDataflowAsset();
-		const TSharedPtr<Dataflow::FEngineContext>& DataflowContext = EditorContent->GetDataflowContext();
+		const TSharedPtr<UE::Dataflow::FEngineContext>& DataflowContext = EditorContent->GetDataflowContext();
 		if(DataflowAsset && DataflowContext)
 		{
 			for (TObjectPtr<const UDataflowEdNode> Target : DataflowAsset->GetRenderTargets())
@@ -150,13 +150,13 @@ void FDataflowConstructionScene::UpdateDynamicMeshComponents()
 					GeometryCollection::Facades::FRenderingFacade Facade(*RenderCollection);
 					Facade.DefineSchema();
 
-					Dataflow::RenderNodeOutput(Facade, *Target, *EditorContent);
+					UE::Dataflow::RenderNodeOutput(Facade, *Target, *EditorContent);
 
 					const int32 NumGeometry = Facade.NumGeometry();
 					for (int32 MeshIndex = 0; MeshIndex < NumGeometry; ++MeshIndex)
 					{
 						FDynamicMesh3 DynamicMesh;
-						Dataflow::Conversion::RenderingFacadeToDynamicMesh(Facade, MeshIndex, DynamicMesh);
+						UE::Dataflow::Conversion::RenderingFacadeToDynamicMesh(Facade, MeshIndex, DynamicMesh);
 
 						if (DynamicMesh.VertexCount())
 						{
@@ -188,7 +188,7 @@ void FDataflowConstructionScene::UpdateDynamicMeshComponents()
 					GeometryCollection::Facades::FRenderingFacade Facade(*RenderCollection);
 					Facade.DefineSchema();
 
-					Dataflow::RenderNodeOutput(Facade, *Target, *EditorContent);
+					UE::Dataflow::RenderNodeOutput(Facade, *Target, *EditorContent);
 
 					const int32 NumGeometry = Facade.NumGeometry();
 					for (int32 MeshIndex = 0; MeshIndex < NumGeometry; ++MeshIndex)
@@ -203,7 +203,7 @@ void FDataflowConstructionScene::UpdateDynamicMeshComponents()
 						else
 						{
 							FDynamicMesh3 DynamicMesh;
-							Dataflow::Conversion::RenderingFacadeToDynamicMesh(Facade, MeshIndex, DynamicMesh);
+							UE::Dataflow::Conversion::RenderingFacadeToDynamicMesh(Facade, MeshIndex, DynamicMesh);
 
 							if (DynamicMesh.VertexCount())
 							{
diff --git a/Engine/Plugins/Experimental/Dataflow/Source/DataflowEditor/Private/Dataflow/DataflowConstructionViewportClient.cpp b/Engine/Plugins/Experimental/Dataflow/Source/DataflowEditor/Private/Dataflow/DataflowConstructionViewportClient.cpp
--- a/Engine/Plugins/Experimental/Dataflow/Source/DataflowEditor/Private/Dataflow/DataflowConstructionViewportClient.cpp
+++ b/Engine/Plugins/Experimental/Dataflow/Source/DataflowEditor/Private/Dataflow/DataflowConstructionViewportClient.cpp
@@ -206,7 +206,7 @@ void FDataflowConstructionViewportClient::ProcessClick(FSceneView& View, HHitPro
 	OnSelectionChangedMulticast.Broadcast(CurrentlySelectedComponents);
 }
 
-void FDataflowConstructionViewportClient::SetConstructionViewMode(const Dataflow::IDataflowConstructionViewMode* InViewMode)
+void FDataflowConstructionViewportClient::SetConstructionViewMode(const UE::Dataflow::IDataflowConstructionViewMode* InViewMode)
 {
 	checkf(InViewMode, TEXT("SetConstructionViewMode received null IDataflowConstructionViewMode pointer"));
 
diff --git a/Engine/Plugins/Experimental/Dataflow/Source/DataflowEditor/Private/Dataflow/DataflowEditor.cpp b/Engine/Plugins/Experimental/Dataflow/Source/DataflowEditor/Private/Dataflow/DataflowEditor.cpp
--- a/Engine/Plugins/Experimental/Dataflow/Source/DataflowEditor/Private/Dataflow/DataflowEditor.cpp
+++ b/Engine/Plugins/Experimental/Dataflow/Source/DataflowEditor/Private/Dataflow/DataflowEditor.cpp
@@ -36,7 +36,7 @@ void UDataflowEditor::Initialize(const TArray<TObjectPtr<UObject>>& InObjects, c
 		{
 			if(UDataflow* DataflowAsset = Cast<UDataflow>(ContentOwner))
 			{
-				EditorContent = DataflowContextHelpers::CreateNewDataflowContent<UDataflowBaseContent>(ContentOwner);
+				EditorContent = UE::DataflowContextHelpers::CreateNewDataflowContent<UDataflowBaseContent>(ContentOwner);
 				EditorContent->SetDataflowOwner(DataflowAsset);
 				EditorContent->SetDataflowAsset(DataflowAsset);
 			}
@@ -69,7 +69,7 @@ void UDataflowEditor::Initialize(const TArray<TObjectPtr<UObject>>& InObjects, c
 		UpdateEditorContent();
 
 		// Update and build the terminal contents
-		UpdateTerminalContents(Dataflow::FTimestamp::Invalid);
+		UpdateTerminalContents(UE::Dataflow::FTimestamp::Invalid);
 		
 		// Potentially we could add additional objects to edit here (fields, meshes....)
 		// If these objects have a matching factory we would be able to use geometry tools
@@ -96,7 +96,7 @@ void UDataflowEditor::UpdateEditorContent()
 	}
 }
 
-void UDataflowEditor::RemoveTerminalContents(const TSharedPtr<Dataflow::FGraph>& DataflowGraph, ValidTerminalsType& ValidTerminals)
+void UDataflowEditor::RemoveTerminalContents(const TSharedPtr<UE::Dataflow::FGraph>& DataflowGraph, ValidTerminalsType& ValidTerminals)
 {
 	for(int32 ContentIndex = TerminalContents.Num()-1; ContentIndex >= 0; --ContentIndex)
 	{
@@ -122,7 +122,7 @@ void UDataflowEditor::RemoveTerminalContents(const TSharedPtr<Dataflow::FGraph>&
 	}
 }
 
-void UDataflowEditor::AddTerminalContents(const TSharedPtr<Dataflow::FGraph>& DataflowGraph, ValidTerminalsType& ValidTerminals)
+void UDataflowEditor::AddTerminalContents(const TSharedPtr<UE::Dataflow::FGraph>& DataflowGraph, ValidTerminalsType& ValidTerminals)
 {
 	for(const TSharedPtr<FDataflowNode>& DataflowNode : DataflowGraph->GetFilteredNodes(FDataflowTerminalNode::StaticType()))
 	{
@@ -151,14 +151,14 @@ void UDataflowEditor::AddTerminalContents(const TSharedPtr<Dataflow::FGraph>& Da
 	}
 }
 
-void UDataflowEditor::UpdateTerminalContents(const Dataflow::FTimestamp TimeStamp)
+void UDataflowEditor::UpdateTerminalContents(const UE::Dataflow::FTimestamp TimeStamp)
 {
 	bHasTerminalsDirty = false;
 	
 	// update of the terminal contents only if no terminal asset on the main editor content
 	if(EditorContent && !EditorContent->GetTerminalAsset() && EditorContent->GetDataflowAsset())
 	{
-		if(const TSharedPtr<Dataflow::FGraph> DataflowGraph = EditorContent->GetDataflowAsset()->GetDataflow())
+		if(const TSharedPtr<UE::Dataflow::FGraph> DataflowGraph = EditorContent->GetDataflowAsset()->GetDataflow())
 		{
 			ValidTerminalsType ValidTerminals;
 
diff --git a/Engine/Plugins/Experimental/Dataflow/Source/DataflowEditor/Private/Dataflow/DataflowEditorCommands.cpp b/Engine/Plugins/Experimental/Dataflow/Source/DataflowEditor/Private/Dataflow/DataflowEditorCommands.cpp
--- a/Engine/Plugins/Experimental/Dataflow/Source/DataflowEditor/Private/Dataflow/DataflowEditorCommands.cpp
+++ b/Engine/Plugins/Experimental/Dataflow/Source/DataflowEditor/Private/Dataflow/DataflowEditorCommands.cpp
@@ -72,11 +72,11 @@ void FDataflowEditorCommandsImpl::RegisterCommands()
 	UI_COMMAND(StartSimulationScene, "StartSimulationScene", "Start the simulation scene", EUserInterfaceActionType::ToggleButton, FInputChord());
 	UI_COMMAND(StepSimulationScene, "StepSimulationScene", "Step the simulation scene", EUserInterfaceActionType::ToggleButton, FInputChord());
 
-	for (const TPair<FName, TUniquePtr<Dataflow::IDataflowConstructionViewMode>>& NameAndMode : Dataflow::FRenderingViewModeFactory::GetInstance().GetViewModes())
+	for (const TPair<FName, TUniquePtr<UE::Dataflow::IDataflowConstructionViewMode>>& NameAndMode : UE::Dataflow::FRenderingViewModeFactory::GetInstance().GetViewModes())
 	{
 		TSharedPtr< FUICommandInfo > SetViewModeCommand;
 		
-		const Dataflow::IDataflowConstructionViewMode* const ViewMode = NameAndMode.Value.Get();
+		const UE::Dataflow::IDataflowConstructionViewMode* const ViewMode = NameAndMode.Value.Get();
 		checkf(ViewMode, TEXT("Registered mode in FRenderingViewModeFactory has no associated IDataflowConstructionViewMode object. Registered name: %s"), *NameAndMode.Key.ToString());
 
 		FUICommandInfo::MakeCommandInfo(
@@ -92,7 +92,7 @@ void FDataflowEditorCommandsImpl::RegisterCommands()
 		SetConstructionViewModeCommands.Add(ViewMode->GetName(), SetViewModeCommand);
 	}
 
-	Dataflow::FDataflowToolRegistry& ToolRegistry = Dataflow::FDataflowToolRegistry::Get();
+	UE::Dataflow::FDataflowToolRegistry& ToolRegistry = UE::Dataflow::FDataflowToolRegistry::Get();
 	const TArray<FName> NodeNames = ToolRegistry.GetNodeNames();
 	for (const FName& NodeName : NodeNames)
 	{
@@ -113,7 +113,7 @@ void FDataflowEditorCommandsImpl::RegisterCommands()
 
 void FDataflowEditorCommandsImpl::UpdateToolCommandBinding(UInteractiveTool* Tool, TSharedPtr<FUICommandList> UICommandList, bool bUnbind)
 {
-	Dataflow::FDataflowToolRegistry& ToolRegistry = Dataflow::FDataflowToolRegistry::Get();
+	UE::Dataflow::FDataflowToolRegistry& ToolRegistry = UE::Dataflow::FDataflowToolRegistry::Get();
 	if (bUnbind)
 	{
 		ToolRegistry.UnbindActiveCommands(UICommandList);
@@ -145,14 +145,14 @@ bool FDataflowEditorCommands::IsRegistered()
 }
 
 
-const FDataflowNode* FDataflowEditorCommands::EvaluateNode(Dataflow::FContext& Context, Dataflow::FTimestamp& InOutLastNodeTimestamp,
+const FDataflowNode* FDataflowEditorCommands::EvaluateNode(UE::Dataflow::FContext& Context, UE::Dataflow::FTimestamp& InOutLastNodeTimestamp,
 	const UDataflow* Dataflow, const FDataflowNode* Node, const FDataflowOutput* Output, const FString& NodeName, UObject* Asset)
 {
 	UE_LOG(LogChaosDataflow, Verbose, TEXT("FDataflowEditorCommands::EvaluateNode(): Node [%s], NodeName [%s] Output [%s]"), Node ? *Node->GetName().ToString() : TEXT("nullptr"), *NodeName, Output ? *Output->GetName().ToString() : TEXT("nullptr"));
 
 	if (!Node && Dataflow)
 	{
-		if (const TSharedPtr<const Dataflow::FGraph> Graph = Dataflow->GetDataflow())
+		if (const TSharedPtr<const UE::Dataflow::FGraph> Graph = Dataflow->GetDataflow())
 		{
 			Node = Graph->FindBaseNode(FName(NodeName)).Get();
 		}
@@ -188,7 +188,7 @@ bool FDataflowEditorCommands::OnNodeVerifyTitleCommit(const FText& NewText, UEdG
 	{
 		if (UDataflowEdNode* DataflowNode = Cast<UDataflowEdNode>(GraphNode))
 		{
-			if (TSharedPtr<Dataflow::FGraph> Graph = DataflowNode->GetDataflowGraph())
+			if (TSharedPtr<UE::Dataflow::FGraph> Graph = DataflowNode->GetDataflowGraph())
 			{
 				if( Graph->FindBaseNode(FName(NewString)).Get()==nullptr )
 				{
@@ -217,7 +217,7 @@ void FDataflowEditorCommands::OnNodeTitleCommitted(const FText& InNewText, EText
 	{
 		if (UDataflowEdNode* DataflowNode = Cast<UDataflowEdNode>(GraphNode))
 		{
-			if (TSharedPtr<Dataflow::FGraph> Graph = DataflowNode->GetDataflowGraph())
+			if (TSharedPtr<UE::Dataflow::FGraph> Graph = DataflowNode->GetDataflowGraph())
 			{
 				if (TSharedPtr<FDataflowNode> Node = Graph->FindBaseNode(DataflowNode->GetDataflowNodeGuid()))
 				{
@@ -296,7 +296,7 @@ void FDataflowEditorCommands::OnAssetPropertyValueChanged(TObjectPtr<UDataflowBa
 	}
 }
 
-void FDataflowEditorCommands::OnPropertyValueChanged(UDataflow* OutDataflow, TSharedPtr<Dataflow::FEngineContext>& Context, Dataflow::FTimestamp& OutLastNodeTimestamp, const FPropertyChangedEvent& InPropertyChangedEvent, const TSet<TObjectPtr<UObject> >& SelectedNodes)
+void FDataflowEditorCommands::OnPropertyValueChanged(UDataflow* OutDataflow, TSharedPtr<UE::Dataflow::FEngineContext>& Context, UE::Dataflow::FTimestamp& OutLastNodeTimestamp, const FPropertyChangedEvent& InPropertyChangedEvent, const TSet<TObjectPtr<UObject> >& SelectedNodes)
 {
 	switch (InPropertyChangedEvent.ChangeType)
 	{
@@ -318,7 +318,7 @@ void FDataflowEditorCommands::OnPropertyValueChanged(UDataflow* OutDataflow, TSh
 					if (TSharedPtr<FDataflowNode> DataflowNode = Node->GetDataflowNode())
 					{
 						DataflowNode->Invalidate();
-						OutLastNodeTimestamp = Dataflow::FTimestamp::Invalid;
+						OutLastNodeTimestamp = UE::Dataflow::FTimestamp::Invalid;
 
 						// Reflect the active state on the drawing of the node
 						if (DataflowNode->bActive != Node->IsNodeEnabled())
@@ -343,7 +343,7 @@ void FDataflowEditorCommands::DeleteNodes(UDataflow* Graph, const FGraphPanelSel
 		{
 			if (UDataflowEdNode* EdNode = dynamic_cast<UDataflowEdNode*>(Node))
 			{
-				if (const TSharedPtr<Dataflow::FGraph> DataflowGraph = EdNode->GetDataflowGraph())
+				if (const TSharedPtr<UE::Dataflow::FGraph> DataflowGraph = EdNode->GetDataflowGraph())
 				{
 					Graph->RemoveNode(EdNode);
 					if (TSharedPtr<FDataflowNode> DataflowNode = DataflowGraph->FindBaseNode(EdNode->GetDataflowNodeGuid()))
@@ -378,7 +378,7 @@ void FDataflowEditorCommands::OnSelectedNodesChanged(TSharedPtr<IStructureDetail
 
 	if (Graph && PropertiesEditor)
 	{
-		if (const TSharedPtr<Dataflow::FGraph> DataflowGraph = Graph->GetDataflow())
+		if (const TSharedPtr<UE::Dataflow::FGraph> DataflowGraph = Graph->GetDataflow())
 		{
 			auto AsRawPointers = [](const TSet<TObjectPtr<UObject> >& NewSelection) {
 				TSet<UObject*> Raw; for (UObject* Elem : NewSelection) Raw.Add(Elem);
@@ -476,7 +476,7 @@ void FDataflowEditorCommands::DuplicateNodes(UDataflow* Graph, const TSharedPtr<
 		}
 
 		FDataflowAssetEdit Edit = Graph->EditDataflow();
-		if (Dataflow::FGraph* DataflowGraph = Edit.GetGraph())
+		if (UE::Dataflow::FGraph* DataflowGraph = Edit.GetGraph())
 		{
 			// Process Dataflow nodes
 			TSet<UDataflowEdNode*> DuplicatedEdNodes;
@@ -675,7 +675,7 @@ void FDataflowEditorCommands::CopyNodes(UDataflow* InGraph, const TSharedPtr<SDa
 			}
 
 			FDataflowAssetEdit Edit = InGraph->EditDataflow();
-			if (Dataflow::FGraph* DataflowGraph = Edit.GetGraph())
+			if (UE::Dataflow::FGraph* DataflowGraph = Edit.GetGraph())
 			{
 				FDataflowCopyPasteContent CopyPasteContent;
 
@@ -915,7 +915,7 @@ void FDataflowEditorCommands::PasteNodes(UDataflow* Graph, const TSharedPtr<SDat
 				const FName OutputputName = *PropertyOut;
 
 				FDataflowAssetEdit Edit = Graph->EditDataflow();
-				if (Dataflow::FGraph* DataflowGraph = Edit.GetGraph())
+				if (UE::Dataflow::FGraph* DataflowGraph = Edit.GetGraph())
 				{
 					if (TSharedPtr<FDataflowNode> DataflowNodeFrom = DataflowGraph->FindBaseNode(GuidOut))
 					{
diff --git a/Engine/Plugins/Experimental/Dataflow/Source/DataflowEditor/Private/Dataflow/DataflowEditorMode.cpp b/Engine/Plugins/Experimental/Dataflow/Source/DataflowEditor/Private/Dataflow/DataflowEditorMode.cpp
--- a/Engine/Plugins/Experimental/Dataflow/Source/DataflowEditor/Private/Dataflow/DataflowEditorMode.cpp
+++ b/Engine/Plugins/Experimental/Dataflow/Source/DataflowEditor/Private/Dataflow/DataflowEditorMode.cpp
@@ -51,7 +51,7 @@
 
 const FEditorModeID UDataflowEditorMode::EM_DataflowEditorModeId = TEXT("EM_DataflowAssetEditorMode");
 
-namespace Dataflow::Private
+namespace UE::Dataflow::Private
 {
 	bool bDataflowEditorEnableToolsInPIE = true;
 	FAutoConsoleVariableRef CVARDataflowEditorEnableToolsInPIE(TEXT("p.Dataflow.EnableToolsInPIE"), bDataflowEditorEnableToolsInPIE,
@@ -87,7 +87,7 @@ void UDataflowEditorMode::Enter()
 	UE::TransformGizmoUtil::RegisterTransformGizmoContextObject(GetInteractiveToolsContext());
 
 	// Initialize view mode to a default
-	ConstructionViewMode = Dataflow::FRenderingViewModeFactory::GetInstance().GetViewMode(Dataflow::FDataflowConstruction3DViewMode::Name);
+	ConstructionViewMode = UE::Dataflow::FRenderingViewModeFactory::GetInstance().GetViewMode(UE::Dataflow::FDataflowConstruction3DViewMode::Name);
 }
 
 void UDataflowEditorMode::SetDataflowEditor(UDataflowEditor* InDataflowEditor) 
@@ -141,7 +141,7 @@ void UDataflowEditorMode::RegisterDataflowTool(TSharedPtr<FUICommandInfo> UIComm
 			if (const IDataflowEditorToolBuilder* const DataflowToolBuilder = Cast<IDataflowEditorToolBuilder>(Builder))
 			{
 				// Check if we need to switch view modes before starting the tool
-				TArray<const Dataflow::IDataflowConstructionViewMode*> SupportedModes;
+				TArray<const UE::Dataflow::IDataflowConstructionViewMode*> SupportedModes;
 
 				DataflowToolBuilder->GetSupportedConstructionViewModes(*ContextObject, SupportedModes);
 
@@ -181,7 +181,7 @@ void UDataflowEditorMode::RegisterDataflowTool(TSharedPtr<FUICommandInfo> UIComm
 
 			// Make sure the ContextObject's selected Collection is the from the Input side of the selected node (so that the tool gets the Collection as it appears before node execution)
 
-			if (TSharedPtr<Dataflow::FEngineContext> DataflowContext = ContextObject->GetDataflowContext())
+			if (TSharedPtr<UE::Dataflow::FEngineContext> DataflowContext = ContextObject->GetDataflowContext())
 			{
 				if (UDataflowEdNode* const SelectedNode = ContextObject->GetSelectedNode())
 				{
@@ -254,7 +254,7 @@ void UDataflowEditorMode::RegisterTools()
 
 	UEditorInteractiveToolsContext* const ConstructionViewportToolsContext = GetInteractiveToolsContext();
 
-	Dataflow::FDataflowToolRegistry& ToolRegistry = Dataflow::FDataflowToolRegistry::Get();
+	UE::Dataflow::FDataflowToolRegistry& ToolRegistry = UE::Dataflow::FDataflowToolRegistry::Get();
 	const TArray<FName> NodeNames = ToolRegistry.GetNodeNames();
 	for (const FName& RegisteredNodeName : NodeNames)
 	{
@@ -293,7 +293,7 @@ bool UDataflowEditorMode::ShouldToolStartBeAllowed(const FString& ToolIdentifier
 	}
 
 
-	if (Dataflow::Private::bDataflowEditorEnableToolsInPIE)
+	if (UE::Dataflow::Private::bDataflowEditorEnableToolsInPIE)
 	{
 		// UEdMode::ShouldToolStartBeAllowed returns (!GEditor->PlayWorld && !GIsPlayInEditorWorld) but we want to allow tools to start while in PIE
 		return true;
@@ -723,8 +723,8 @@ void UDataflowEditorMode::SetConstructionViewMode(const FName& NewViewModeName)
 		}
 	}
 
-	const Dataflow::FRenderingViewModeFactory& ViewModes = Dataflow::FRenderingViewModeFactory::GetInstance();
-	const Dataflow::IDataflowConstructionViewMode* const NewMode = ViewModes.GetViewMode(NewViewModeName);
+	const UE::Dataflow::FRenderingViewModeFactory& ViewModes = UE::Dataflow::FRenderingViewModeFactory::GetInstance();
+	const UE::Dataflow::IDataflowConstructionViewMode* const NewMode = ViewModes.GetViewMode(NewViewModeName);
 	if (!NewMode)
 	{
 		UE_LOG(LogChaos, Warning, TEXT("Warning : Unknown rendering view mode: %s"), *NewViewModeName.ToString());
@@ -784,7 +784,7 @@ void UDataflowEditorMode::SetConstructionViewMode(const FName& NewViewModeName)
 	}
 }
 
-const Dataflow::IDataflowConstructionViewMode* UDataflowEditorMode::GetConstructionViewMode() const
+const UE::Dataflow::IDataflowConstructionViewMode* UDataflowEditorMode::GetConstructionViewMode() const
 {
 	return ConstructionViewMode;
 }
@@ -801,9 +801,9 @@ bool UDataflowEditorMode::CanChangeConstructionViewModeTo(const FName& NewViewMo
 				{
 					if (const UDataflowEdNode* const SelectedDataflowEdNode = Cast<UDataflowEdNode>(SelectedNode))
 					{
-						if (const Dataflow::IDataflowConstructionViewMode* const ViewMode = Dataflow::FRenderingViewModeFactory::GetInstance().GetViewMode(NewViewModeName))
+						if (const UE::Dataflow::IDataflowConstructionViewMode* const ViewMode = UE::Dataflow::FRenderingViewModeFactory::GetInstance().GetViewMode(NewViewModeName))
 						{
-							if (Dataflow::CanRenderNodeOutput(*SelectedDataflowEdNode, *EditorContent, *ViewMode))
+							if (UE::Dataflow::CanRenderNodeOutput(*SelectedDataflowEdNode, *EditorContent, *ViewMode))
 							{
 								return true;
 							}
@@ -830,10 +830,10 @@ bool UDataflowEditorMode::CanChangeConstructionViewModeTo(const FName& NewViewMo
 		const UDataflowContextObject* const DataflowContextObject = ConstructionToolsContext->ContextObjectStore->FindContext<UDataflowContextObject>();
 		checkf(DataflowContextObject, TEXT("No Dataflow Context Object found in ContextObjectStore, despite having an Active Tool. This should have been created by the time a tool is activated"));
 
-		TArray<const Dataflow::IDataflowConstructionViewMode*> SupportedViewModes;
+		TArray<const UE::Dataflow::IDataflowConstructionViewMode*> SupportedViewModes;
 		DataflowToolBuilder->GetSupportedConstructionViewModes(*DataflowContextObject, SupportedViewModes);
 
-		if (const Dataflow::IDataflowConstructionViewMode* const NewViewMode = Dataflow::FRenderingViewModeFactory::GetInstance().GetViewMode(NewViewModeName))
+		if (const UE::Dataflow::IDataflowConstructionViewMode* const NewViewMode = UE::Dataflow::FRenderingViewModeFactory::GetInstance().GetViewMode(NewViewModeName))
 		{
 			return SupportedViewModes.Contains(NewViewMode);
 		}
diff --git a/Engine/Plugins/Experimental/Dataflow/Source/DataflowEditor/Private/Dataflow/DataflowEditorModule.cpp b/Engine/Plugins/Experimental/Dataflow/Source/DataflowEditor/Private/Dataflow/DataflowEditorModule.cpp
--- a/Engine/Plugins/Experimental/Dataflow/Source/DataflowEditor/Private/Dataflow/DataflowEditorModule.cpp
+++ b/Engine/Plugins/Experimental/Dataflow/Source/DataflowEditor/Private/Dataflow/DataflowEditorModule.cpp
@@ -42,7 +42,7 @@ public:
 };
 
 
-class FDataflowToolActionCommandBindings : public Dataflow::FDataflowToolRegistry::IDataflowToolActionCommands
+class FDataflowToolActionCommandBindings : public UE::Dataflow::FDataflowToolRegistry::IDataflowToolActionCommands
 {
 public:
 	FDataflowToolActionCommandBindings()
@@ -74,13 +74,13 @@ void FDataflowEditorModule::StartupModule()
 	// Register type customizations
 	if (FPropertyEditorModule* const PropertyModule = FModuleManager::GetModulePtr<FPropertyEditorModule>("PropertyEditor"))
 	{
-		PropertyModule->RegisterCustomPropertyTypeLayout(ScalarVertexPropertyGroupName, FOnGetPropertyTypeCustomizationInstance::CreateStatic(&Dataflow::FScalarVertexPropertyGroupCustomization::MakeInstance));
-		PropertyModule->RegisterCustomPropertyTypeLayout(DataflowFunctionPropertyName, FOnGetPropertyTypeCustomizationInstance::CreateStatic(&Dataflow::FFunctionPropertyCustomization::MakeInstance));
+		PropertyModule->RegisterCustomPropertyTypeLayout(ScalarVertexPropertyGroupName, FOnGetPropertyTypeCustomizationInstance::CreateStatic(&UE::Dataflow::FScalarVertexPropertyGroupCustomization::MakeInstance));
+		PropertyModule->RegisterCustomPropertyTypeLayout(DataflowFunctionPropertyName, FOnGetPropertyTypeCustomizationInstance::CreateStatic(&UE::Dataflow::FFunctionPropertyCustomization::MakeInstance));
 	}
 
-	Dataflow::RenderingCallbacks();
+	UE::Dataflow::RenderingCallbacks();
 
-	Dataflow::FDataflowToolRegistry& ToolRegistry = Dataflow::FDataflowToolRegistry::Get();
+	UE::Dataflow::FDataflowToolRegistry& ToolRegistry = UE::Dataflow::FDataflowToolRegistry::Get();
 
 	UDataflowEditorWeightMapPaintToolBuilder* const ToolBuilder = NewObject<UDataflowEditorWeightMapPaintToolBuilder>();
 	TSharedRef<const FDataflowToolActionCommandBindings> Actions = MakeShared<FDataflowToolActionCommandBindings>();
@@ -102,7 +102,7 @@ void FDataflowEditorModule::ShutdownModule()
 
 	FDataflowEditorCommands::Unregister();
 
-	Dataflow::FDataflowToolRegistry& ToolRegistry = Dataflow::FDataflowToolRegistry::Get();
+	UE::Dataflow::FDataflowToolRegistry& ToolRegistry = UE::Dataflow::FDataflowToolRegistry::Get();
 	ToolRegistry.RemoveNodeToToolMapping(FDataflowCollectionAddScalarVertexPropertyNode::StaticType());
 }
 
diff --git a/Engine/Plugins/Experimental/Dataflow/Source/DataflowEditor/Private/Dataflow/DataflowEditorToolkit.cpp b/Engine/Plugins/Experimental/Dataflow/Source/DataflowEditor/Private/Dataflow/DataflowEditorToolkit.cpp
--- a/Engine/Plugins/Experimental/Dataflow/Source/DataflowEditor/Private/Dataflow/DataflowEditorToolkit.cpp
+++ b/Engine/Plugins/Experimental/Dataflow/Source/DataflowEditor/Private/Dataflow/DataflowEditorToolkit.cpp
@@ -726,8 +726,8 @@ void FDataflowEditorToolkit::OnPropertyValueChanged(const FPropertyChangedEvent&
 		ensure(EditorContent);
 		if (UDataflow* DataflowAsset = EditorContent->GetDataflowAsset())
 		{
-			TSharedPtr<Dataflow::FEngineContext> DataflowContext = EditorContent->GetDataflowContext();
-			Dataflow::FTimestamp LastNodeTimestamp = EditorContent->GetLastModifiedTimestamp();
+			TSharedPtr<UE::Dataflow::FEngineContext> DataflowContext = EditorContent->GetDataflowContext();
+			UE::Dataflow::FTimestamp LastNodeTimestamp = EditorContent->GetLastModifiedTimestamp();
 			
 			FDataflowEditorCommands::OnPropertyValueChanged(DataflowAsset, DataflowContext, LastNodeTimestamp, PropertyChangedEvent, SelectedDataflowNodes);
 
@@ -952,7 +952,7 @@ void FDataflowEditorToolkit::OnNodeSelectionChanged(const TSet<UObject*>& InNewS
 				EditorContent->SetSelectedNode(PrimarySelection);
 
 				// Call the node's OnSelected function. Some nodes use this to cache information from the inputs (e.g. FDataflowCollectionAddScalarVertexPropertyNode::CachedCollectionGroupNames)
-				TSharedPtr<Dataflow::FEngineContext> DataflowContext = EditorContent->GetDataflowContext();
+				TSharedPtr<UE::Dataflow::FEngineContext> DataflowContext = EditorContent->GetDataflowContext();
 				if (PrimarySelection && DataflowContext.IsValid())
 				{
 					if (const TSharedPtr<FDataflowNode> DataflowNode = PrimarySelection->GetDataflowNode())
@@ -993,22 +993,22 @@ void FDataflowEditorToolkit::OnNodeSelectionChanged(const TSet<UObject*>& InNewS
 
 		if (PrimarySelection && GetEditorContent())
 		{
-			if (!Dataflow::CanRenderNodeOutput(*PrimarySelection, *GetEditorContent(), *DataflowMode->GetConstructionViewMode()))
+			if (!UE::Dataflow::CanRenderNodeOutput(*PrimarySelection, *GetEditorContent(), *DataflowMode->GetConstructionViewMode()))
 			{
 				// Selected node can't render with the current view mode. Check through available view modes and see if it can render with any of them
 
 				bFoundViewMode = false;
 
-				TArray<Dataflow::FRenderingParameter> RenderingParameters = PrimarySelection->GetRenderParameters();
-				for (const Dataflow::FRenderingParameter& Param : RenderingParameters)
+				TArray<UE::Dataflow::FRenderingParameter> RenderingParameters = PrimarySelection->GetRenderParameters();
+				for (const UE::Dataflow::FRenderingParameter& Param : RenderingParameters)
 				{
 					const FName NodeOutputTypeName = Param.Type;
 
-					for (const TPair<FName, TUniquePtr<Dataflow::IDataflowConstructionViewMode>>& ViewMode : Dataflow::FRenderingViewModeFactory::GetInstance().GetViewModes())
+					for (const TPair<FName, TUniquePtr<UE::Dataflow::IDataflowConstructionViewMode>>& ViewMode : UE::Dataflow::FRenderingViewModeFactory::GetInstance().GetViewModes())
 					{
 						check(ViewMode.Value.IsValid());
 
-						const bool bCanRender = Dataflow::CanRenderNodeOutput(*PrimarySelection, *GetEditorContent(), *ViewMode.Value);
+						const bool bCanRender = UE::Dataflow::CanRenderNodeOutput(*PrimarySelection, *GetEditorContent(), *ViewMode.Value);
 
 						if (bCanRender)
 						{
@@ -1029,7 +1029,7 @@ void FDataflowEditorToolkit::OnNodeSelectionChanged(const TSet<UObject*>& InNewS
 		if (!bFoundViewMode)
 		{
 			// TODO: Clear and disable View Mode Button. For now set default mode to the built-in 3D view mode.
-			DataflowMode->SetConstructionViewMode(Dataflow::FDataflowConstruction3DViewMode::Name);
+			DataflowMode->SetConstructionViewMode(UE::Dataflow::FDataflowConstruction3DViewMode::Name);
 		}
 	}
 }
@@ -1079,11 +1079,11 @@ void FDataflowEditorToolkit::Tick(float DeltaTime)
 	{
 		if (EditorContent->GetDataflowAsset())
 		{
-			Dataflow::FTimestamp InitTimeStamp = EditorContent->GetLastModifiedTimestamp();
+			UE::Dataflow::FTimestamp InitTimeStamp = EditorContent->GetLastModifiedTimestamp();
 			if (!EditorContent->GetDataflowContext())
 			{
-				EditorContent->SetDataflowContext(MakeShared<Dataflow::FEngineContext>(EditorContent->GetDataflowOwner()));
-				InitTimeStamp = Dataflow::FTimestamp::Invalid;
+				EditorContent->SetDataflowContext(MakeShared<UE::Dataflow::FEngineContext>(EditorContent->GetDataflowOwner()));
+				InitTimeStamp = UE::Dataflow::FTimestamp::Invalid;
 			}
 
 			// Update the list of dataflow terminal contents 
@@ -1097,16 +1097,16 @@ void FDataflowEditorToolkit::Tick(float DeltaTime)
 				{
 					if (const UDataflow* const Dataflow = EditorContent->GetDataflowAsset())
 					{
-						if (const TSharedPtr<const Dataflow::FGraph> Graph = Dataflow->GetDataflow())
+						if (const TSharedPtr<const UE::Dataflow::FGraph> Graph = Dataflow->GetDataflow())
 						{
 							const FName TerminalNodeName(TerminalContent->GetDataflowTerminal());
 							const FDataflowNode* const Node = Graph->FindBaseNode(TerminalNodeName).Get();
 
-							Dataflow::FTimestamp TerminalNodeTimeStamp = InitTimeStamp;
+							UE::Dataflow::FTimestamp TerminalNodeTimeStamp = InitTimeStamp;
 							EvaluateNode(Node, nullptr, TerminalNodeTimeStamp);  // When Node is null, EvaluateNode falls back on the EditorContent terminal node
 
 							// Take the Max of the existing time stamp, as other terminal nodes might have more recent invalidations
-							const Dataflow::FTimestamp LastModifiedTimestamp = FMath::Max(EditorContent->GetLastModifiedTimestamp(), TerminalNodeTimeStamp);
+							const UE::Dataflow::FTimestamp LastModifiedTimestamp = FMath::Max(EditorContent->GetLastModifiedTimestamp(), TerminalNodeTimeStamp);
 							EditorContent->SetLastModifiedTimestamp(LastModifiedTimestamp);
 						}
 					}
@@ -1114,7 +1114,7 @@ void FDataflowEditorToolkit::Tick(float DeltaTime)
 			}
 			else
 			{
-				Dataflow::FTimestamp TerminalNodeTimeStamp = InitTimeStamp;
+				UE::Dataflow::FTimestamp TerminalNodeTimeStamp = InitTimeStamp;
 				EvaluateNode(nullptr, nullptr, TerminalNodeTimeStamp);
 				EditorContent->SetLastModifiedTimestamp(TerminalNodeTimeStamp);
 			}
@@ -1137,7 +1137,7 @@ TStatId FDataflowEditorToolkit::GetStatId() const
 	RETURN_QUICK_DECLARE_CYCLE_STAT(FDataflowEditorToolkit, STATGROUP_Tickables);
 }
 
-void FDataflowEditorToolkit::EvaluateNode(const FDataflowNode* Node, const FDataflowOutput* Output, Dataflow::FTimestamp& InOutTimestamp)
+void FDataflowEditorToolkit::EvaluateNode(const FDataflowNode* Node, const FDataflowOutput* Output, UE::Dataflow::FTimestamp& InOutTimestamp)
 {
 	UE_LOG(LogChaosDataflow, Verbose, TEXT("FDataflowEditorToolkit::EvaluateNode(): Node [%s], Output [%s]"), Node ? *Node->GetName().ToString() : TEXT("nullptr"), Output ? *Output->GetName().ToString() : TEXT("nullptr"));
 
@@ -1159,14 +1159,14 @@ void FDataflowEditorToolkit::EvaluateNode(const FDataflowNode* Node, const FData
 TSharedRef<SDataflowGraphEditor> FDataflowEditorToolkit::CreateGraphEditorWidget(UDataflow* DataflowToEdit, TSharedPtr<IStructureDetailsView> InNodeDetailsEditor)
 {
 	ensure(DataflowToEdit);
-	using namespace Dataflow;
+	using namespace UE::Dataflow;
 
 	const FDataflowEditorCommands::FGraphEvaluationCallback Evaluate =
 		[this](const FDataflowNode* Node, const FDataflowOutput* Output)
 		{
 			if (const TObjectPtr<UDataflowBaseContent>& EditorContent = GetEditorContent())
 			{
-				Dataflow::FTimestamp LastNodeTimestamp = EditorContent->GetLastModifiedTimestamp();
+				UE::Dataflow::FTimestamp LastNodeTimestamp = EditorContent->GetLastModifiedTimestamp();
 
 				EvaluateNode(Node, Output, LastNodeTimestamp);
 
@@ -1179,7 +1179,7 @@ TSharedRef<SDataflowGraphEditor> FDataflowEditorToolkit::CreateGraphEditorWidget
 			OnFinishEvaluate();
 		};
 	
-	DataflowEditor->UpdateTerminalContents(Dataflow::FTimestamp::Invalid);
+	DataflowEditor->UpdateTerminalContents(FTimestamp::Invalid);
 	
 	SGraphEditor::FGraphEditorEvents InEvents;
 	InEvents.OnVerifyTextCommit = FOnNodeVerifyTextCommit::CreateSP(this, &FDataflowEditorToolkit::OnNodeVerifyTitleCommit);
diff --git a/Engine/Plugins/Experimental/Dataflow/Source/DataflowEditor/Private/Dataflow/DataflowEditorUtil.h b/Engine/Plugins/Experimental/Dataflow/Source/DataflowEditor/Private/Dataflow/DataflowEditorUtil.h
--- a/Engine/Plugins/Experimental/Dataflow/Source/DataflowEditor/Private/Dataflow/DataflowEditorUtil.h
+++ b/Engine/Plugins/Experimental/Dataflow/Source/DataflowEditor/Private/Dataflow/DataflowEditorUtil.h
@@ -17,7 +17,7 @@ class USkeleton;
 class UAnimationAsset;
 class UMaterial;
 class UDataflowEdNode;
-namespace Dataflow { class IDataflowConstructionViewMode; }
+namespace UE::Dataflow { class IDataflowConstructionViewMode; }
 
 namespace UE::Dataflow::Private
 {
@@ -44,10 +44,10 @@ namespace UE
 	}
 }
 
-namespace Dataflow
+namespace UE::Dataflow
 {
-	TSharedPtr<::Dataflow::FEngineContext> GetContext(TObjectPtr<UDataflowBaseContent> Content);
+	TSharedPtr<UE::Dataflow::FEngineContext> GetContext(TObjectPtr<UDataflowBaseContent> Content);
 
-	bool CanRenderNodeOutput(const UDataflowEdNode& EdNode, const UDataflowBaseContent& EditorContent, const ::Dataflow::IDataflowConstructionViewMode& ViewMode);
+	bool CanRenderNodeOutput(const UDataflowEdNode& EdNode, const UDataflowBaseContent& EditorContent, const UE::Dataflow::IDataflowConstructionViewMode& ViewMode);
 	void RenderNodeOutput(GeometryCollection::Facades::FRenderingFacade& Facade, const UDataflowEdNode& Node, const UDataflowBaseContent& EditorContent);
 }
\ No newline at end of file
diff --git a/Engine/Plugins/Experimental/Dataflow/Source/DataflowEditor/Private/Dataflow/DataflowEngineRendering.cpp b/Engine/Plugins/Experimental/Dataflow/Source/DataflowEditor/Private/Dataflow/DataflowEngineRendering.cpp
--- a/Engine/Plugins/Experimental/Dataflow/Source/DataflowEditor/Private/Dataflow/DataflowEngineRendering.cpp
+++ b/Engine/Plugins/Experimental/Dataflow/Source/DataflowEditor/Private/Dataflow/DataflowEngineRendering.cpp
@@ -14,9 +14,9 @@
 #include "GeometryCollection/GeometryCollectionAlgo.h"
 #include "UDynamicMesh.h"
 
-namespace Dataflow
+namespace UE::Dataflow
 {
-	void RenderBasicGeometryCollection(GeometryCollection::Facades::FRenderingFacade& RenderCollection, const Dataflow::FGraphRenderingState& State, TArray<FLinearColor>* VertexColorOverride = nullptr)
+	void RenderBasicGeometryCollection(GeometryCollection::Facades::FRenderingFacade& RenderCollection, const UE::Dataflow::FGraphRenderingState& State, TArray<FLinearColor>* VertexColorOverride = nullptr)
 	{
 		FManagedArrayCollection Default;
 		FName PrimaryOutput = State.GetRenderOutputs()[0]; // "Collection"
@@ -120,7 +120,7 @@ namespace Dataflow
 
 	}
 
-	void RenderMeshIndexedGeometryCollection(GeometryCollection::Facades::FRenderingFacade& RenderCollection, const Dataflow::FGraphRenderingState& State, TArray<FLinearColor>* VertexColorOverride = nullptr )
+	void RenderMeshIndexedGeometryCollection(GeometryCollection::Facades::FRenderingFacade& RenderCollection, const UE::Dataflow::FGraphRenderingState& State, TArray<FLinearColor>* VertexColorOverride = nullptr )
 	{
 		auto ToD = [](FVector3f V) { return FVector3d(V.X, V.Y, V.Z); };
 		auto ToF = [](FVector3d V) { return FVector3f(V.X, V.Y, V.Z); };
@@ -253,12 +253,12 @@ namespace Dataflow
 
 	class FGeometryCollectionSurfaceRenderCallbacks : public FRenderingFactory::ICallbackInterface
 	{
-		virtual Dataflow::FRenderKey GetRenderKey() const override
+		virtual UE::Dataflow::FRenderKey GetRenderKey() const override
 		{
 			return { "SurfaceRender", FGeometryCollection::StaticType() };
 		}
 
-		virtual bool CanRender(const Dataflow::IDataflowConstructionViewMode& ViewMode) const override
+		virtual bool CanRender(const UE::Dataflow::IDataflowConstructionViewMode& ViewMode) const override
 		{
 			return (ViewMode.GetName() == FDataflowConstruction3DViewMode::Name);
 		}
@@ -303,12 +303,12 @@ namespace Dataflow
 
 	class FGeometryCollectionSurfaceWeightsRenderCallbacks : public FRenderingFactory::ICallbackInterface
 	{
-		virtual Dataflow::FRenderKey GetRenderKey() const override
+		virtual UE::Dataflow::FRenderKey GetRenderKey() const override
 		{
 			return { "SurfaceWeightsRender", FGeometryCollection::StaticType() };
 		}
 
-		virtual bool CanRender(const Dataflow::IDataflowConstructionViewMode& ViewMode) const override
+		virtual bool CanRender(const UE::Dataflow::IDataflowConstructionViewMode& ViewMode) const override
 		{
 			return (ViewMode.GetName() == FDataflowConstruction3DViewMode::Name);
 		}
@@ -404,12 +404,12 @@ namespace Dataflow
 
 	class FDynamicMesh3SurfaceRenderCallbacks : public FRenderingFactory::ICallbackInterface
 	{
-		virtual Dataflow::FRenderKey GetRenderKey() const override
+		virtual UE::Dataflow::FRenderKey GetRenderKey() const override
 		{
 			return { "SurfaceRender", FName("FDynamicMesh3") };
 		}
 
-		virtual bool CanRender(const Dataflow::IDataflowConstructionViewMode& ViewMode) const override
+		virtual bool CanRender(const UE::Dataflow::IDataflowConstructionViewMode& ViewMode) const override
 		{
 			return (ViewMode.GetName() == FDataflowConstruction3DViewMode::Name);
 		}
@@ -471,12 +471,12 @@ namespace Dataflow
 
 	class FBoxSurfaceRenderCallbacks : public FRenderingFactory::ICallbackInterface
 	{
-		virtual Dataflow::FRenderKey GetRenderKey() const override
+		virtual UE::Dataflow::FRenderKey GetRenderKey() const override
 		{
 			return { "SurfaceRender", FName("FBox") };
 		}
 
-		virtual bool CanRender(const Dataflow::IDataflowConstructionViewMode& ViewMode) const override
+		virtual bool CanRender(const UE::Dataflow::IDataflowConstructionViewMode& ViewMode) const override
 		{
 			return (ViewMode.GetName() == FDataflowConstruction3DViewMode::Name);
 		}
@@ -538,12 +538,12 @@ namespace Dataflow
 
 	class FFieldVolumeRenderCallbacks : public FRenderingFactory::ICallbackInterface
 	{
-		virtual Dataflow::FRenderKey GetRenderKey() const override
+		virtual UE::Dataflow::FRenderKey GetRenderKey() const override
 		{
 			return { "VolumeRender", FFieldCollection::StaticType() };
 		}
 
-		virtual bool CanRender(const Dataflow::IDataflowConstructionViewMode& ViewMode) const override
+		virtual bool CanRender(const UE::Dataflow::IDataflowConstructionViewMode& ViewMode) const override
 		{
 			return (ViewMode.GetName() == FDataflowConstruction3DViewMode::Name);
 		}
@@ -593,7 +593,7 @@ namespace Dataflow
 
 	void RenderingCallbacks()
 	{
-		using namespace Dataflow;
+		using namespace UE::Dataflow;
 
 		FRenderingFactory::GetInstance()->RegisterCallbacks(MakeUnique<FGeometryCollectionSurfaceRenderCallbacks>());
 		FRenderingFactory::GetInstance()->RegisterCallbacks(MakeUnique<FGeometryCollectionSurfaceWeightsRenderCallbacks>());
diff --git a/Engine/Plugins/Experimental/Dataflow/Source/DataflowEditor/Private/Dataflow/DataflowFunctionPropertyCustomization.cpp b/Engine/Plugins/Experimental/Dataflow/Source/DataflowEditor/Private/Dataflow/DataflowFunctionPropertyCustomization.cpp
--- a/Engine/Plugins/Experimental/Dataflow/Source/DataflowEditor/Private/Dataflow/DataflowFunctionPropertyCustomization.cpp
+++ b/Engine/Plugins/Experimental/Dataflow/Source/DataflowEditor/Private/Dataflow/DataflowFunctionPropertyCustomization.cpp
@@ -11,7 +11,7 @@
 
 #define LOCTEXT_NAMESPACE "FunctionPropertyCustomization"
 
-namespace Dataflow
+namespace UE::Dataflow
 {
 	namespace Private
 	{
@@ -182,6 +182,6 @@ namespace Dataflow
 			}
 		}
 	}
-}  // End namespace Dataflow
+}  // End namespace UE::Dataflow
 
 #undef LOCTEXT_NAMESPACE
diff --git a/Engine/Plugins/Experimental/Dataflow/Source/DataflowEditor/Private/Dataflow/DataflowGraphEditor.cpp b/Engine/Plugins/Experimental/Dataflow/Source/DataflowEditor/Private/Dataflow/DataflowGraphEditor.cpp
--- a/Engine/Plugins/Experimental/Dataflow/Source/DataflowEditor/Private/Dataflow/DataflowGraphEditor.cpp
+++ b/Engine/Plugins/Experimental/Dataflow/Source/DataflowEditor/Private/Dataflow/DataflowGraphEditor.cpp
@@ -164,7 +164,7 @@ void SDataflowGraphEditor::Construct(const FArguments& InArgs, UObject* InAssetO
 	SetNodeFactory( MakeShared<FDataflowGraphNodeFactory>(this) );
 }
 
-const TSharedPtr<Dataflow::FEngineContext> SDataflowGraphEditor::GetDataflowContext() const
+const TSharedPtr<UE::Dataflow::FEngineContext> SDataflowGraphEditor::GetDataflowContext() const
 {
 	if (DataflowEditor)
 	{
@@ -173,7 +173,7 @@ const TSharedPtr<Dataflow::FEngineContext> SDataflowGraphEditor::GetDataflowCont
 			return DataflowEditor->GetEditorContent()->GetDataflowContext();
 		}
 	}
-	return TSharedPtr<Dataflow::FEngineContext>(nullptr);
+	return TSharedPtr<UE::Dataflow::FEngineContext>(nullptr);
 }
 
 void SDataflowGraphEditor::EvaluateNode()
@@ -185,7 +185,7 @@ void SDataflowGraphEditor::EvaluateNode()
 					Cast<UDataflowEdNode>(SelectedNode)->GetDataflowNode()->GetName().ToString() : FString();
 			}));
 
-	using namespace Dataflow;
+	using namespace UE::Dataflow;
 
 	TOptional<FContextThreaded> DefaultContext(EvaluateGraphCallback ? TOptional<FContextThreaded>() : TOptional<FContextThreaded>(FContextThreaded()));
 
@@ -456,7 +456,7 @@ void SDataflowGraphEditor::CreateVertexSelectionNode(const FString & InArray)
 		if (UEdGraphNode* NewEdNode = NodeAction->PerformAction(Graph, nullptr, GetGraphEditor()->GetPasteLocation(), false))
 		{
 			FDataflowAssetEdit Edit = Graph->EditDataflow();
-			if (Dataflow::FGraph* DataflowGraph = Edit.GetGraph())
+			if (UE::Dataflow::FGraph* DataflowGraph = Edit.GetGraph())
 			{
 				if (TSharedPtr<FDataflowNode> Node = DataflowGraph->FindBaseNode(((UDataflowEdNode*)NewEdNode)->DataflowNodeGuid))
 				{
@@ -524,7 +524,7 @@ void SDataflowGraphEditor::OnAddOptionPin()
 {
 	UDataflow* const Graph = DataflowAsset.Get();
 	FDataflowAssetEdit Edit = Graph->EditDataflow();
-	if (Dataflow::FGraph* const DataflowGraph = Edit.GetGraph())
+	if (UE::Dataflow::FGraph* const DataflowGraph = Edit.GetGraph())
 	{
 		const FGraphPanelSelectionSet& SelectedNodes = GetSelectedNodes();
 
@@ -556,7 +556,7 @@ bool SDataflowGraphEditor::CanAddOptionPin() const
 	bool bCanAddOptionPin = false;
 
 	const UDataflow* const Graph = DataflowAsset.Get();
-	if (const Dataflow::FGraph* const DataflowGraph = Graph->GetDataflow().Get())
+	if (const UE::Dataflow::FGraph* const DataflowGraph = Graph->GetDataflow().Get())
 	{
 		const FGraphPanelSelectionSet& SelectedNodes = GetSelectedNodes();
 
@@ -589,7 +589,7 @@ void SDataflowGraphEditor::OnRemoveOptionPin()
 {
 	UDataflow* const Graph = DataflowAsset.Get();
 	FDataflowAssetEdit Edit = Graph->EditDataflow();
-	if (Dataflow::FGraph* const DataflowGraph = Edit.GetGraph())
+	if (UE::Dataflow::FGraph* const DataflowGraph = Edit.GetGraph())
 	{
 		const FGraphPanelSelectionSet& SelectedNodes = GetSelectedNodes();
 
@@ -621,7 +621,7 @@ bool SDataflowGraphEditor::CanRemoveOptionPin() const
 	bool bCanRemoveOptionPin = false;
 
 	const UDataflow* const Graph = DataflowAsset.Get();
-	if (const Dataflow::FGraph* const DataflowGraph = Graph->GetDataflow().Get())
+	if (const UE::Dataflow::FGraph* const DataflowGraph = Graph->GetDataflow().Get())
 	{
 		const FGraphPanelSelectionSet& SelectedNodes = GetSelectedNodes();
 
diff --git a/Engine/Plugins/Experimental/Dataflow/Source/DataflowEditor/Private/Dataflow/DataflowSNode.cpp b/Engine/Plugins/Experimental/Dataflow/Source/DataflowEditor/Private/Dataflow/DataflowSNode.cpp
--- a/Engine/Plugins/Experimental/Dataflow/Source/DataflowEditor/Private/Dataflow/DataflowSNode.cpp
+++ b/Engine/Plugins/Experimental/Dataflow/Source/DataflowEditor/Private/Dataflow/DataflowSNode.cpp
@@ -147,9 +147,9 @@ TSharedPtr<SGraphPin> SDataflowEdNode::CreatePinWidget(UEdGraphPin* Pin) const
 			{
 				if (FDataflowOutput* Output = DataflowNode->FindOutput(Pin->GetFName()))
 				{
-					if (const TSharedPtr<Dataflow::FEngineContext> DataflowContext = DataflowInterface->GetDataflowContext())
+					if (const TSharedPtr<UE::Dataflow::FEngineContext> DataflowContext = DataflowInterface->GetDataflowContext())
 					{
-						TSet<Dataflow::FContextCacheKey> CacheKeys;
+						TSet<UE::Dataflow::FContextCacheKey> CacheKeys;
 						const int32 NumKeys = DataflowContext->GetKeys(CacheKeys);
 
 						//
@@ -213,12 +213,12 @@ void SDataflowEdNode::UpdateErrorInfo()
 	{
 		if (const TSharedPtr<FDataflowNode> DataflowNode = DataflowGraphNode->GetDataflowNode())
 		{
-			if (Dataflow::FNodeFactory::IsNodeExperimental(DataflowNode->GetType()))
+			if (UE::Dataflow::FNodeFactory::IsNodeExperimental(DataflowNode->GetType()))
 			{
 				ErrorMsg = FString(TEXT("Experimental"));
 				ErrorColor = FAppStyle::GetColor("ErrorReporting.WarningBackgroundColor");
 			} 
-			if (Dataflow::FNodeFactory::IsNodeDeprecated(DataflowNode->GetType()))
+			if (UE::Dataflow::FNodeFactory::IsNodeDeprecated(DataflowNode->GetType()))
 			{
 				ErrorMsg = FString(TEXT("Deprecated"));
 				ErrorColor = FAppStyle::GetColor("ErrorReporting.WarningBackgroundColor");
@@ -234,7 +234,7 @@ FReply SDataflowEdNode::OnMouseButtonDoubleClick(const FGeometry& InMyGeometry,
 	{
 		if (UDataflowEdNode* DataflowNode = Cast<UDataflowEdNode>(GraphNode))
 		{
-			if (TSharedPtr<Dataflow::FGraph> Graph = DataflowNode->GetDataflowGraph())
+			if (TSharedPtr<UE::Dataflow::FGraph> Graph = DataflowNode->GetDataflowGraph())
 			{
 				if (TSharedPtr<FDataflowNode> Node = Graph->FindBaseNode(DataflowNode->GetDataflowNodeGuid()))
 				{
@@ -353,9 +353,9 @@ TSharedPtr<FAssetSchemaAction_Dataflow_CreateNode_DataflowEdNode> FAssetSchemaAc
 {
 	if (const UDataflow* Dataflow = Cast<UDataflow>(ParentGraph))
 	{
-		if (Dataflow::FNodeFactory* Factory = Dataflow::FNodeFactory::GetInstance())
+		if (UE::Dataflow::FNodeFactory* Factory = UE::Dataflow::FNodeFactory::GetInstance())
 		{
-			const Dataflow::FFactoryParameters& Param = Factory->GetParameters(InNodeTypeName);
+			const UE::Dataflow::FFactoryParameters& Param = Factory->GetParameters(InNodeTypeName);
 			if (Param.IsValid())
 			{
 				const bool bIsSimulationNode = Param.Tags.Contains(UDataflow::SimulationTag);
@@ -422,7 +422,7 @@ void SDataflowEdNode::CopyDataflowNodeSettings(TSharedPtr<FDataflowNode> SourceD
 
 static UDataflowEdNode* CreateNode(UDataflow* Dataflow, UEdGraphPin* FromPin, const FVector2D Location, bool bSelectNewNode, const FName NodeUniqueName, const FName NodeTypeName, TSharedPtr<FDataflowNode> DataflowNodeToDuplicate, bool bCopySettings = false)
 {
-	if (Dataflow::FNodeFactory* Factory = Dataflow::FNodeFactory::GetInstance())
+	if (UE::Dataflow::FNodeFactory* Factory = UE::Dataflow::FNodeFactory::GetInstance())
 	{
 		if (TSharedPtr<FDataflowNode> DataflowNode =
 			Factory->NewNodeFromRegisteredType(
@@ -502,9 +502,9 @@ UEdGraphNode* FAssetSchemaAction_Dataflow_CreateNode_DataflowEdNode::PerformActi
 //
 TSharedPtr<FAssetSchemaAction_Dataflow_DuplicateNode_DataflowEdNode> FAssetSchemaAction_Dataflow_DuplicateNode_DataflowEdNode::CreateAction(UEdGraph* ParentGraph, const FName& InNodeTypeName)
 {
-	if (Dataflow::FNodeFactory* Factory = Dataflow::FNodeFactory::GetInstance())
+	if (UE::Dataflow::FNodeFactory* Factory = UE::Dataflow::FNodeFactory::GetInstance())
 	{
-		const Dataflow::FFactoryParameters& Param = Factory->GetParameters(InNodeTypeName);
+		const UE::Dataflow::FFactoryParameters& Param = Factory->GetParameters(InNodeTypeName);
 		if (Param.IsValid())
 		{
 			const FText ToolTip = FText::FromString(Param.ToolTip.IsEmpty() ? FString("Add a Dataflow node.") : Param.ToolTip);
@@ -542,9 +542,9 @@ UEdGraphNode* FAssetSchemaAction_Dataflow_DuplicateNode_DataflowEdNode::PerformA
 //
 TSharedPtr<FAssetSchemaAction_Dataflow_PasteNode_DataflowEdNode> FAssetSchemaAction_Dataflow_PasteNode_DataflowEdNode::CreateAction(UEdGraph* ParentGraph, const FName& InNodeTypeName)
 {
-	if (Dataflow::FNodeFactory* Factory = Dataflow::FNodeFactory::GetInstance())
+	if (UE::Dataflow::FNodeFactory* Factory = UE::Dataflow::FNodeFactory::GetInstance())
 	{
-		const Dataflow::FFactoryParameters& Param = Factory->GetParameters(InNodeTypeName);
+		const UE::Dataflow::FFactoryParameters& Param = Factory->GetParameters(InNodeTypeName);
 		if (Param.IsValid())
 		{
 			const FText ToolTip = FText::FromString(Param.ToolTip.IsEmpty() ? FString("Add a Dataflow node.") : Param.ToolTip);
@@ -561,7 +561,7 @@ TSharedPtr<FAssetSchemaAction_Dataflow_PasteNode_DataflowEdNode> FAssetSchemaAct
 
 static UDataflowEdNode* CreateNodeFromPaste(UDataflow* Dataflow, UEdGraphPin* FromPin, const FVector2D Location, bool bSelectNewNode, const FName NodeUniqueName, const FName NodeTypeName, FString NodeProperties)
 {
-	if (Dataflow::FNodeFactory* Factory = Dataflow::FNodeFactory::GetInstance())
+	if (UE::Dataflow::FNodeFactory* Factory = UE::Dataflow::FNodeFactory::GetInstance())
 	{
 		if (TSharedPtr<FDataflowNode> DataflowNode =
 			Factory->NewNodeFromRegisteredType(
