Dataflow Editor: add Analytics events for starting/stopping the Dataflow Editor

diff --git a/Engine/Plugins/Experimental/Dataflow/Source/DataflowEditor/Private/Dataflow/DataflowEditorMode.cpp b/Engine/Plugins/Experimental/Dataflow/Source/DataflowEditor/Private/Dataflow/DataflowEditorMode.cpp
--- a/Engine/Plugins/Experimental/Dataflow/Source/DataflowEditor/Private/Dataflow/DataflowEditorMode.cpp
+++ b/Engine/Plugins/Experimental/Dataflow/Source/DataflowEditor/Private/Dataflow/DataflowEditorMode.cpp
@@ -28,6 +28,7 @@
 #include "EditorModeManager.h"
 #include "EdModeInteractiveToolsContext.h"
 #include "Elements/Framework/EngineElementsLibrary.h"
+#include "EngineAnalytics.h"
 #include "MeshSelectionTool.h"
 #include "MeshVertexPaintTool.h"
 #include "MeshAttributePaintTool.h"
@@ -88,6 +89,16 @@ void UDataflowEditorMode::Enter()
 
 	// Initialize view mode to a default
 	ConstructionViewMode = UE::Dataflow::FRenderingViewModeFactory::GetInstance().GetViewMode(UE::Dataflow::FDataflowConstruction3DViewMode::Name);
+
+	// Log mode starting
+	if (FEngineAnalytics::IsAvailable())
+	{
+		LastModeStartTimestamp = FDateTime::UtcNow();
+		TArray<FAnalyticsEventAttribute> EventAttributes;
+		EventAttributes.Add(FAnalyticsEventAttribute(TEXT("Timestamp"), LastModeStartTimestamp.ToString()));
+		FEngineAnalytics::GetProvider().RecordEvent(TEXT("Editor.Usage.DataflowEditor.Enter"), EventAttributes);
+	}
+
 }
 
 void UDataflowEditorMode::SetDataflowEditor(UDataflowEditor* InDataflowEditor) 
@@ -407,6 +418,18 @@ void UDataflowEditorMode::Exit()
 		SimulationScene = nullptr;
 	}
 
+	// Log mode exit
+	if (FEngineAnalytics::IsAvailable())
+	{
+		const FTimespan ModeUsageDuration = FDateTime::UtcNow() - LastModeStartTimestamp;
+
+		TArray<FAnalyticsEventAttribute> Attributes;
+		Attributes.Add(FAnalyticsEventAttribute(TEXT("Timestamp"), FDateTime::UtcNow().ToString()));
+		Attributes.Add(FAnalyticsEventAttribute(TEXT("Duration.Seconds"), static_cast<float>(ModeUsageDuration.GetTotalSeconds())));
+
+		FEngineAnalytics::GetProvider().RecordEvent(TEXT("Editor.Usage.DataflowEditor.Exit"));
+	}
+
 	Super::Exit();
 }
 
diff --git a/Engine/Plugins/Experimental/Dataflow/Source/DataflowEditor/Public/Dataflow/DataflowEditorMode.h b/Engine/Plugins/Experimental/Dataflow/Source/DataflowEditor/Public/Dataflow/DataflowEditorMode.h
--- a/Engine/Plugins/Experimental/Dataflow/Source/DataflowEditor/Public/Dataflow/DataflowEditorMode.h
+++ b/Engine/Plugins/Experimental/Dataflow/Source/DataflowEditor/Public/Dataflow/DataflowEditorMode.h
@@ -221,5 +221,8 @@ private:
 
 	UDataflowEditor* DataflowEditor = nullptr;
 
+	// Timestamp for telemetry
+	FDateTime LastModeStartTimestamp;
+
 };
 
