CIS - Fixed issue Issue 608001: ScalarVertexPropertyGroupCustomization.cpp(24): Found an '#undef LOCTEXT_NAMESPACE' without a corresponding '#define'

#rnx

diff --git a/Engine/Plugins/Experimental/Dataflow/Source/DataflowEditor/Private/Dataflow/ScalarVertexPropertyGroupCustomization.cpp b/Engine/Plugins/Experimental/Dataflow/Source/DataflowEditor/Private/Dataflow/ScalarVertexPropertyGroupCustomization.cpp
--- a/Engine/Plugins/Experimental/Dataflow/Source/DataflowEditor/Private/Dataflow/ScalarVertexPropertyGroupCustomization.cpp
+++ b/Engine/Plugins/Experimental/Dataflow/Source/DataflowEditor/Private/Dataflow/ScalarVertexPropertyGroupCustomization.cpp
@@ -20,5 +20,3 @@ namespace UE::Dataflow
 		return GET_MEMBER_NAME_CHECKED(FDataflowCollectionAddScalarVertexPropertyNode, Collection);
 	}
 }  // End namespace UE::Chaos::ClothAsset
-
-#undef LOCTEXT_NAMESPACE
