Fix crash in Dataflow Editor when the asset doesn't have a valid Dataflow asset

diff --git a/Engine/Plugins/Experimental/Dataflow/Source/DataflowEditor/Private/Dataflow/DataflowSimulationScene.cpp b/Engine/Plugins/Experimental/Dataflow/Source/DataflowEditor/Private/Dataflow/DataflowSimulationScene.cpp
--- a/Engine/Plugins/Experimental/Dataflow/Source/DataflowEditor/Private/Dataflow/DataflowSimulationScene.cpp
+++ b/Engine/Plugins/Experimental/Dataflow/Source/DataflowEditor/Private/Dataflow/DataflowSimulationScene.cpp
@@ -219,17 +219,21 @@ void FDataflowSimulationScene::TickDataflowScene(const float DeltaSeconds)
 {
 	if(const TObjectPtr<UDataflowBaseContent>& EditorContent = GetEditorContent())
 	{
-		if(UE::Dataflow::ShouldResetWorld(EditorContent->GetDataflowAsset(), GetWorld(), LastTimeStamp) || EditorContent->IsSimulationDirty())
+		if (const TObjectPtr<UDataflow> DataflowGraph = EditorContent->GetDataflowAsset())
 		{
-			// Unregister components, cache manager, selection...
-			ResetSimulationScene();
+			if (UE::Dataflow::ShouldResetWorld(DataflowGraph, GetWorld(), LastTimeStamp) || EditorContent->IsSimulationDirty())
+			{
+				// Unregister components, cache manager, selection...
+				ResetSimulationScene();
 
-			// Register components, cache manager, selection...
-			CreateSimulationScene();
+				// Register components, cache manager, selection...
+				CreateSimulationScene();
 
-			// Reset the dirty flag
-			EditorContent->SetSimulationDirty(false);
+				// Reset the dirty flag
+				EditorContent->SetSimulationDirty(false);
+			}
 		}
+
 		// Load the cache at some point in time
 		if(SceneDescription->CacheAsset)
 		{
