Dataflow Editor: fix for meshes not rendering in 2D view mode.
- added a 2-sided material for meshes in the Construction view
- disabled the Advanced Preview Floor in orthographic view modes to avoid overly darkening the material.

diff --git a/Engine/Plugins/Experimental/Dataflow/Source/DataflowEditor/Private/Dataflow/DataflowConstructionScene.cpp b/Engine/Plugins/Experimental/Dataflow/Source/DataflowEditor/Private/Dataflow/DataflowConstructionScene.cpp
--- a/Engine/Plugins/Experimental/Dataflow/Source/DataflowEditor/Private/Dataflow/DataflowConstructionScene.cpp
+++ b/Engine/Plugins/Experimental/Dataflow/Source/DataflowEditor/Private/Dataflow/DataflowConstructionScene.cpp
@@ -12,6 +12,7 @@
 #include "Drawing/MeshElementsVisualizer.h"
 #include "Elements/Framework/EngineElementsLibrary.h"
 #include "Selection.h"
+#include "AssetViewerSettings.h"
 
 #define LOCTEXT_NAMESPACE "FDataflowConstructionScene"
 
@@ -221,6 +222,28 @@ void FDataflowConstructionScene::UpdateDynamicMeshComponents()
 					}
 				}
 			}
+
+			// Hide the floor in orthographic view modes
+			if (const UE::Dataflow::IDataflowConstructionViewMode* ConstructionViewMode = EditorContent->GetConstructionViewMode())
+			{
+				if (!ConstructionViewMode->IsPerspective())
+				{
+					constexpr bool bDontModifyProfile = true;
+					SetFloorVisibility(false, bDontModifyProfile);
+				}
+				else
+				{
+					// Restore visibility from profile settings
+					const int32 ProfileIndex = GetCurrentProfileIndex();
+					if (DefaultSettings->Profiles.IsValidIndex(ProfileIndex))
+					{
+						const bool bProfileSetting = DefaultSettings->Profiles[CurrentProfileIndex].bShowFloor;
+						constexpr bool bDontModifyProfile = true;
+						SetFloorVisibility(bProfileSetting, bDontModifyProfile);
+					}
+				}
+			}
+
 		}
 	}
 }
@@ -263,7 +286,8 @@ TObjectPtr<UDynamicMeshComponent>& FDataflowConstructionScene::AddDynamicMeshCom
 	}
 	else
 	{
-		DynamicMeshComponent->SetOverrideRenderMaterial(FDataflowEditorStyle::Get().VertexMaterial);
+		ensure(FDataflowEditorStyle::Get().DefaultTwoSidedMaterial);
+		DynamicMeshComponent->SetOverrideRenderMaterial(FDataflowEditorStyle::Get().DefaultTwoSidedMaterial);
 		DynamicMeshComponent->SetShadowsEnabled(false);
 	}
 	//else if (FDataflowEditorStyle::Get().DefaultMaterial)
diff --git a/Engine/Plugins/Experimental/Dataflow/Source/DataflowEditor/Private/Dataflow/DataflowEditorStyle.h b/Engine/Plugins/Experimental/Dataflow/Source/DataflowEditor/Private/Dataflow/DataflowEditorStyle.h
--- a/Engine/Plugins/Experimental/Dataflow/Source/DataflowEditor/Private/Dataflow/DataflowEditorStyle.h
+++ b/Engine/Plugins/Experimental/Dataflow/Source/DataflowEditor/Private/Dataflow/DataflowEditorStyle.h
@@ -51,6 +51,8 @@ public:
 
 		DefaultMaterial = Cast<UMaterial>(StaticLoadObject(UMaterial::StaticClass(), NULL, TEXT("/Engine/BasicShapes/BasicShapeMaterial")));
 		VertexMaterial = Cast<UMaterial>(StaticLoadObject(UMaterial::StaticClass(), NULL, TEXT("/Engine/EditorMaterials/Dataflow/DataflowVertexMaterial")));
+		DefaultTwoSidedMaterial = Cast<UMaterial>(StaticLoadObject(UMaterial::StaticClass(), NULL, TEXT("/Engine/EditorMaterials/Dataflow/DataflowTwoSidedVertexMaterial")));
+
 		FSlateStyleRegistry::RegisterSlateStyle(*this);
 	}
 
@@ -63,6 +65,7 @@ public:
 	/** Default Rendering Material for Mesh surfaces */
 	UMaterial* DefaultMaterial = nullptr;
 	UMaterial* VertexMaterial = nullptr;
+	UMaterial* DefaultTwoSidedMaterial = nullptr;
 
 
 public:
