Chaos Cache Collection - Fixed the playback of different caches by making the USD file name unique to the terminal asset and cache collection.

#rnx

diff --git a/Engine/Plugins/Experimental/Dataflow/Source/DataflowEditor/Private/Dataflow/DataflowSimulationControls.cpp b/Engine/Plugins/Experimental/Dataflow/Source/DataflowEditor/Private/Dataflow/DataflowSimulationControls.cpp
--- a/Engine/Plugins/Experimental/Dataflow/Source/DataflowEditor/Private/Dataflow/DataflowSimulationControls.cpp
+++ b/Engine/Plugins/Experimental/Dataflow/Source/DataflowEditor/Private/Dataflow/DataflowSimulationControls.cpp
@@ -1,9 +1,10 @@
-﻿// Copyright Epic Games, Inc. All Rights Reserved.
+// Copyright Epic Games, Inc. All Rights Reserved.
 
 #include "Dataflow/DataflowSimulationControls.h"
 
 #include "Animation/AnimSingleNodeInstance.h"
 #include "Animation/AnimSequenceBase.h"
+#include "Chaos/CacheCollection.h"
 #include "Chaos/CacheManagerActor.h"
 #include "Chaos/Adapters/CacheAdapter.h"
 #include "Dataflow/DataflowContent.h"
@@ -43,8 +44,14 @@ namespace UE::Dataflow
 	{
 		if(CacheManager)
 		{
+			const FString BaseName = CacheCollection ? CacheCollection->GetName() : TEXT("CacheActor");
+			const uint32 CacheCollectionPathHash = CacheCollection ? GetTypeHash(CacheCollection->GetPathName()) : 0;
+			const uint32 TerminalAssetPathHash = (DataflowContent && DataflowContent->GetTerminalAsset()) ? GetTypeHash(DataflowContent->GetTerminalAsset()->GetPathName()) : 0;
+			const uint32 CacheActorHash = HashCombineFast(CacheCollectionPathHash, TerminalAssetPathHash);
+			const FString CacheActorName = FString::Printf(TEXT("%s_%08X"), *BaseName, CacheActorHash);
+
 			FActorSpawnParameters SpawnParameters;
-			SpawnParameters.Name = TEXT("CacheActor");
+			SpawnParameters.Name = FName(CacheActorName);
 			SpawnParameters.NameMode = FActorSpawnParameters::ESpawnActorNameMode::Requested;
 			SpawnParameters.Owner = CacheManager.Get(); 
 			SpawnParameters.bDeferConstruction = true;
