Dataflow Editor: fix for not being able to select a mesh in the construction viewport when there is no terminal node in the graph. This was caused by SetLastModifiedTimestamp being called with Invalid on every tick, which rebuilds the construction scene.

diff --git a/Engine/Plugins/Experimental/Dataflow/Source/DataflowEditor/Private/Dataflow/DataflowEditorToolkit.cpp b/Engine/Plugins/Experimental/Dataflow/Source/DataflowEditor/Private/Dataflow/DataflowEditorToolkit.cpp
--- a/Engine/Plugins/Experimental/Dataflow/Source/DataflowEditor/Private/Dataflow/DataflowEditorToolkit.cpp
+++ b/Engine/Plugins/Experimental/Dataflow/Source/DataflowEditor/Private/Dataflow/DataflowEditorToolkit.cpp
@@ -1115,16 +1115,23 @@ void FDataflowEditorToolkit::Tick(float DeltaTime)
 
 							// Take the Max of the existing time stamp, as other terminal nodes might have more recent invalidations
 							const UE::Dataflow::FTimestamp LastModifiedTimestamp = FMath::Max(EditorContent->GetLastModifiedTimestamp(), TerminalNodeTimeStamp);
-							EditorContent->SetLastModifiedTimestamp(LastModifiedTimestamp);
+							
+							constexpr bool bDontMakeDirty = false;
+							EditorContent->SetLastModifiedTimestamp(LastModifiedTimestamp, bDontMakeDirty);
 						}
 					}
 				}
+
+				const bool bMakeDirty = (EditorContent->GetLastModifiedTimestamp() != InitTimeStamp);
+				EditorContent->SetLastModifiedTimestamp(EditorContent->GetLastModifiedTimestamp(), bMakeDirty);
 			}
 			else
 			{
 				UE::Dataflow::FTimestamp TerminalNodeTimeStamp = InitTimeStamp;
 				EvaluateNode(nullptr, nullptr, TerminalNodeTimeStamp);
-				EditorContent->SetLastModifiedTimestamp(TerminalNodeTimeStamp);
+
+				const bool bMakeDirty = (TerminalNodeTimeStamp != InitTimeStamp);
+				EditorContent->SetLastModifiedTimestamp(TerminalNodeTimeStamp, bMakeDirty);
 			}
 
 			// Ensure the context object's selected node matches the selected node in the graph editor
diff --git a/Engine/Source/Runtime/Experimental/Dataflow/Core/Public/Dataflow/DataflowNodeParameters.h b/Engine/Source/Runtime/Experimental/Dataflow/Core/Public/Dataflow/DataflowNodeParameters.h
--- a/Engine/Source/Runtime/Experimental/Dataflow/Core/Public/Dataflow/DataflowNodeParameters.h
+++ b/Engine/Source/Runtime/Experimental/Dataflow/Core/Public/Dataflow/DataflowNodeParameters.h
@@ -38,6 +38,7 @@ namespace UE::Dataflow
 		FTimestamp(Type InValue) : Value(InValue) {}
 		bool operator>=(const FTimestamp& InTimestamp) const { return Value >= InTimestamp.Value; }
 		bool operator<(const FTimestamp& InTimestamp) const { return Value < InTimestamp.Value; }
+		bool operator==(const FTimestamp& InTimestamp) const { return Value == InTimestamp.Value; }
 		bool IsInvalid() { return Value == Invalid; }
 
 		static DATAFLOWCORE_API Type Current();
