Cloth Editor and Dataflow Editor: fix for toolbar actions running in the wrong editor instance. We should register toolbar commands in the EditorToolkit, not the EditorModeToolkit

diff --git a/Engine/Plugins/ChaosClothAssetEditor/Source/ChaosClothAssetEditor/Private/ChaosClothAsset/ClothEditorMode.cpp b/Engine/Plugins/ChaosClothAssetEditor/Source/ChaosClothAssetEditor/Private/ChaosClothAsset/ClothEditorMode.cpp
--- a/Engine/Plugins/ChaosClothAssetEditor/Source/ChaosClothAssetEditor/Private/ChaosClothAsset/ClothEditorMode.cpp
+++ b/Engine/Plugins/ChaosClothAssetEditor/Source/ChaosClothAssetEditor/Private/ChaosClothAsset/ClothEditorMode.cpp
@@ -298,35 +298,31 @@ void UChaosClothAssetEditorMode::RegisterClothTool(TSharedPtr<FUICommandInfo> UI
 
 }
 
-void UChaosClothAssetEditorMode::RegisterAddNodeCommand(TSharedPtr<FUICommandInfo> AddNodeCommand, const FName& NewNodeType, TSharedPtr<FUICommandInfo> StartToolCommand)
+void UChaosClothAssetEditorMode::AddNode(FName NewNodeType)
 {
-	auto AddNode = [this](const FName& NewNodeType)
-	{
-		const FName ConnectionType = FManagedArrayCollection::StaticType();
-		const FName ConnectionName("Collection");
-
-		UEdGraphNode* const CurrentlySelectedNode = GetSingleSelectedNodeWithOutputType(ConnectionType);
-		checkf(CurrentlySelectedNode, TEXT("No node with FManagedArrayCollection output is currently selected in the Dataflow graph"));
+	const FName ConnectionType = FManagedArrayCollection::StaticType();
+	const FName ConnectionName("Collection");
 
-		const UEdGraphNode* const NewNode = CreateAndConnectNewNode(NewNodeType, *CurrentlySelectedNode, ConnectionType, ConnectionName);
-		verifyf(NewNode, TEXT("Failed to create a new node: %s"), *NewNodeType.ToString());
+	UEdGraphNode* const CurrentlySelectedNode = GetSingleSelectedNodeWithOutputType(ConnectionType);
+	checkf(CurrentlySelectedNode, TEXT("No node with FManagedArrayCollection output is currently selected in the Dataflow graph"));
 
-		StartToolForSelectedNode(NewNode);
-	};
+	const UEdGraphNode* const NewNode = CreateAndConnectNewNode(NewNodeType, *CurrentlySelectedNode, ConnectionType, ConnectionName);
+	verifyf(NewNode, TEXT("Failed to create a new node: %s"), *NewNodeType.ToString());
 
-	auto CanAddNode = [this](const FName& NewNodeType) -> bool
-	{
-		const UEdGraphNode* const CurrentlySelectedNode = GetSingleSelectedNodeWithOutputType(FManagedArrayCollection::StaticType());
-		return (CurrentlySelectedNode != nullptr);
-	};
+	StartToolForSelectedNode(NewNode);
+}
 
-	const TSharedRef<FUICommandList>& CommandList = Toolkit->GetToolkitCommands();
+bool UChaosClothAssetEditorMode::CanAddNode(FName NewNodeType) const
+{
+	const UEdGraphNode* const CurrentlySelectedNode = GetSingleSelectedNodeWithOutputType(FManagedArrayCollection::StaticType());
+	return (CurrentlySelectedNode != nullptr);
+}
 
-	CommandList->MapAction(AddNodeCommand,
-		FExecuteAction::CreateWeakLambda(this, AddNode, NewNodeType),
-		FCanExecuteAction::CreateWeakLambda(this, CanAddNode, NewNodeType)
-	);
+void UChaosClothAssetEditorMode::RegisterAddNodeCommand(TSharedPtr<FUICommandInfo> AddNodeCommand, const FName& NewNodeType, TSharedPtr<FUICommandInfo> StartToolCommand)
+{
+	// ToolkitCommands->MapAction(AddNodeCommand) is done in FChaosClothAssetEditorToolkit
 
+	NodeTypeToAddNodeCommandMap.Add(NewNodeType, AddNodeCommand);
 	NodeTypeToToolCommandMap.Add(NewNodeType, StartToolCommand);
 }
 
diff --git a/Engine/Plugins/ChaosClothAssetEditor/Source/ChaosClothAssetEditor/Private/ChaosClothAsset/ClothEditorMode.h b/Engine/Plugins/ChaosClothAssetEditor/Source/ChaosClothAssetEditor/Private/ChaosClothAsset/ClothEditorMode.h
--- a/Engine/Plugins/ChaosClothAssetEditor/Source/ChaosClothAssetEditor/Private/ChaosClothAsset/ClothEditorMode.h
+++ b/Engine/Plugins/ChaosClothAssetEditor/Source/ChaosClothAssetEditor/Private/ChaosClothAsset/ClothEditorMode.h
@@ -223,6 +223,9 @@ private:
 	*/
 	UEdGraphNode* CreateAndConnectNewNode(const FName& NewNodeTypeName,	UEdGraphNode& UpstreamNode,	const FName& ConnectionTypeName, const FName& NewNodeConnectionName);
 
+	void AddNode(FName NewNodeType);
+	bool CanAddNode(FName NewNodeType) const;
+
 
 	void InitializeContextObject();
 	void UpdateContextObject(const TSharedPtr<FManagedArrayCollection>& Collection);
@@ -314,6 +317,9 @@ private:
 	TSharedPtr<FManagedArrayCollection> SelectedClothCollection = nullptr;
 	TSharedPtr<FManagedArrayCollection> SelectedInputClothCollection = nullptr;
 
+	// Correspondence between node types and commands to add the node to the graph
+	TMap<FName, TSharedPtr<const FUICommandInfo>> NodeTypeToAddNodeCommandMap;
+
 	// Correspondence between node types and commands to launch tools
 	TMap<FName, TSharedPtr<const FUICommandInfo>> NodeTypeToToolCommandMap;
 
diff --git a/Engine/Plugins/ChaosClothAssetEditor/Source/ChaosClothAssetEditor/Private/ChaosClothAsset/ClothEditorModeToolkit.cpp b/Engine/Plugins/ChaosClothAssetEditor/Source/ChaosClothAssetEditor/Private/ChaosClothAsset/ClothEditorModeToolkit.cpp
--- a/Engine/Plugins/ChaosClothAssetEditor/Source/ChaosClothAssetEditor/Private/ChaosClothAsset/ClothEditorModeToolkit.cpp
+++ b/Engine/Plugins/ChaosClothAssetEditor/Source/ChaosClothAssetEditor/Private/ChaosClothAsset/ClothEditorModeToolkit.cpp
@@ -141,25 +141,6 @@ void FChaosClothAssetEditorModeToolkit::Init(const TSharedPtr<IToolkitHost>& Ini
 	];
 }
 
-void FChaosClothAssetEditorModeToolkit::BuildEditorToolBar(const FName& EditorToolBarName)
-{
-	check(EditorToolBarName != FName());
-
-	const TSharedRef<const FUICommandList> EdModeToolkitCommands = GetToolkitCommands();
-
-	UToolMenu* const ToolBarMenu = UToolMenus::Get()->ExtendMenu(EditorToolBarName);
-	FToolMenuSection& Section = ToolBarMenu->FindOrAddSection("ClothTools");
-
-	FToolMenuEntry& WeightMapButtonEntry = Section.AddEntry(FToolMenuEntry::InitToolBarButton(FChaosClothAssetEditorCommands::Get().AddWeightMapNode));
-	WeightMapButtonEntry.SetCommandList(EdModeToolkitCommands);
-
-	FToolMenuEntry& SkinWeightsButtonEntry = Section.AddEntry(FToolMenuEntry::InitToolBarButton(FChaosClothAssetEditorCommands::Get().AddTransferSkinWeightsNode));
-	SkinWeightsButtonEntry.SetCommandList(EdModeToolkitCommands);
-
-	FToolMenuEntry& SelectionButtonEntry = Section.AddEntry(FToolMenuEntry::InitToolBarButton(FChaosClothAssetEditorCommands::Get().AddMeshSelectionNode));
-	SelectionButtonEntry.SetCommandList(EdModeToolkitCommands);
-}
-
 const FSlateBrush* FChaosClothAssetEditorModeToolkit::GetActiveToolIcon(const FString& ActiveToolIdentifier) const
 {
 	FName ActiveToolIconName = ISlateStyle::Join(FChaosClothAssetEditorCommands::Get().GetContextName(), TCHAR_TO_ANSI(*ActiveToolIdentifier));
diff --git a/Engine/Plugins/ChaosClothAssetEditor/Source/ChaosClothAssetEditor/Private/ChaosClothAsset/ClothEditorModeToolkit.h b/Engine/Plugins/ChaosClothAssetEditor/Source/ChaosClothAssetEditor/Private/ChaosClothAsset/ClothEditorModeToolkit.h
--- a/Engine/Plugins/ChaosClothAssetEditor/Source/ChaosClothAssetEditor/Private/ChaosClothAsset/ClothEditorModeToolkit.h
+++ b/Engine/Plugins/ChaosClothAssetEditor/Source/ChaosClothAssetEditor/Private/ChaosClothAsset/ClothEditorModeToolkit.h
@@ -23,10 +23,6 @@ public:
 
 	void Init(const TSharedPtr<IToolkitHost>& InitToolkitHost, TWeakObjectPtr<UEdMode> InOwningMode);
 
-	// Add buttons to the toolbar with the specified name
-	// Note: Most FModeToolkits would define BuildToolPalette, but we are putting buttons in the top toolbar instead
-	void BuildEditorToolBar(const FName& EditorToolBarName);
-
 	virtual const FSlateBrush* GetActiveToolIcon(const FString& Identifier) const override;
 
 	virtual void OnToolStarted(UInteractiveToolManager* Manager, UInteractiveTool* Tool) override;
diff --git a/Engine/Plugins/ChaosClothAssetEditor/Source/ChaosClothAssetEditor/Private/ChaosClothAsset/ClothEditorToolkit.cpp b/Engine/Plugins/ChaosClothAssetEditor/Source/ChaosClothAssetEditor/Private/ChaosClothAsset/ClothEditorToolkit.cpp
--- a/Engine/Plugins/ChaosClothAssetEditor/Source/ChaosClothAssetEditor/Private/ChaosClothAsset/ClothEditorToolkit.cpp
+++ b/Engine/Plugins/ChaosClothAssetEditor/Source/ChaosClothAssetEditor/Private/ChaosClothAsset/ClothEditorToolkit.cpp
@@ -45,6 +45,7 @@
 #include "UObject/PackageReload.h"
 #include "ContextObjectStore.h"
 #include "SClothEditorAdvancedPreviewDetailsTab.h"
+#include "ToolMenus.h"
 
 #define LOCTEXT_NAMESPACE "ChaosClothAssetEditorToolkit"
 
@@ -304,10 +305,6 @@ void FChaosClothAssetEditorToolkit::InitializeEdMode(UBaseCharacterFXEditorMode*
 	{
 		FChaosClothAssetEditorModeToolkit* ClothModeToolkit = static_cast<FChaosClothAssetEditorModeToolkit*>(ModeToolkit.Get());
 		ClothModeToolkit->SetRestSpaceViewportWidget(RestSpaceViewportWidget);
-
-		FName ParentToolbarName;
-		const FName ToolBarName = GetToolMenuToolbarName(ParentToolbarName);
-		ClothModeToolkit->BuildEditorToolBar(ToolBarName);
 	}
 }
 
@@ -508,6 +505,22 @@ void FChaosClothAssetEditorToolkit::PostInitAssetEditor()
 
 	// Handle Dataflow asset reload event
 	OnPackageReloadedDelegateHandle = FCoreUObjectDelegates::OnPackageReloaded.AddSP(this, &FChaosClothAssetEditorToolkit::HandlePackageReloaded);
+
+	// Populate editor toolbar
+
+	FName ParentToolbarName;
+	const FName ToolBarName = GetToolMenuToolbarName(ParentToolbarName);
+	UToolMenu* const AssetToolbar = UToolMenus::Get()->ExtendMenu(ToolBarName);
+	FToolMenuSection& Section = AssetToolbar->FindOrAddSection("ClothTools");
+
+	for (const TPair<FName, TSharedPtr<const FUICommandInfo>>& NodeAndAddCommand : ClothMode->NodeTypeToAddNodeCommandMap)
+	{
+		ToolkitCommands->MapAction(NodeAndAddCommand.Value,
+			FExecuteAction::CreateUObject(ClothMode, &UChaosClothAssetEditorMode::AddNode, NodeAndAddCommand.Key),
+			FCanExecuteAction::CreateUObject(ClothMode, &UChaosClothAssetEditorMode::CanAddNode, NodeAndAddCommand.Key));
+
+		Section.AddEntry(FToolMenuEntry::InitToolBarButton(NodeAndAddCommand.Value));
+	}
 }
 
 void FChaosClothAssetEditorToolkit::InitToolMenuContext(FToolMenuContext& MenuContext)
diff --git a/Engine/Plugins/Experimental/Dataflow/Source/DataflowEditor/Private/Dataflow/DataflowEditorMode.cpp b/Engine/Plugins/Experimental/Dataflow/Source/DataflowEditor/Private/Dataflow/DataflowEditorMode.cpp
--- a/Engine/Plugins/Experimental/Dataflow/Source/DataflowEditor/Private/Dataflow/DataflowEditorMode.cpp
+++ b/Engine/Plugins/Experimental/Dataflow/Source/DataflowEditor/Private/Dataflow/DataflowEditorMode.cpp
@@ -229,34 +229,24 @@ void UDataflowEditorMode::RegisterDataflowTool(TSharedPtr<FUICommandInfo> UIComm
 	);
 }
 
-void UDataflowEditorMode::RegisterAddNodeCommand(TSharedPtr<FUICommandInfo> AddNodeCommand, const FName& NewNodeType)
+void UDataflowEditorMode::AddNode(FName NewNodeType)
 {
-	auto AddNode = [this](const FName& NewNodeType)
-	{
-		const FName ConnectionType = FManagedArrayCollection::StaticType();
-		const FName ConnectionName("Collection");
-
-		UEdGraphNode* const CurrentlySelectedNode = GetSingleSelectedNodeWithOutputType(ConnectionType);
-		checkf(CurrentlySelectedNode, TEXT("No node with FManagedArrayCollection output is currently selected in the Dataflow graph"));
+	const FName ConnectionType = FManagedArrayCollection::StaticType();
+	const FName ConnectionName("Collection");
 
-		const UEdGraphNode* const NewNode = CreateAndConnectNewNode(NewNodeType, *CurrentlySelectedNode, ConnectionType, ConnectionName);
-		verifyf(NewNode, TEXT("Failed to create a new node: %s"), *NewNodeType.ToString());
+	UEdGraphNode* const CurrentlySelectedNode = GetSingleSelectedNodeWithOutputType(ConnectionType);
+	checkf(CurrentlySelectedNode, TEXT("No node with FManagedArrayCollection output is currently selected in the Dataflow graph"));
 
-		StartToolForSelectedNode(NewNode);
-	};
-
-	auto CanAddNode = [this](const FName& NewNodeType) -> bool
-	{
-		const UEdGraphNode* const CurrentlySelectedNode = GetSingleSelectedNodeWithOutputType(FManagedArrayCollection::StaticType());
-		return (CurrentlySelectedNode != nullptr);
-	};
+	const UEdGraphNode* const NewNode = CreateAndConnectNewNode(NewNodeType, *CurrentlySelectedNode, ConnectionType, ConnectionName);
+	verifyf(NewNode, TEXT("Failed to create a new node: %s"), *NewNodeType.ToString());
 
-	const TSharedRef<FUICommandList>& CommandList = Toolkit->GetToolkitCommands();
+	StartToolForSelectedNode(NewNode);
+}
 
-	CommandList->MapAction(AddNodeCommand,
-		FExecuteAction::CreateWeakLambda(this, AddNode, NewNodeType),
-		FCanExecuteAction::CreateWeakLambda(this, CanAddNode, NewNodeType)
-	);
+bool UDataflowEditorMode::CanAddNode(FName NewNodeType) const
+{
+	const UEdGraphNode* const CurrentlySelectedNode = GetSingleSelectedNodeWithOutputType(FManagedArrayCollection::StaticType());
+	return (CurrentlySelectedNode != nullptr);
 }
 
 void UDataflowEditorMode::RegisterTools()
@@ -283,8 +273,8 @@ void UDataflowEditorMode::RegisterTools()
 		NodeTypeToToolCommandMap.Add(RegisteredNodeName, CommandInfo);
 	}
 
-	// Register "Add Node" commands for buttons in the UI. So far we just have the one
-	RegisterAddNodeCommand(CommandInfos.AddWeightMapNode, FDataflowCollectionAddScalarVertexPropertyNode::StaticType());
+	// Register "Add Node" commands for buttons in the UI. The EditorToolkit will construct the actual toolbar buttons.
+	NodeTypeToAddNodeCommandMap.Add(FDataflowCollectionAddScalarVertexPropertyNode::StaticType(), CommandInfos.AddWeightMapNode);
 }
 
 bool UDataflowEditorMode::ShouldToolStartBeAllowed(const FString& ToolIdentifier) const
diff --git a/Engine/Plugins/Experimental/Dataflow/Source/DataflowEditor/Private/Dataflow/DataflowEditorModeToolkit.cpp b/Engine/Plugins/Experimental/Dataflow/Source/DataflowEditor/Private/Dataflow/DataflowEditorModeToolkit.cpp
--- a/Engine/Plugins/Experimental/Dataflow/Source/DataflowEditor/Private/Dataflow/DataflowEditorModeToolkit.cpp
+++ b/Engine/Plugins/Experimental/Dataflow/Source/DataflowEditor/Private/Dataflow/DataflowEditorModeToolkit.cpp
@@ -133,19 +133,6 @@ void FDataflowEditorModeToolkit::Init(const TSharedPtr<IToolkitHost>& InitToolki
 	];
 }
 
-void FDataflowEditorModeToolkit::BuildEditorToolBar(const FName& EditorToolBarName)
-{
-	check(EditorToolBarName != FName());
-
-	const TSharedRef<const FUICommandList> EdModeToolkitCommands = GetToolkitCommands();
-
-	UToolMenu* const ToolBarMenu = UToolMenus::Get()->ExtendMenu(EditorToolBarName);
-	FToolMenuSection& Section = ToolBarMenu->FindOrAddSection("DataflowTools");
-
-	FToolMenuEntry& WeightMapButtonEntry = Section.AddEntry(FToolMenuEntry::InitToolBarButton(FDataflowEditorCommandsImpl::Get().AddWeightMapNode));
-	WeightMapButtonEntry.SetCommandList(EdModeToolkitCommands);
-}
-
 const FSlateBrush* FDataflowEditorModeToolkit::GetActiveToolIcon(const FString& ActiveToolIdentifier) const
 {
 	FName ActiveToolIconName = ISlateStyle::Join(FDataflowEditorCommandsImpl::Get().GetContextName(), TCHAR_TO_ANSI(*ActiveToolIdentifier));
diff --git a/Engine/Plugins/Experimental/Dataflow/Source/DataflowEditor/Private/Dataflow/DataflowEditorToolkit.cpp b/Engine/Plugins/Experimental/Dataflow/Source/DataflowEditor/Private/Dataflow/DataflowEditorToolkit.cpp
--- a/Engine/Plugins/Experimental/Dataflow/Source/DataflowEditor/Private/Dataflow/DataflowEditorToolkit.cpp
+++ b/Engine/Plugins/Experimental/Dataflow/Source/DataflowEditor/Private/Dataflow/DataflowEditorToolkit.cpp
@@ -55,7 +55,7 @@
 #include "Styling/SlateStyleRegistry.h"
 #include "Widgets/Docking/SDockTab.h"
 #include "GeometryCache.h"
-
+#include "ToolMenus.h"
 
 #define LOCTEXT_NAMESPACE "DataflowEditorToolkit"
 
@@ -499,6 +499,22 @@ void FDataflowEditorToolkit::PostInitAssetEditor()
 
 	FDataflowConstructionViewportClient* ConstructionViewportClient = static_cast<FDataflowConstructionViewportClient*>(ViewportClient.Get());
 	OnConstructionSelectionChangedDelegateHandle = ConstructionViewportClient->OnSelectionChangedMulticast.AddSP(this, &FDataflowEditorToolkit::OnConstructionViewSelectionChanged);
+
+	// Populate editor toolbar
+
+	FName ParentToolbarName;
+	const FName ToolBarName = GetToolMenuToolbarName(ParentToolbarName);
+	UToolMenu* const AssetToolbar = UToolMenus::Get()->ExtendMenu(ToolBarName);
+	FToolMenuSection& Section = AssetToolbar->FindOrAddSection("ClothTools");
+
+	for (const TPair<FName, TSharedPtr<const FUICommandInfo>>& NodeAndAddCommand : DataflowMode->NodeTypeToAddNodeCommandMap)
+	{
+		ToolkitCommands->MapAction(NodeAndAddCommand.Value,
+			FExecuteAction::CreateUObject(DataflowMode, &UDataflowEditorMode::AddNode, NodeAndAddCommand.Key),
+			FCanExecuteAction::CreateUObject(DataflowMode, &UDataflowEditorMode::CanAddNode, NodeAndAddCommand.Key));
+
+		Section.AddEntry(FToolMenuEntry::InitToolBarButton(NodeAndAddCommand.Value));
+	}
 }
 
 void FDataflowEditorToolkit::InitializeEdMode(UBaseCharacterFXEditorMode* EdMode)
@@ -522,10 +538,6 @@ void FDataflowEditorToolkit::InitializeEdMode(UBaseCharacterFXEditorMode* EdMode
 		FDataflowEditorModeToolkit* DataflowModeToolkit = static_cast<FDataflowEditorModeToolkit*>(ModeToolkit.Get());
 		DataflowModeToolkit->SetConstructionViewportWidget(DataflowConstructionViewport);
 		DataflowModeToolkit->SetSimulationViewportWidget(DataflowSimulationViewport);
-
-		FName ParentToolbarName;
-		const FName ToolBarName = GetToolMenuToolbarName(ParentToolbarName);
-		DataflowModeToolkit->BuildEditorToolBar(ToolBarName);
 	}
 
 	// @todo(brice) : This used to crash when comnmented out. 
diff --git a/Engine/Plugins/Experimental/Dataflow/Source/DataflowEditor/Public/Dataflow/DataflowEditorMode.h b/Engine/Plugins/Experimental/Dataflow/Source/DataflowEditor/Public/Dataflow/DataflowEditorMode.h
--- a/Engine/Plugins/Experimental/Dataflow/Source/DataflowEditor/Public/Dataflow/DataflowEditorMode.h
+++ b/Engine/Plugins/Experimental/Dataflow/Source/DataflowEditor/Public/Dataflow/DataflowEditorMode.h
@@ -117,8 +117,9 @@ private:
 							  UEditorInteractiveToolsContext* const ToolsContext,
 							  EToolsContextScope ToolScope = EToolsContextScope::Default);
 
-	// Command for quickly adding a new node to the graph, e.g. triggered by a button press
-	void RegisterAddNodeCommand(TSharedPtr<FUICommandInfo> AddNodeCommand, const FName& NewNodeType);
+	// Functions for quickly adding a node to the graph, e.g. when triggered by a toolbar button
+	void AddNode(FName NewNodeType);
+	bool CanAddNode(FName NewNodeType) const;
 
 	void SetSimulationViewportClient(TWeakPtr<FDataflowSimulationViewportClient, ESPMode::ThreadSafe>);
 	void RefocusSimulationViewportClient();
@@ -173,6 +174,9 @@ private:
 	/** Dataflow simulation scene from the toolkit */
 	FDataflowSimulationScene* SimulationScene = nullptr;
 
+	/** Correspondence between node types and commands to add the node to the graph */
+	TMap<FName, TSharedPtr<const FUICommandInfo>> NodeTypeToAddNodeCommandMap;
+
 	/** Correspondence between node types and commands to launch tools */
 	TMap<FName, TSharedPtr<const FUICommandInfo>> NodeTypeToToolCommandMap;
 
diff --git a/Engine/Plugins/Experimental/Dataflow/Source/DataflowEditor/Public/Dataflow/DataflowEditorModeToolkit.h b/Engine/Plugins/Experimental/Dataflow/Source/DataflowEditor/Public/Dataflow/DataflowEditorModeToolkit.h
--- a/Engine/Plugins/Experimental/Dataflow/Source/DataflowEditor/Public/Dataflow/DataflowEditorModeToolkit.h
+++ b/Engine/Plugins/Experimental/Dataflow/Source/DataflowEditor/Public/Dataflow/DataflowEditorModeToolkit.h
@@ -21,10 +21,6 @@ class DATAFLOWEDITOR_API FDataflowEditorModeToolkit : public FBaseCharacterFXEdi
 public:
 	void Init(const TSharedPtr<IToolkitHost>& InitToolkitHost, TWeakObjectPtr<UEdMode> InOwningMode);
 
-	// Add buttons to the toolbar with the specified name
-	// Note: Most FModeToolkits would define BuildToolPalette, but we are putting buttons in the top toolbar instead
-	void BuildEditorToolBar(const FName& EditorToolBarName);
-
 	// FBaseCharacterFXEditorModeToolkit interface
 	virtual const FSlateBrush* GetActiveToolIcon(const FString& Identifier) const override;
 
