Dataflow Editor: remove log output whenever a Geometry Collection is rendered in the construction view

diff --git a/Engine/Plugins/Experimental/Dataflow/Source/DataflowEditor/Private/Dataflow/DataflowEngineRendering.cpp b/Engine/Plugins/Experimental/Dataflow/Source/DataflowEditor/Private/Dataflow/DataflowEngineRendering.cpp
--- a/Engine/Plugins/Experimental/Dataflow/Source/DataflowEditor/Private/Dataflow/DataflowEngineRendering.cpp
+++ b/Engine/Plugins/Experimental/Dataflow/Source/DataflowEditor/Private/Dataflow/DataflowEngineRendering.cpp
@@ -276,7 +276,6 @@ namespace UE::Dataflow
 				const bool bFoundTransforms = Collection.FindAttributeTyped<FTransform3f>(FTransformCollection::TransformAttribute, FTransformCollection::TransformGroup) != nullptr;
 				const bool bFoundBoneMap = Collection.FindAttributeTyped<int32>("BoneMap", FGeometryCollection::VerticesGroup) != nullptr;
 				const bool bFoundParents = Collection.FindAttributeTyped<int32>(FTransformCollection::ParentAttribute, FTransformCollection::TransformGroup) != nullptr;
-				UE_LOG(LogTemp, Warning, TEXT("Render GC with found params = %d %d %d %d %d"), bFoundIndices, bFoundVertices, bFoundTransforms, bFoundBoneMap, bFoundParents);
 				bool bFoundRenderData = bFoundIndices && bFoundVertices && bFoundTransforms && bFoundBoneMap && bFoundParents
 					&& Collection.NumElements(FTransformCollection::TransformGroup) > 0;
 
@@ -284,7 +283,6 @@ namespace UE::Dataflow
 				const bool bFoundVertexCount = Collection.FindAttributeTyped<int32>("VertexCount", FGeometryCollection::GeometryGroup) != nullptr;
 				const bool bFoundFaceStart = Collection.FindAttributeTyped<int32>("FaceStart", FGeometryCollection::GeometryGroup) != nullptr;
 				const bool bFoundFaceCount = Collection.FindAttributeTyped<int32>("FaceCount", FGeometryCollection::GeometryGroup) != nullptr;
-				UE_LOG(LogTemp, Warning, TEXT("Render GC with found mesh group params = %d %d %d %d"), bFoundVertexStart, bFoundVertexCount, bFoundFaceStart, bFoundFaceCount);
 				bool bFoundGeometryAttributes = bFoundVertexStart && bFoundVertexCount && bFoundFaceStart && bFoundFaceCount
 					&& Collection.NumElements(FGeometryCollection::GeometryGroup) > 0;
 
