Cloth Editor: don't add toggle widgets to the nodes -- the toggle does nothing in the Cloth Editor

diff --git a/Engine/Plugins/ChaosClothAssetEditor/Source/ChaosClothAssetEditor/Private/ChaosClothAsset/ClothEditorToolkit.cpp b/Engine/Plugins/ChaosClothAssetEditor/Source/ChaosClothAssetEditor/Private/ChaosClothAsset/ClothEditorToolkit.cpp
--- a/Engine/Plugins/ChaosClothAssetEditor/Source/ChaosClothAssetEditor/Private/ChaosClothAsset/ClothEditorToolkit.cpp
+++ b/Engine/Plugins/ChaosClothAssetEditor/Source/ChaosClothAssetEditor/Private/ChaosClothAsset/ClothEditorToolkit.cpp
@@ -79,6 +79,12 @@ public:
 	{
 		return ensure(ClothAssetEditorToolkit) ? ClothAssetEditorToolkit->GetDataflowContext() : TSharedPtr<UE::Dataflow::FContext>();
 	}
+
+	virtual bool NodesHaveToggleWidget() const override
+	{
+		return false;
+	}
+
 	//~ End FDataflowSNodeInterface interface
 
 private:
diff --git a/Engine/Plugins/Experimental/Dataflow/Source/DataflowEditor/Private/Dataflow/DataflowSNode.cpp b/Engine/Plugins/Experimental/Dataflow/Source/DataflowEditor/Private/Dataflow/DataflowSNode.cpp
--- a/Engine/Plugins/Experimental/Dataflow/Source/DataflowEditor/Private/Dataflow/DataflowSNode.cpp
+++ b/Engine/Plugins/Experimental/Dataflow/Source/DataflowEditor/Private/Dataflow/DataflowSNode.cpp
@@ -180,7 +180,7 @@ TArray<FOverlayWidgetInfo> SDataflowEdNode::GetOverlayWidgets(bool bSelected, co
 {
 	TArray<FOverlayWidgetInfo> Widgets = SGraphNode::GetOverlayWidgets(bSelected, WidgetSize);
 
-	if (DataflowGraphNode && DataflowGraphNode->GetDataflowNode() && DataflowGraphNode->GetDataflowNode()->GetRenderParameters().Num())
+	if (DataflowGraphNode && DataflowInterface->NodesHaveToggleWidget() && DataflowGraphNode->GetDataflowNode() && DataflowGraphNode->GetDataflowNode()->GetRenderParameters().Num())
 	{
 		const FVector2D ImageSize = RenderCheckBoxWidget->GetDesiredSize();
 
diff --git a/Engine/Plugins/Experimental/Dataflow/Source/DataflowEditor/Public/Dataflow/DataflowSEditorInterface.h b/Engine/Plugins/Experimental/Dataflow/Source/DataflowEditor/Public/Dataflow/DataflowSEditorInterface.h
--- a/Engine/Plugins/Experimental/Dataflow/Source/DataflowEditor/Public/Dataflow/DataflowSEditorInterface.h
+++ b/Engine/Plugins/Experimental/Dataflow/Source/DataflowEditor/Public/Dataflow/DataflowSEditorInterface.h
@@ -19,6 +19,8 @@ public:
 	/** Dataflow editor content accessors */
 	virtual TSharedPtr<UE::Dataflow::FContext> GetDataflowContext() const = 0;
 
+	virtual bool NodesHaveToggleWidget() const { return true; }
+
 protected:
 	FDataflowSEditorInterface() = default;
 };
