ChaosFlesh: added visibility control for SKM in dataflow simulation

diff --git a/Engine/Plugins/Experimental/Dataflow/Source/DataflowEditor/Private/Dataflow/DataflowSimulationControls.cpp b/Engine/Plugins/Experimental/Dataflow/Source/DataflowEditor/Private/Dataflow/DataflowSimulationControls.cpp
--- a/Engine/Plugins/Experimental/Dataflow/Source/DataflowEditor/Private/Dataflow/DataflowSimulationControls.cpp
+++ b/Engine/Plugins/Experimental/Dataflow/Source/DataflowEditor/Private/Dataflow/DataflowSimulationControls.cpp
@@ -93,7 +93,7 @@ namespace UE::Dataflow
 		return nullptr;
 	}
 	
-	void SetupSkeletonAnimation(const TObjectPtr<AActor>& PreviewActor)
+	void SetupSkeletonAnimation(const TObjectPtr<AActor>& PreviewActor, const bool bSkeletalMeshVisibility)
 	{
 		if(PreviewActor)
 		{
@@ -104,6 +104,7 @@ namespace UE::Dataflow
 			{
 				if(USkeletalMeshComponent* SkeletalMeshComponent = Cast<USkeletalMeshComponent>(PrimComponent))
 				{
+					SkeletalMeshComponent->SetVisibility(bSkeletalMeshVisibility);
 					SkeletalMeshComponent->SetAnimationMode(EAnimationMode::AnimationSingleNode);
 					SkeletalMeshComponent->InitAnim(true);
 					
diff --git a/Engine/Plugins/Experimental/Dataflow/Source/DataflowEditor/Private/Dataflow/DataflowSimulationGenerator.cpp b/Engine/Plugins/Experimental/Dataflow/Source/DataflowEditor/Private/Dataflow/DataflowSimulationGenerator.cpp
--- a/Engine/Plugins/Experimental/Dataflow/Source/DataflowEditor/Private/Dataflow/DataflowSimulationGenerator.cpp
+++ b/Engine/Plugins/Experimental/Dataflow/Source/DataflowEditor/Private/Dataflow/DataflowSimulationGenerator.cpp
@@ -1,4 +1,4 @@
-﻿// Copyright Epic Games, Inc. All Rights Reserved.
+// Copyright Epic Games, Inc. All Rights Reserved.
 
 #include "Dataflow/DataflowSimulationGenerator.h"
 #include "Dataflow/DataflowSimulationManager.h"
@@ -65,7 +65,7 @@ void FDataflowSimulationTask::DoWork()
 
 bool FDataflowTaskManager::AllocateSimulationResource(const FVector2f& TimeRange, const int32 FrameRate,
 		const TObjectPtr<UChaosCacheCollection>& CacheAsset, const TSubclassOf<AActor>& ActorClass,
-		const TObjectPtr<UDataflowBaseContent>& DataflowContent, const FTransform& BlueprintTransform)
+		const TObjectPtr<UDataflowBaseContent>& DataflowContent, const FTransform& BlueprintTransform, const bool bSkeletalMeshVisibility)
 {
 	SimulationWorld = UWorld::CreateWorld(EWorldType::Editor, false);
 	SimulationWorld->bPostTickComponentUpdate = false;
@@ -76,7 +76,7 @@ bool FDataflowTaskManager::AllocateSimulationResource(const FVector2f& TimeRange
 	CacheManager = SimulationWorld->SpawnActor<AChaosCacheManager>();
 
 	PreviewActor = UE::Dataflow::SpawnSimulatedActor(ActorClass, CacheManager, CacheAsset, true, DataflowContent, BlueprintTransform);
-	UE::Dataflow::SetupSkeletonAnimation(PreviewActor);
+	UE::Dataflow::SetupSkeletonAnimation(PreviewActor, bSkeletalMeshVisibility);
 
 	// Init the cache manager
 	CacheManager->SetObservedComponentProperties(CacheManager->CacheMode);
@@ -190,7 +190,7 @@ void FDataflowSimulationGenerator::StartGenerateSimulation()
 	TaskManager->SimulationTask->GetTask().TaskManager = TaskManager;
 	TaskManager->SimulationTask->GetTask().bAsyncCaching = CacheParams.bAsyncCaching;
 	
-	TaskManager->AllocateSimulationResource(CacheParams.TimeRange, CacheParams.FrameRate, CacheAsset, BlueprintClass, DataflowContent, BlueprintTransform);
+	TaskManager->AllocateSimulationResource(CacheParams.TimeRange, CacheParams.FrameRate, CacheAsset, BlueprintClass, DataflowContent, BlueprintTransform, bSkeletalMeshVisibility);
 
 	if(CacheParams.bAsyncCaching)
 	{
@@ -278,6 +278,11 @@ void FDataflowSimulationGenerator::SetDataflowContent(const TObjectPtr<UDataflow
 	DataflowContent = InDataflowContent;
 }
 
+void FDataflowSimulationGenerator::SetSkeletalMeshVisibility(const bool bInSkeletalMeshVisibility)
+{
+	bSkeletalMeshVisibility = bInSkeletalMeshVisibility;
+}
+
 void FDataflowSimulationGenerator::RequestGeneratorAction(EDataflowGeneratorActions ActionType)
 {
 	if (PendingAction != EDataflowGeneratorActions::NoAction)
diff --git a/Engine/Plugins/Experimental/Dataflow/Source/DataflowEditor/Private/Dataflow/DataflowSimulationScene.cpp b/Engine/Plugins/Experimental/Dataflow/Source/DataflowEditor/Private/Dataflow/DataflowSimulationScene.cpp
--- a/Engine/Plugins/Experimental/Dataflow/Source/DataflowEditor/Private/Dataflow/DataflowSimulationScene.cpp
+++ b/Engine/Plugins/Experimental/Dataflow/Source/DataflowEditor/Private/Dataflow/DataflowSimulationScene.cpp
@@ -198,7 +198,7 @@ void FDataflowSimulationScene::CreateSimulationScene()
 			SceneDescription->CacheAsset, false, GetEditorContent(), SceneDescription->BlueprintTransform);
 
 		// Setup all the skelmesh animations
-		UE::Dataflow::SetupSkeletonAnimation(PreviewActor);
+		UE::Dataflow::SetupSkeletonAnimation(PreviewActor, SceneDescription->bSkeletalMeshVisibility);
 		
 		GetWorld()->GetSubsystem<UDataflowSimulationManager>()->SetSimulationEnabled(false);
 	}
@@ -335,7 +335,7 @@ void UDataflowSimulationSceneDescription::GenerateGeometryCache()
 		}
 		if (!GeometryCachable)
 		{
-			UE_LOG(LogChaosDataflow, Error, TEXT("No Flesh Component in the Preview Actor"));
+			UE_LOG(LogDataflowSimulationGeometryCache, Error, TEXT("No GeometryCachable Component in the Preview Actor"));
 			return;
 		}
 		TOptional<TArray<int32>> OptionalMap = GeometryCachable->GetMeshImportVertexMap(*EmbeddedSkeletalMesh);
diff --git a/Engine/Plugins/Experimental/Dataflow/Source/DataflowEditor/Public/Dataflow/DataflowSimulationControls.h b/Engine/Plugins/Experimental/Dataflow/Source/DataflowEditor/Public/Dataflow/DataflowSimulationControls.h
--- a/Engine/Plugins/Experimental/Dataflow/Source/DataflowEditor/Public/Dataflow/DataflowSimulationControls.h
+++ b/Engine/Plugins/Experimental/Dataflow/Source/DataflowEditor/Public/Dataflow/DataflowSimulationControls.h
@@ -1,4 +1,4 @@
-﻿// Copyright Epic Games, Inc. All Rights Reserved.
+// Copyright Epic Games, Inc. All Rights Reserved.
 
 #pragma once
 
@@ -22,7 +22,7 @@ namespace UE::Dataflow
 		const bool bIsRecording, const TObjectPtr<UDataflowBaseContent>& DataflowContent, const FTransform& ActorTransform);
 
 	/** Setup the skelmesh animations to be used in the scene/generator */
-	void SetupSkeletonAnimation(const TObjectPtr<AActor>& PreviewActor);
+	void SetupSkeletonAnimation(const TObjectPtr<AActor>& PreviewActor, const bool bSkeletalMeshVisibility = true);
 
 	/** Update the skelmesh animation at some point in time (GT) */
 	void UpdateSkeletonAnimation(const TObjectPtr<AActor>& PreviewActor, const float SimulationTime);
diff --git a/Engine/Plugins/Experimental/Dataflow/Source/DataflowEditor/Public/Dataflow/DataflowSimulationGenerator.h b/Engine/Plugins/Experimental/Dataflow/Source/DataflowEditor/Public/Dataflow/DataflowSimulationGenerator.h
--- a/Engine/Plugins/Experimental/Dataflow/Source/DataflowEditor/Public/Dataflow/DataflowSimulationGenerator.h
+++ b/Engine/Plugins/Experimental/Dataflow/Source/DataflowEditor/Public/Dataflow/DataflowSimulationGenerator.h
@@ -1,4 +1,4 @@
-﻿// Copyright Epic Games, Inc. All Rights Reserved.
+// Copyright Epic Games, Inc. All Rights Reserved.
 
 #pragma once
 
@@ -106,7 +106,7 @@ namespace UE::Dataflow
 		/** Allocate the simulation resource from the properties */
         bool AllocateSimulationResource(const FVector2f& TimeRange, const int32 FrameRate,
         	const TObjectPtr<UChaosCacheCollection>& CacheAsset, const TSubclassOf<AActor>& ActorClass,
-        	const TObjectPtr<UDataflowBaseContent>& DataflowContent, const FTransform& BlueprintTransform);
+        	const TObjectPtr<UDataflowBaseContent>& DataflowContent, const FTransform& BlueprintTransform, const bool bSkeletalMeshVisibility = true);
 
         /** Free the simulation resource */
         void FreeSimulationResource();
@@ -166,6 +166,9 @@ namespace UE::Dataflow
 		/** Set the dataflow content */
 		void SetDataflowContent(const TObjectPtr<UDataflowBaseContent>& InDataflowContent);
 
+		/** Set the skeletal mesh visibility */
+		void SetSkeletalMeshVisibility(const bool bInSkeletalMeshVisibility);
+
 		/** Enqueue a generator action to be processed on the async thread */
 		void RequestGeneratorAction(EDataflowGeneratorActions Action);
 		
@@ -195,6 +198,9 @@ namespace UE::Dataflow
 		/** Dataflow content */
 		TObjectPtr<UDataflowBaseContent> DataflowContent;
 
+		/** Skeletal mesh visibility */
+		bool bSkeletalMeshVisibility = true;
+
 		/** Pending action to be send to the async thread */
 		EDataflowGeneratorActions PendingAction = EDataflowGeneratorActions::NoAction;
 
diff --git a/Engine/Plugins/Experimental/Dataflow/Source/DataflowEditor/Public/Dataflow/DataflowSimulationScene.h b/Engine/Plugins/Experimental/Dataflow/Source/DataflowEditor/Public/Dataflow/DataflowSimulationScene.h
--- a/Engine/Plugins/Experimental/Dataflow/Source/DataflowEditor/Public/Dataflow/DataflowSimulationScene.h
+++ b/Engine/Plugins/Experimental/Dataflow/Source/DataflowEditor/Public/Dataflow/DataflowSimulationScene.h
@@ -65,6 +65,9 @@ public:
 	UFUNCTION(CallInEditor, Category = "Geometry", meta = (EditCondition = "CacheAsset != nullptr && GeometryCacheAsset != nullptr && EmbeddedSkeletalMesh != nullptr"))
 	void NewGeometryCache();
 
+	/** Visibility of the skeletal mesh */
+	UPROPERTY(EditAnywhere, Category = "Skeletal Mesh")
+	bool bSkeletalMeshVisibility = true;
 private:
 
 	//~ UObject Interface
