Dataflow : Fix assert when selecting multiple nodes and then selection one part of the selection

diff --git a/Engine/Plugins/Experimental/Dataflow/Source/DataflowEditor/Private/Dataflow/DataflowEditorToolkit.cpp b/Engine/Plugins/Experimental/Dataflow/Source/DataflowEditor/Private/Dataflow/DataflowEditorToolkit.cpp
--- a/Engine/Plugins/Experimental/Dataflow/Source/DataflowEditor/Private/Dataflow/DataflowEditorToolkit.cpp
+++ b/Engine/Plugins/Experimental/Dataflow/Source/DataflowEditor/Private/Dataflow/DataflowEditorToolkit.cpp
@@ -921,7 +921,7 @@ void FDataflowEditorToolkit::OnNodeSelectionChanged(const TSet<UObject*>& InNewS
 			}
 
 			// Something new has been selected.
-			if (NewlySelectedNodes.Num() == 1)
+			if (NewlySelectedNodes.Num() >= 1)
 			{
 				PrimarySelection = Cast< UDataflowEdNode>(NewlySelectedNodes.Array()[0]);
 				ResetListeners(PrimarySelection);
