Dataflow Editor: add mesh stats visualization (just triangle and vertex count of current scene)

diff --git a/Engine/Plugins/Experimental/Dataflow/Source/DataflowEditor/Private/Dataflow/DataflowEditorModule.cpp b/Engine/Plugins/Experimental/Dataflow/Source/DataflowEditor/Private/Dataflow/DataflowEditorModule.cpp
--- a/Engine/Plugins/Experimental/Dataflow/Source/DataflowEditor/Private/Dataflow/DataflowEditorModule.cpp
+++ b/Engine/Plugins/Experimental/Dataflow/Source/DataflowEditor/Private/Dataflow/DataflowEditorModule.cpp
@@ -2,6 +2,7 @@
 
 #include "Dataflow/DataflowEditorModule.h"
 
+#include "Dataflow/DataflowConstructionVisualization.h"
 #include "Dataflow/DataflowEditorStyle.h"
 #include "Dataflow/DataflowEditorMode.h"
 #include "Dataflow/DataflowEditorToolkit.h"
@@ -10,6 +11,7 @@
 #include "Dataflow/DataflowFunctionProperty.h"
 #include "Dataflow/DataflowFunctionPropertyCustomization.h"
 #include "Dataflow/DataflowSNodeFactories.h"
+#include "Dataflow/MeshStatsConstructionVisualization.h"
 #include "Dataflow/ScalarVertexPropertyGroupCustomization.h"
 #include "Dataflow/DataflowToolRegistry.h"
 #include "DataflowEditorTools/DataflowEditorWeightMapPaintTool.h"
@@ -86,6 +88,9 @@ void FDataflowEditorModule::StartupModule()
 	TSharedRef<const FDataflowToolActionCommandBindings> Actions = MakeShared<FDataflowToolActionCommandBindings>();
 	ToolRegistry.AddNodeToToolMapping(FDataflowCollectionAddScalarVertexPropertyNode::StaticType(), ToolBuilder, Actions);
 
+	UE::Dataflow::FDataflowConstructionVisualizationRegistry& ConstructionVisualizationRegistry = UE::Dataflow::FDataflowConstructionVisualizationRegistry::GetInstance();
+	ConstructionVisualizationRegistry.RegisterVisualization(MakeUnique<UE::Dataflow::FMeshStatsConstructionVisualization>());
+
 	FDataflowEditorCommands::Register();
 }
 
@@ -104,6 +109,9 @@ void FDataflowEditorModule::ShutdownModule()
 
 	UE::Dataflow::FDataflowToolRegistry& ToolRegistry = UE::Dataflow::FDataflowToolRegistry::Get();
 	ToolRegistry.RemoveNodeToToolMapping(FDataflowCollectionAddScalarVertexPropertyNode::StaticType());
+
+	UE::Dataflow::FDataflowConstructionVisualizationRegistry& ConstructionVisualizationRegistry = UE::Dataflow::FDataflowConstructionVisualizationRegistry::GetInstance();
+	ConstructionVisualizationRegistry.DeregisterVisualization(UE::Dataflow::FMeshStatsConstructionVisualization::Name);
 }
 
 IMPLEMENT_MODULE(FDataflowEditorModule, DataflowEditor)
diff --git a/Engine/Plugins/Experimental/Dataflow/Source/DataflowEditor/Private/Dataflow/MeshStatsConstructionVisualization.cpp b/Engine/Plugins/Experimental/Dataflow/Source/DataflowEditor/Private/Dataflow/MeshStatsConstructionVisualization.cpp
new file mode 100644
--- /dev/null
+++ b/Engine/Plugins/Experimental/Dataflow/Source/DataflowEditor/Private/Dataflow/MeshStatsConstructionVisualization.cpp
@@ -0,0 +1,72 @@
+// Copyright Epic Games, Inc. All Rights Reserved.
+
+#include "Dataflow/MeshStatsConstructionVisualization.h"
+#include "Components/DynamicMeshComponent.h"
+#include "Dataflow/DataflowConstructionViewportClient.h"
+#include "Dataflow/DataflowConstructionScene.h"
+#include "Framework/MultiBox/MultiBoxBuilder.h"
+#include "CanvasItem.h"
+#include "CanvasTypes.h"
+
+#define LOCTEXT_NAMESPACE "MeshStatsConstructionVisualization"
+
+namespace UE::Dataflow
+{
+	FName FMeshStatsConstructionVisualization::Name = FName("MeshStatsConstructionVisualization");
+
+	void FMeshStatsConstructionVisualization::ExtendViewportShowMenu(const TSharedPtr<FDataflowConstructionViewportClient>& ViewportClient, FMenuBuilder& MenuBuilder)
+	{
+		MenuBuilder.BeginSection(TEXT("MeshStatsVisualization"), LOCTEXT("MeshStatsVisualizationSectionName", "Mesh"));
+		{
+			const FUIAction MeshStatsToggleAction(
+				FExecuteAction::CreateLambda([this, ViewportClient]()
+					{
+						bMeshStatsVisualizationEnabled = !bMeshStatsVisualizationEnabled;
+						if (ViewportClient)
+						{
+							ViewportClient->Invalidate();
+						}
+					}),
+				FCanExecuteAction(),
+				FIsActionChecked::CreateLambda([this]()
+					{
+						return bMeshStatsVisualizationEnabled;
+					}));
+
+			MenuBuilder.AddMenuEntry(LOCTEXT("MeshStatsVisualization_MeshStatsEnabled", "Mesh Stats"),
+				LOCTEXT("MeshStatsVisualization_MeshStatsEnabled_TooltipText", "Display mesh stats"),
+				FSlateIcon(),
+				MeshStatsToggleAction,
+				NAME_None,
+				EUserInterfaceActionType::ToggleButton);
+		}
+		MenuBuilder.EndSection();
+	}
+
+	void FMeshStatsConstructionVisualization::DrawCanvas(const FDataflowConstructionScene* ConstructionScene, FCanvas* Canvas, const FSceneView* SceneView)
+	{
+		if (bMeshStatsVisualizationEnabled && ConstructionScene)
+		{
+			int32 TotalTris = 0;
+			int32 TotalVertices = 0;
+
+			for (const UDynamicMeshComponent* const MeshComponent : ConstructionScene->GetDynamicMeshComponents())
+			{
+				if (const UE::Geometry::FDynamicMesh3* const Mesh = MeshComponent->GetMesh())
+				{
+					TotalTris += Mesh->TriangleCount();
+					TotalVertices += Mesh->VertexCount();
+				}
+			}
+
+			FCanvasTextItem MessageTextItem(FVector2D(10, 40), FText::GetEmpty(), GEngine->GetSmallFont(), FLinearColor::White);
+			MessageTextItem.Text = FText::Format(LOCTEXT("MeshStatsVisualization_DisplayMessage", "Triangles: {0}   Vertices: {1}"), TotalTris, TotalVertices);
+			MessageTextItem.EnableShadow(FLinearColor::Black);
+			Canvas->DrawItem(MessageTextItem);
+		}
+	}
+
+}
+
+#undef LOCTEXT_NAMESPACE 
+
diff --git a/Engine/Plugins/Experimental/Dataflow/Source/DataflowEditor/Private/Dataflow/MeshStatsConstructionVisualization.h b/Engine/Plugins/Experimental/Dataflow/Source/DataflowEditor/Private/Dataflow/MeshStatsConstructionVisualization.h
new file mode 100644
--- /dev/null
+++ b/Engine/Plugins/Experimental/Dataflow/Source/DataflowEditor/Private/Dataflow/MeshStatsConstructionVisualization.h
@@ -0,0 +1,24 @@
+// Copyright Epic Games, Inc. All Rights Reserved.
+
+#include "Dataflow/DataflowConstructionVisualization.h"
+
+namespace UE::Dataflow
+{
+	/** A simple visualization that displays the number of triangles in the Construction View Scene */
+	class FMeshStatsConstructionVisualization final : public IDataflowConstructionVisualization
+	{
+	public:
+		static FName Name;
+
+	private:
+		virtual FName GetName() const override
+		{
+			return Name;
+		}
+		virtual void ExtendViewportShowMenu(const TSharedPtr<FDataflowConstructionViewportClient>& ViewportClient, FMenuBuilder& MenuBuilder) override;
+		virtual void DrawCanvas(const FDataflowConstructionScene* ConstructionScene, FCanvas* Canvas, const FSceneView* SceneView) override;
+
+		bool bMeshStatsVisualizationEnabled = false;
+	};
+}
+
