Include missing pragma once in header file

diff --git a/Engine/Plugins/Experimental/Dataflow/Source/DataflowEditor/Private/Dataflow/MeshStatsConstructionVisualization.h b/Engine/Plugins/Experimental/Dataflow/Source/DataflowEditor/Private/Dataflow/MeshStatsConstructionVisualization.h
--- a/Engine/Plugins/Experimental/Dataflow/Source/DataflowEditor/Private/Dataflow/MeshStatsConstructionVisualization.h
+++ b/Engine/Plugins/Experimental/Dataflow/Source/DataflowEditor/Private/Dataflow/MeshStatsConstructionVisualization.h
@@ -1,4 +1,5 @@
 // Copyright Epic Games, Inc. All Rights Reserved.
+#pragma once
 
 #include "Dataflow/DataflowConstructionVisualization.h"
 
