Dataflow Editor: if there is a single mesh in the construction viewport after changing nodes, automatically select the mesh. This allows interactive tools to start when their corresponding node is single-clicked, rather than having to also click in the viewport. Also enables quickly switching tools when selecting different nodes

diff --git a/Engine/Plugins/Experimental/Dataflow/Source/DataflowEditor/Private/Dataflow/DataflowConstructionScene.cpp b/Engine/Plugins/Experimental/Dataflow/Source/DataflowEditor/Private/Dataflow/DataflowConstructionScene.cpp
--- a/Engine/Plugins/Experimental/Dataflow/Source/DataflowEditor/Private/Dataflow/DataflowConstructionScene.cpp
+++ b/Engine/Plugins/Experimental/Dataflow/Source/DataflowEditor/Private/Dataflow/DataflowConstructionScene.cpp
@@ -185,6 +185,18 @@ void FDataflowConstructionScene::UpdateDynamicMeshComponents()
 				}
 			}
 
+
+			// If we have a single mesh component in the scene, select it
+			if (DynamicMeshComponents.Num() == 1 && DataflowModeManager)
+			{
+				if (USelection* const SelectedComponents = DataflowModeManager->GetSelectedComponents())
+				{
+					UDynamicMeshComponent* const DynamicMeshComponent = DynamicMeshComponents.CreateIterator()->Value;
+					SelectedComponents->Select(DynamicMeshComponent);
+					DynamicMeshComponent->PushSelectionToProxy();
+				}
+			}
+
 			// Add hidden DynamicMeshComponents for any targets that we want to render in wireframe
 			// 
 			// Note: UMeshElementsVisualizers need source meshes to pull from. We add invisible dynamic mesh components to the existing DynamicMeshComponents collection
diff --git a/Engine/Plugins/Experimental/Dataflow/Source/DataflowEditor/Private/Dataflow/DataflowEditorMode.cpp b/Engine/Plugins/Experimental/Dataflow/Source/DataflowEditor/Private/Dataflow/DataflowEditorMode.cpp
--- a/Engine/Plugins/Experimental/Dataflow/Source/DataflowEditor/Private/Dataflow/DataflowEditorMode.cpp
+++ b/Engine/Plugins/Experimental/Dataflow/Source/DataflowEditor/Private/Dataflow/DataflowEditorMode.cpp
@@ -580,6 +580,15 @@ void UDataflowEditorMode::ModeTick(float DeltaTime)
 		}
 	}
 
+	checkf(ConstructionScene, TEXT("Null ConstructionScene in DataflowEditorMode"));
+	if (const UDataflowBaseContent* const EditorContent = ConstructionScene->GetEditorContent())
+	{
+		if (EditorContent->IsConstructionDirty())
+		{
+			ConstructionScene->UpdateConstructionScene();
+		}
+	}
+
 	if (!NodeTypeForPendingToolStart.IsNone() && !GetToolManager()->HasActiveTool(EToolSide::Left))
 	{
 		const TSharedRef<FUICommandList> CommandList = Toolkit->GetToolkitCommands();
