Dataflow Editor: add a LOD selection button to the simulation viewport

diff --git a/Engine/Plugins/Experimental/Dataflow/Source/DataflowEditor/Private/Dataflow/DataflowSimulationScene.cpp b/Engine/Plugins/Experimental/Dataflow/Source/DataflowEditor/Private/Dataflow/DataflowSimulationScene.cpp
--- a/Engine/Plugins/Experimental/Dataflow/Source/DataflowEditor/Private/Dataflow/DataflowSimulationScene.cpp
+++ b/Engine/Plugins/Experimental/Dataflow/Source/DataflowEditor/Private/Dataflow/DataflowSimulationScene.cpp
@@ -183,6 +183,8 @@ void FDataflowSimulationScene::BindSceneSelection()
 
 void FDataflowSimulationScene::CreateSimulationScene()
 {
+	const int32 PreviewLOD = GetPreviewLOD();
+
 	if(SimulationGenerator && SceneDescription && SceneDescription->BlueprintClass && GetWorld())
 	{
 		SimulationGenerator->SetCacheParams(SceneDescription->CacheParams);
@@ -205,6 +207,8 @@ void FDataflowSimulationScene::CreateSimulationScene()
 
 	// update the selection binding since we are constantly editing the graph
 	BindSceneSelection();
+
+	SetPreviewLOD(PreviewLOD);
 }
 
 void FDataflowSimulationScene::UpdateSimulationCache()
@@ -215,6 +219,28 @@ void FDataflowSimulationScene::UpdateSimulationCache()
 	}
 }
 
+void FDataflowSimulationScene::SetPreviewLOD(int32 InLOD)
+{
+	CurrentPreviewLOD = InLOD;
+
+	if (PreviewActor)
+	{
+		PreviewActor->ForEachComponent<UActorComponent>(true, [this](UActorComponent* Component)
+		{
+			if (ILODSyncInterface* const LODInterface = Cast<ILODSyncInterface>(Component))
+			{
+				LODInterface->SetForceStreamedLOD(CurrentPreviewLOD);
+				LODInterface->SetForceRenderedLOD(CurrentPreviewLOD);
+			}
+		});
+	}
+}
+
+int32 FDataflowSimulationScene::GetPreviewLOD() const
+{
+	return CurrentPreviewLOD;
+}
+
 void FDataflowSimulationScene::TickDataflowScene(const float DeltaSeconds)
 {
 	if(const TObjectPtr<UDataflowBaseContent>& EditorContent = GetEditorContent())
diff --git a/Engine/Plugins/Experimental/Dataflow/Source/DataflowEditor/Private/Dataflow/DataflowSimulationViewportToolbar.cpp b/Engine/Plugins/Experimental/Dataflow/Source/DataflowEditor/Private/Dataflow/DataflowSimulationViewportToolbar.cpp
--- a/Engine/Plugins/Experimental/Dataflow/Source/DataflowEditor/Private/Dataflow/DataflowSimulationViewportToolbar.cpp
+++ b/Engine/Plugins/Experimental/Dataflow/Source/DataflowEditor/Private/Dataflow/DataflowSimulationViewportToolbar.cpp
@@ -1,4 +1,4 @@
-﻿// Copyright Epic Games, Inc. All Rights Reserved.
+// Copyright Epic Games, Inc. All Rights Reserved.
 #include "Dataflow/DataflowSimulationViewportToolbar.h"
 #include "Dataflow/DataflowSimulationViewport.h"
 #include "Dataflow/DataflowEditorCommands.h"
@@ -8,6 +8,7 @@
 #include "SEditorViewportToolBarMenu.h"
 #include "Widgets/Layout/SBorder.h"
 #include "Widgets/SBoxPanel.h"
+#include "LODSyncInterface.h"
 
 #define LOCTEXT_NAMESPACE "DataflowSimulationViewportToolBar"
 
@@ -34,6 +35,18 @@ void SDataflowSimulationViewportToolBar::ExtendLeftAlignedToolbarSlots(TSharedPt
 	};
 	
 	const FMargin ToolbarSlotPadding(2.0f, 2.0f);
+	MainBoxPtr->AddSlot()
+		.AutoWidth()
+		.Padding(ToolbarSlotPadding)
+		.HAlign(HAlign_Left)
+		[
+			SNew(SEditorViewportToolbarMenu)
+				.Label(this, &SDataflowSimulationViewportToolBar::LODButtonLabel)
+				.ParentToolBar(ParentToolBarPtr)
+				.OnGetMenuContent(this, &SDataflowSimulationViewportToolBar::MakeLODMenu)
+				.Visibility(TAttribute<EVisibility>::Create(TAttribute<EVisibility>::FGetter::CreateSP(this, &SDataflowSimulationViewportToolBar::LODButtonVisibility)))
+		];
+
 	MainBoxPtr->AddSlot()
 		.Padding(ToolbarSlotPadding)
 		[
@@ -90,4 +103,108 @@ TSharedRef<SWidget> SDataflowSimulationViewportToolBar::MakeToolBar(const TShare
 	return ToolbarBuilder.MakeWidget();
 }
 
+EVisibility SDataflowSimulationViewportToolBar::LODButtonVisibility() const
+{
+	const TSharedPtr<SDataflowSimulationViewport> Viewport = EditorViewport.Pin();
+	checkf(Viewport, TEXT("Viewport not found in ViewportToolBar"));
+
+	return (NumLODs() > 1) ? EVisibility::Visible : EVisibility::Collapsed;
+}
+
+
+FText SDataflowSimulationViewportToolBar::LODButtonLabel() const
+{
+	if (const TSharedPtr<SDataflowSimulationViewport> Viewport = EditorViewport.Pin())
+	{
+		if (const TSharedPtr<FDataflowSimulationScene>& SimScene = Viewport->GetSimulationScene())
+		{
+			const int32 CurrentLOD = SimScene->GetPreviewLOD();
+			if (CurrentLOD == INDEX_NONE)
+			{
+				return LOCTEXT("LODMenuButtonLabelAuto", "LOD Auto");
+			}
+			else
+			{
+				return FText::Format(LOCTEXT("LODMenuButtonLabel", "LOD {0}"), CurrentLOD);
+			}
+		}
+	}
+
+	return LOCTEXT("LODMenuButtonLabelNoLODs", "LOD");
+}
+
+int32 SDataflowSimulationViewportToolBar::NumLODs() const
+{
+	int32 MaxNumLODs = 0;
+	if (const TSharedPtr<SDataflowSimulationViewport> Viewport = EditorViewport.Pin())
+	{
+		if (const TSharedPtr<const FDataflowSimulationScene>& SimScene = Viewport->GetSimulationScene())
+		{
+			if (const AActor* const PreviewActor = SimScene->GetPreviewActor())
+			{
+				PreviewActor->ForEachComponent<UActorComponent>(true, [&MaxNumLODs](UActorComponent* Component)
+				{
+					if (const ILODSyncInterface* const LODInterface = Cast<ILODSyncInterface>(Component))
+					{
+						MaxNumLODs = FMath::Max(MaxNumLODs, LODInterface->GetNumSyncLODs());
+					}
+				});
+			}
+		}
+	}
+	return MaxNumLODs;
+}
+
+TSharedRef<SWidget> SDataflowSimulationViewportToolBar::MakeLODMenu() const
+{
+	const auto SetCurrentLOD = [this](int32 NewLOD)
+	{
+		if (const TSharedPtr<SDataflowSimulationViewport> Viewport = EditorViewport.Pin())
+		{
+			if (const TSharedPtr<FDataflowSimulationScene>& SimScene = Viewport->GetSimulationScene())
+			{
+				SimScene->SetPreviewLOD(NewLOD);
+			}
+		}
+	};
+
+	const auto IsLODCurrent = [this](int32 QueryLOD) -> bool
+	{
+		if (const TSharedPtr<const SDataflowSimulationViewport> Viewport = EditorViewport.Pin())
+		{
+			if (const TSharedPtr<const FDataflowSimulationScene>& SimScene = Viewport->GetSimulationScene())
+			{
+				return (SimScene->GetPreviewLOD() == QueryLOD);
+			}
+		}
+		return false;
+	};
+
+	constexpr bool bInShouldCloseWindowAfterMenuSelection = true;
+	FMenuBuilder MenuBuilder(bInShouldCloseWindowAfterMenuSelection, CommandList);
+
+	MenuBuilder.PushCommandList(CommandList.ToSharedRef());
+	MenuBuilder.BeginSection("ClothAssetPreviewLODs", LOCTEXT("LODMenuSectionLabel", "LODs"));
+	{
+		const FUIAction LODAutoAction(
+			FExecuteAction::CreateLambda(SetCurrentLOD, (int32)INDEX_NONE),
+			FCanExecuteAction(),
+			FIsActionChecked::CreateLambda(IsLODCurrent, (int32)INDEX_NONE));
+		MenuBuilder.AddMenuEntry(LOCTEXT("LODMenuEntryLabelAuto", "LOD Auto"), FText::GetEmpty(), FSlateIcon(), LODAutoAction, NAME_None, EUserInterfaceActionType::RadioButton);
+
+		for (int32 LODIndex = 0; LODIndex < NumLODs(); ++LODIndex)
+		{
+			const FUIAction Action(
+				FExecuteAction::CreateLambda(SetCurrentLOD, LODIndex),
+				FCanExecuteAction(),
+				FIsActionChecked::CreateLambda(IsLODCurrent, LODIndex));
+			const FText MenuEntryLabel = FText::Format(LOCTEXT("LODMenuEntryLabel", "LOD {0}"), LODIndex);
+			MenuBuilder.AddMenuEntry(MenuEntryLabel, FText::GetEmpty(), FSlateIcon(), Action, NAME_None, EUserInterfaceActionType::RadioButton);
+		}
+	}
+	MenuBuilder.EndSection();
+	MenuBuilder.PopCommandList();
+	return MenuBuilder.MakeWidget();
+}
+
 #undef LOCTEXT_NAMESPACE
diff --git a/Engine/Plugins/Experimental/Dataflow/Source/DataflowEditor/Public/Dataflow/DataflowSimulationScene.h b/Engine/Plugins/Experimental/Dataflow/Source/DataflowEditor/Public/Dataflow/DataflowSimulationScene.h
--- a/Engine/Plugins/Experimental/Dataflow/Source/DataflowEditor/Public/Dataflow/DataflowSimulationScene.h
+++ b/Engine/Plugins/Experimental/Dataflow/Source/DataflowEditor/Public/Dataflow/DataflowSimulationScene.h
@@ -147,6 +147,10 @@ public:
 	TObjectPtr<AActor> GetPreviewActor() { return PreviewActor; }
 	const TObjectPtr<AActor> GetPreviewActor() const { return PreviewActor; }
 
+	/** LOD for preview actor components */
+	void SetPreviewLOD(int32 InLOD);
+	int32 GetPreviewLOD() const;
+
 private:
 
 	/** Bind the scene selection to the components */
@@ -175,6 +179,8 @@ private:
 
 	/** Handle for the delegate */
 	FDelegateHandle OnObjectsReinstancedHandle;
+
+	int32 CurrentPreviewLOD = INDEX_NONE;
 };
 
 
diff --git a/Engine/Plugins/Experimental/Dataflow/Source/DataflowEditor/Public/Dataflow/DataflowSimulationViewportToolbar.h b/Engine/Plugins/Experimental/Dataflow/Source/DataflowEditor/Public/Dataflow/DataflowSimulationViewportToolbar.h
--- a/Engine/Plugins/Experimental/Dataflow/Source/DataflowEditor/Public/Dataflow/DataflowSimulationViewportToolbar.h
+++ b/Engine/Plugins/Experimental/Dataflow/Source/DataflowEditor/Public/Dataflow/DataflowSimulationViewportToolbar.h
@@ -1,4 +1,4 @@
-﻿// Copyright Epic Games, Inc. All Rights Reserved.
+// Copyright Epic Games, Inc. All Rights Reserved.
 #pragma once
 
 #include "CoreMinimal.h"
@@ -28,6 +28,11 @@ private:
 
 	TSharedRef<SWidget> MakeToolBar(const TSharedPtr<FExtender> InExtenders) const;
 
+	EVisibility LODButtonVisibility() const;
+	FText LODButtonLabel() const;
+	int32 NumLODs() const;
+	TSharedRef<SWidget> MakeLODMenu() const;
+
 	TSharedPtr<FUICommandList> CommandList;
 	TSharedPtr<FExtender> Extenders;
 };
\ No newline at end of file
