Dataflow: fix some mistakes with new registry classes:
- Remove Instance variables for TLazySingleton classes
- Add missing 'F' to class name

diff --git a/Engine/Plugins/ChaosClothAssetEditor/Source/ChaosClothAssetDataflowNodes/Private/ChaosClothAsset/ChaosClothAssetDataflowNodesModule.cpp b/Engine/Plugins/ChaosClothAssetEditor/Source/ChaosClothAssetDataflowNodes/Private/ChaosClothAsset/ChaosClothAssetDataflowNodesModule.cpp
--- a/Engine/Plugins/ChaosClothAssetEditor/Source/ChaosClothAssetDataflowNodes/Private/ChaosClothAsset/ChaosClothAssetDataflowNodesModule.cpp
+++ b/Engine/Plugins/ChaosClothAssetEditor/Source/ChaosClothAssetDataflowNodes/Private/ChaosClothAsset/ChaosClothAssetDataflowNodesModule.cpp
@@ -348,7 +348,7 @@ PRAGMA_ENABLE_DEPRECATION_WARNINGS
 
 			Private::RegisterRenderingCallbacks();
 
-			DataflowAddScalarVertexPropertyCallbackRegistry::Get().RegisterCallbacks(MakeUnique<Private::FClothCollectionAddScalarVertexPropertyCallbacks>());
+			FDataflowAddScalarVertexPropertyCallbackRegistry::Get().RegisterCallbacks(MakeUnique<Private::FClothCollectionAddScalarVertexPropertyCallbacks>());
 		}
 
 		virtual void ShutdownModule() override
@@ -377,7 +377,7 @@ PRAGMA_ENABLE_DEPRECATION_WARNINGS
 				}
 			}
 
-			DataflowAddScalarVertexPropertyCallbackRegistry::Get().DeregisterCallbacks(Private::FClothCollectionAddScalarVertexPropertyCallbacks::Name);
+			FDataflowAddScalarVertexPropertyCallbackRegistry::Get().DeregisterCallbacks(Private::FClothCollectionAddScalarVertexPropertyCallbacks::Name);
 		}
 	};
 }  // End namespace UE::Chaos::ClothAsset
diff --git a/Engine/Plugins/Experimental/Dataflow/Source/DataflowEditor/Private/Dataflow/DataflowRenderingViewMode.cpp b/Engine/Plugins/Experimental/Dataflow/Source/DataflowEditor/Private/Dataflow/DataflowRenderingViewMode.cpp
--- a/Engine/Plugins/Experimental/Dataflow/Source/DataflowEditor/Private/Dataflow/DataflowRenderingViewMode.cpp
+++ b/Engine/Plugins/Experimental/Dataflow/Source/DataflowEditor/Private/Dataflow/DataflowRenderingViewMode.cpp
@@ -57,7 +57,6 @@ namespace UE::Dataflow
 		return FText(LOCTEXT("3DViewTooltipText", "Default 3D View"));
 	}
 
-	TUniquePtr<FRenderingViewModeFactory> FRenderingViewModeFactory::Instance = nullptr;
 
 	//
 	// Factory
diff --git a/Engine/Plugins/Experimental/Dataflow/Source/DataflowEditor/Private/Dataflow/ScalarVertexPropertyGroupCustomization.cpp b/Engine/Plugins/Experimental/Dataflow/Source/DataflowEditor/Private/Dataflow/ScalarVertexPropertyGroupCustomization.cpp
--- a/Engine/Plugins/Experimental/Dataflow/Source/DataflowEditor/Private/Dataflow/ScalarVertexPropertyGroupCustomization.cpp
+++ b/Engine/Plugins/Experimental/Dataflow/Source/DataflowEditor/Private/Dataflow/ScalarVertexPropertyGroupCustomization.cpp
@@ -12,7 +12,7 @@ namespace UE::Dataflow
 
 	TArray<FName> FScalarVertexPropertyGroupCustomization::GetTargetGroupNames(const FManagedArrayCollection& Collection) const
 	{
-		return DataflowAddScalarVertexPropertyCallbackRegistry::Get().GetTargetGroupNames();
+		return FDataflowAddScalarVertexPropertyCallbackRegistry::Get().GetTargetGroupNames();
 	}
 
 	FName FScalarVertexPropertyGroupCustomization::GetCollectionPropertyName() const
diff --git a/Engine/Plugins/Experimental/Dataflow/Source/DataflowEditor/Public/Dataflow/DataflowConstructionVisualization.h b/Engine/Plugins/Experimental/Dataflow/Source/DataflowEditor/Public/Dataflow/DataflowConstructionVisualization.h
--- a/Engine/Plugins/Experimental/Dataflow/Source/DataflowEditor/Public/Dataflow/DataflowConstructionVisualization.h
+++ b/Engine/Plugins/Experimental/Dataflow/Source/DataflowEditor/Public/Dataflow/DataflowConstructionVisualization.h
@@ -42,8 +42,6 @@ namespace UE::Dataflow
 
 	private:
 
-		DATAFLOWEDITOR_API static TUniquePtr<FDataflowConstructionVisualizationRegistry> Instance;
-
 		TMap<FName, TUniquePtr<IDataflowConstructionVisualization>> VisualizationMap;
 	};
 
diff --git a/Engine/Plugins/Experimental/Dataflow/Source/DataflowEditor/Public/Dataflow/DataflowRenderingViewMode.h b/Engine/Plugins/Experimental/Dataflow/Source/DataflowEditor/Public/Dataflow/DataflowRenderingViewMode.h
--- a/Engine/Plugins/Experimental/Dataflow/Source/DataflowEditor/Public/Dataflow/DataflowRenderingViewMode.h
+++ b/Engine/Plugins/Experimental/Dataflow/Source/DataflowEditor/Public/Dataflow/DataflowRenderingViewMode.h
@@ -100,8 +100,6 @@ namespace UE::Dataflow
 
 	private:
 
-		DATAFLOWEDITOR_API static TUniquePtr<FRenderingViewModeFactory> Instance;
-
 		TMap<FName, TUniquePtr<IDataflowConstructionViewMode>> ViewModeMap;
 	};
 }
diff --git a/Engine/Plugins/Experimental/Dataflow/Source/DataflowEditor/Public/Dataflow/DataflowSimulationVisualization.h b/Engine/Plugins/Experimental/Dataflow/Source/DataflowEditor/Public/Dataflow/DataflowSimulationVisualization.h
--- a/Engine/Plugins/Experimental/Dataflow/Source/DataflowEditor/Public/Dataflow/DataflowSimulationVisualization.h
+++ b/Engine/Plugins/Experimental/Dataflow/Source/DataflowEditor/Public/Dataflow/DataflowSimulationVisualization.h
@@ -44,8 +44,6 @@ namespace UE::Dataflow
 
 	private:
 
-		DATAFLOWEDITOR_API static TUniquePtr<FDataflowSimulationVisualizationRegistry> Instance;
-
 		TMap<FName, TUniquePtr<IDataflowSimulationVisualization>> VisualizationMap;
 	};
 
diff --git a/Engine/Plugins/Experimental/Dataflow/Source/DataflowNodes/Private/Dataflow/DataflowCollectionAddScalarVertexPropertyNode.cpp b/Engine/Plugins/Experimental/Dataflow/Source/DataflowNodes/Private/Dataflow/DataflowCollectionAddScalarVertexPropertyNode.cpp
--- a/Engine/Plugins/Experimental/Dataflow/Source/DataflowNodes/Private/Dataflow/DataflowCollectionAddScalarVertexPropertyNode.cpp
+++ b/Engine/Plugins/Experimental/Dataflow/Source/DataflowNodes/Private/Dataflow/DataflowCollectionAddScalarVertexPropertyNode.cpp
@@ -10,31 +10,79 @@
 
 #define LOCTEXT_NAMESPACE "DataflowCollectionAddScalarVertexProperty"
 
+//
+// FDataflowAddScalarVertexPropertyCallbackRegistry
+//
 
+FDataflowAddScalarVertexPropertyCallbackRegistry& FDataflowAddScalarVertexPropertyCallbackRegistry::Get()
+{
+	return TLazySingleton<FDataflowAddScalarVertexPropertyCallbackRegistry>::Get();
+}
+
+void FDataflowAddScalarVertexPropertyCallbackRegistry::RegisterCallbacks(TUniquePtr<IDataflowAddScalarVertexPropertyCallbacks>&& Callbacks)
+{
+	AllCallbacks.Add(Callbacks->GetName(), MoveTemp(Callbacks));
+}
+
+void FDataflowAddScalarVertexPropertyCallbackRegistry::DeregisterCallbacks(const FName& CallbacksName)
+{
+	AllCallbacks.Remove(CallbacksName);
+}
+
+TArray<FName> FDataflowAddScalarVertexPropertyCallbackRegistry::GetTargetGroupNames() const
+{
+	TArray<FName> UniqueNames;
+
+	for (const TPair<FName, TUniquePtr<IDataflowAddScalarVertexPropertyCallbacks>>& CallbacksEntry : AllCallbacks)
+	{
+		for (const FName& GroupName : CallbacksEntry.Value->GetTargetGroupNames())
+		{
+			UniqueNames.AddUnique(GroupName);
+		}
+	}
+	return UniqueNames;
+}
+
+TArray<UE::Dataflow::FRenderingParameter> FDataflowAddScalarVertexPropertyCallbackRegistry::GetRenderingParameters() const
+{
+	TArray<UE::Dataflow::FRenderingParameter> UniqueParameters;
+
+	for (const TPair<FName, TUniquePtr<IDataflowAddScalarVertexPropertyCallbacks>>& CallbacksEntry : AllCallbacks)
+	{
+		for (const UE::Dataflow::FRenderingParameter& RenderingParameter : CallbacksEntry.Value->GetRenderingParameters())
+		{
+			UniqueParameters.AddUnique(RenderingParameter);
+		}
+	}
+	return UniqueParameters;
+}
+
+
+//
+// Deprecated class
+//
+
+PRAGMA_DISABLE_DEPRECATION_WARNINGS
 DataflowAddScalarVertexPropertyCallbackRegistry& DataflowAddScalarVertexPropertyCallbackRegistry::Get()
 {
 	return TLazySingleton<DataflowAddScalarVertexPropertyCallbackRegistry>::Get();
 }
-
 void DataflowAddScalarVertexPropertyCallbackRegistry::TearDown()
 {
 	TLazySingleton<DataflowAddScalarVertexPropertyCallbackRegistry>::TearDown();
 }
-
 void DataflowAddScalarVertexPropertyCallbackRegistry::RegisterCallbacks(TUniquePtr<IDataflowAddScalarVertexPropertyCallbacks>&& Callbacks)
 {
 	AllCallbacks.Add(Callbacks->GetName(), MoveTemp(Callbacks));
 }
-
 void DataflowAddScalarVertexPropertyCallbackRegistry::DeregisterCallbacks(const FName& CallbacksName)
 {
 	AllCallbacks.Remove(CallbacksName);
 }
-
 TArray<FName> DataflowAddScalarVertexPropertyCallbackRegistry::GetTargetGroupNames() const
 {
 	TArray<FName> UniqueNames;
-	
+
 	for (const TPair<FName, TUniquePtr<IDataflowAddScalarVertexPropertyCallbacks>>& CallbacksEntry : AllCallbacks)
 	{
 		for (const FName& GroupName : CallbacksEntry.Value->GetTargetGroupNames())
@@ -44,7 +92,6 @@ TArray<FName> DataflowAddScalarVertexPropertyCallbackRegistry::GetTargetGroupNam
 	}
 	return UniqueNames;
 }
-
 TArray<UE::Dataflow::FRenderingParameter> DataflowAddScalarVertexPropertyCallbackRegistry::GetRenderingParameters() const
 {
 	TArray<UE::Dataflow::FRenderingParameter> UniqueParameters;
@@ -58,7 +105,12 @@ TArray<UE::Dataflow::FRenderingParameter> DataflowAddScalarVertexPropertyCallbac
 	}
 	return UniqueParameters;
 }
+PRAGMA_ENABLE_DEPRECATION_WARNINGS
+
 
+//
+// FDataflowCollectionAddScalarVertexPropertyNode
+//
 
 FDataflowCollectionAddScalarVertexPropertyNode::FDataflowCollectionAddScalarVertexPropertyNode(const UE::Dataflow::FNodeParameters& InParam, FGuid InGuid)
 	: FDataflowNode(InParam, InGuid)
@@ -70,7 +122,7 @@ FDataflowCollectionAddScalarVertexPropertyNode::FDataflowCollectionAddScalarVert
 
 TArray<UE::Dataflow::FRenderingParameter> FDataflowCollectionAddScalarVertexPropertyNode::GetRenderParametersImpl() const
 {
-	return DataflowAddScalarVertexPropertyCallbackRegistry::Get().GetRenderingParameters();
+	return FDataflowAddScalarVertexPropertyCallbackRegistry::Get().GetRenderingParameters();
 }
 
 
diff --git a/Engine/Plugins/Experimental/Dataflow/Source/DataflowNodes/Private/Dataflow/DataflowNodesPlugin.cpp b/Engine/Plugins/Experimental/Dataflow/Source/DataflowNodes/Private/Dataflow/DataflowNodesPlugin.cpp
--- a/Engine/Plugins/Experimental/Dataflow/Source/DataflowNodes/Private/Dataflow/DataflowNodesPlugin.cpp
+++ b/Engine/Plugins/Experimental/Dataflow/Source/DataflowNodes/Private/Dataflow/DataflowNodesPlugin.cpp
@@ -50,12 +50,12 @@ void IDataflowNodesPlugin::StartupModule()
 
 	UE::Dataflow::RegisterNodeFilter(FDataflowTerminalNode::StaticType());
 
-	DataflowAddScalarVertexPropertyCallbackRegistry::Get().RegisterCallbacks(MakeUnique<FGeometryCollectionAddScalarVertexPropertyCallbacks>());
+	FDataflowAddScalarVertexPropertyCallbackRegistry::Get().RegisterCallbacks(MakeUnique<FGeometryCollectionAddScalarVertexPropertyCallbacks>());
 }
 
 void IDataflowNodesPlugin::ShutdownModule()
 {
-	DataflowAddScalarVertexPropertyCallbackRegistry::Get().DeregisterCallbacks(FGeometryCollectionAddScalarVertexPropertyCallbacks::Name);
+	FDataflowAddScalarVertexPropertyCallbackRegistry::Get().DeregisterCallbacks(FGeometryCollectionAddScalarVertexPropertyCallbacks::Name);
 }
 
 
diff --git a/Engine/Plugins/Experimental/Dataflow/Source/DataflowNodes/Public/Dataflow/DataflowCollectionAddScalarVertexPropertyNode.h b/Engine/Plugins/Experimental/Dataflow/Source/DataflowNodes/Public/Dataflow/DataflowCollectionAddScalarVertexPropertyNode.h
--- a/Engine/Plugins/Experimental/Dataflow/Source/DataflowNodes/Public/Dataflow/DataflowCollectionAddScalarVertexPropertyNode.h
+++ b/Engine/Plugins/Experimental/Dataflow/Source/DataflowNodes/Public/Dataflow/DataflowCollectionAddScalarVertexPropertyNode.h
@@ -19,12 +19,11 @@ public:
 	virtual TArray<UE::Dataflow::FRenderingParameter> GetRenderingParameters() const = 0;
 };
 
-class DataflowAddScalarVertexPropertyCallbackRegistry
+class FDataflowAddScalarVertexPropertyCallbackRegistry
 {
 public:
 
-	DATAFLOWNODES_API static DataflowAddScalarVertexPropertyCallbackRegistry& Get();
-	DATAFLOWNODES_API static void TearDown();
+	DATAFLOWNODES_API static FDataflowAddScalarVertexPropertyCallbackRegistry& Get();
 
 	DATAFLOWNODES_API void RegisterCallbacks(TUniquePtr<IDataflowAddScalarVertexPropertyCallbacks>&& Callbacks);
 
@@ -39,6 +38,18 @@ private:
 	TMap<FName, TUniquePtr<IDataflowAddScalarVertexPropertyCallbacks>> AllCallbacks;
 };
 
+class UE_DEPRECATED(5.6, "Please use FDataflowAddScalarVertexPropertyCallbackRegistry") DataflowAddScalarVertexPropertyCallbackRegistry
+{
+public:
+	DATAFLOWNODES_API static DataflowAddScalarVertexPropertyCallbackRegistry& Get();
+	DATAFLOWNODES_API static void TearDown();
+	DATAFLOWNODES_API void RegisterCallbacks(TUniquePtr<IDataflowAddScalarVertexPropertyCallbacks> && Callbacks);
+	DATAFLOWNODES_API void DeregisterCallbacks(const FName & CallbacksName);
+	DATAFLOWNODES_API TArray<FName> GetTargetGroupNames() const;
+	DATAFLOWNODES_API TArray<UE::Dataflow::FRenderingParameter> GetRenderingParameters() const;
+private:
+	TMap<FName, TUniquePtr<IDataflowAddScalarVertexPropertyCallbacks>> AllCallbacks;
+};
 
 /*
 * Custom type so that we can use property type customization
