Implement Copy Interactor Name right click menu for cloth in the Dataflow Editor. This feature was added for 5.5 in the Cloth Editor, but I didn't have time to also put it in the Dataflow Editor as part of that CL.

diff --git a/Engine/Plugins/ChaosClothAssetEditor/Source/ChaosClothAssetEditor/Private/ChaosClothAsset/ChaosClothAssetEditorModule.cpp b/Engine/Plugins/ChaosClothAssetEditor/Source/ChaosClothAssetEditor/Private/ChaosClothAsset/ChaosClothAssetEditorModule.cpp
--- a/Engine/Plugins/ChaosClothAssetEditor/Source/ChaosClothAssetEditor/Private/ChaosClothAsset/ChaosClothAssetEditorModule.cpp
+++ b/Engine/Plugins/ChaosClothAssetEditor/Source/ChaosClothAssetEditor/Private/ChaosClothAsset/ChaosClothAssetEditorModule.cpp
@@ -5,6 +5,7 @@
 #include "ChaosClothAsset/ClothAsset.h"
 #include "ChaosClothAsset/ClothDataflowConstructionVisualization.h"
 #include "ChaosClothAsset/ClothDataflowSimulationVisualization.h"
+#include "ChaosClothAsset/ClothSimulationNodeDetailExtender.h"
 #include "ChaosClothAsset/ClothEditorCommands.h"
 #include "ChaosClothAsset/ClothEditorMode.h"
 #include "ChaosClothAsset/ClothEditorStyle.h"
@@ -101,18 +102,20 @@ PRAGMA_ENABLE_DEPRECATION_WARNINGS
 
 		Dataflow::FDataflowConstructionVisualizationRegistry::GetInstance().RegisterVisualization(MakeUnique<FClothDataflowConstructionVisualization>());
 		Dataflow::FDataflowSimulationVisualizationRegistry::GetInstance().RegisterVisualization(MakeUnique<FClothDataflowSimulationVisualization>());
+		Dataflow::FDataflowNodeDetailExtensionRegistry::GetInstance().RegisterExtension(MakeUnique<FClothSimulationNodeDetailExtender>());
 	}
 
 	void FChaosClothAssetEditorModule::ShutdownModule()
 	{
-		FChaosClothAssetEditorCommands::Unregister();
-
-		FEditorModeRegistry::Get().UnregisterMode(UChaosClothAssetEditorMode::EM_ChaosClothAssetEditorModeId);
+		Dataflow::FDataflowNodeDetailExtensionRegistry::GetInstance().DeregisterExtension(FClothSimulationNodeDetailExtender::Name);
+		Dataflow::FDataflowConstructionVisualizationRegistry::GetInstance().DeregisterVisualization(FClothDataflowConstructionVisualization::Name);
+		Dataflow::FDataflowSimulationVisualizationRegistry::GetInstance().DeregisterVisualization(FClothDataflowSimulationVisualization::Name);
 
 		UToolMenus::UnRegisterStartupCallback(StartupCallbackDelegateHandle);
 
-		Dataflow::FDataflowConstructionVisualizationRegistry::GetInstance().DeregisterVisualization(FClothDataflowConstructionVisualization::Name);
-		Dataflow::FDataflowSimulationVisualizationRegistry::GetInstance().DeregisterVisualization(FClothDataflowSimulationVisualization::Name);
+		FChaosClothAssetEditorCommands::Unregister();
+
+		FEditorModeRegistry::Get().UnregisterMode(UChaosClothAssetEditorMode::EM_ChaosClothAssetEditorModeId);
 	}
 
 } // namespace UE::Chaos::ClothAsset
diff --git a/Engine/Plugins/ChaosClothAssetEditor/Source/ChaosClothAssetEditor/Private/ChaosClothAsset/ClothSimulationNodeDetailExtender.cpp b/Engine/Plugins/ChaosClothAssetEditor/Source/ChaosClothAssetEditor/Private/ChaosClothAsset/ClothSimulationNodeDetailExtender.cpp
--- a/Engine/Plugins/ChaosClothAssetEditor/Source/ChaosClothAssetEditor/Private/ChaosClothAsset/ClothSimulationNodeDetailExtender.cpp
+++ b/Engine/Plugins/ChaosClothAssetEditor/Source/ChaosClothAssetEditor/Private/ChaosClothAsset/ClothSimulationNodeDetailExtender.cpp
@@ -7,6 +7,13 @@
 
 namespace UE::Chaos::ClothAsset
 {
+	const FName FClothSimulationNodeDetailExtender::Name = FName("FClothSimulationNodeDetailExtender");
+
+	FName FClothSimulationNodeDetailExtender::GetName() const
+	{
+		return Name;
+	}
+
 	bool FClothSimulationNodeDetailExtender::IsPropertyExtendable(const UClass* InObjectClass, const IPropertyHandle& PropertyHandle) const
 	{
 		return PropertyHandle.HasMetaData(TEXT("InteractorName"));
diff --git a/Engine/Plugins/ChaosClothAssetEditor/Source/ChaosClothAssetEditor/Private/ChaosClothAsset/ClothSimulationNodeDetailExtender.h b/Engine/Plugins/ChaosClothAssetEditor/Source/ChaosClothAssetEditor/Private/ChaosClothAsset/ClothSimulationNodeDetailExtender.h
--- a/Engine/Plugins/ChaosClothAssetEditor/Source/ChaosClothAssetEditor/Private/ChaosClothAsset/ClothSimulationNodeDetailExtender.h
+++ b/Engine/Plugins/ChaosClothAssetEditor/Source/ChaosClothAssetEditor/Private/ChaosClothAsset/ClothSimulationNodeDetailExtender.h
@@ -2,12 +2,19 @@
 
 #pragma once
 
-#include "IDetailPropertyExtensionHandler.h"
+#include "Dataflow/DataflowNodeDetailExtension.h"
 
 namespace UE::Chaos::ClothAsset
 {
-	struct FClothSimulationNodeDetailExtender : public IDetailPropertyExtensionHandler
+	struct FClothSimulationNodeDetailExtender : public UE::Dataflow::IDataflowNodeDetailExtension
 	{
+		FClothSimulationNodeDetailExtender() = default;
+		static const FName Name;
+
+		// IDataflowNodeDetailExtension
+		virtual FName GetName() const override;
+
+		// IDetailPropertyExtensionHandler
 		virtual bool IsPropertyExtendable(const UClass* InObjectClass, const IPropertyHandle& PropertyHandle) const override;
 		virtual void ExtendWidgetRow(FDetailWidgetRow& InWidgetRow, const IDetailLayoutBuilder& InDetailBuilder, const UClass* InObjectClass, TSharedPtr<IPropertyHandle> PropertyHandle) override;
 	};
diff --git a/Engine/Plugins/Experimental/Dataflow/Source/DataflowEditor/Private/Dataflow/DataflowEditorToolkit.cpp b/Engine/Plugins/Experimental/Dataflow/Source/DataflowEditor/Private/Dataflow/DataflowEditorToolkit.cpp
--- a/Engine/Plugins/Experimental/Dataflow/Source/DataflowEditor/Private/Dataflow/DataflowEditorToolkit.cpp
+++ b/Engine/Plugins/Experimental/Dataflow/Source/DataflowEditor/Private/Dataflow/DataflowEditorToolkit.cpp
@@ -24,6 +24,7 @@
 #include "Dataflow/DataflowEdNode.h"
 #include "Dataflow/DataflowGraphEditor.h"
 #include "Dataflow/DataflowNodeFactory.h"
+#include "Dataflow/DataflowNodeDetailExtension.h"
 #include "Dataflow/DataflowObject.h"
 #include "Dataflow/DataflowObjectInterface.h"
 #include "Dataflow/DataflowEditorPreviewSceneBase.h"
@@ -1255,6 +1256,9 @@ TSharedPtr<IStructureDetailsView> FDataflowEditorToolkit::CreateNodeDetailsEdito
 	LocalDetailsView->GetDetailsView()->SetObject(ObjectToEdit);
 	OnFinishedChangingPropertiesDelegateHandle = LocalDetailsView->GetOnFinishedChangingPropertiesDelegate().AddSP(this, &FDataflowEditorToolkit::OnPropertyValueChanged);
 
+	NodeDetailsExtensionHandler = MakeShared<UE::Dataflow::FDataflowNodeDetailExtensionHandler>();
+	LocalDetailsView->GetDetailsView()->SetExtensionHandler(NodeDetailsExtensionHandler);
+
 	return LocalDetailsView;
 }
 
diff --git a/Engine/Plugins/Experimental/Dataflow/Source/DataflowEditor/Private/Dataflow/DataflowNodeDetailExtension.cpp b/Engine/Plugins/Experimental/Dataflow/Source/DataflowEditor/Private/Dataflow/DataflowNodeDetailExtension.cpp
new file mode 100644
--- /dev/null
+++ b/Engine/Plugins/Experimental/Dataflow/Source/DataflowEditor/Private/Dataflow/DataflowNodeDetailExtension.cpp
@@ -0,0 +1,86 @@
+// Copyright Epic Games, Inc. All Rights Reserved.
+
+#include "Dataflow/DataflowNodeDetailExtension.h"
+#include "Misc/LazySingleton.h"
+#include "ChaosLog.h"
+
+#define LOCTEXT_NAMESPACE "DataflowNodeDetailExtension"
+
+namespace UE::Dataflow
+{
+	FDataflowNodeDetailExtensionRegistry& FDataflowNodeDetailExtensionRegistry::GetInstance()
+	{
+		return TLazySingleton<FDataflowNodeDetailExtensionRegistry>::Get();
+	}
+
+	void FDataflowNodeDetailExtensionRegistry::TearDown()
+	{
+		return TLazySingleton<FDataflowNodeDetailExtensionRegistry>::TearDown();
+	}
+
+	void FDataflowNodeDetailExtensionRegistry::RegisterExtension(TUniquePtr<IDataflowNodeDetailExtension>&& Extension)
+	{
+		const FName NewExtensionName = Extension->GetName();
+		if (ExtensionMap.Contains(NewExtensionName))
+		{
+			UE_LOG(LogChaos, Warning, TEXT("Dataflow node detail extension registration conflicts with existing extension: %s"), *NewExtensionName.ToString());
+		}
+		else
+		{
+			ExtensionMap.Add(NewExtensionName, MoveTemp(Extension));
+		}
+	}
+
+	void FDataflowNodeDetailExtensionRegistry::DeregisterExtension(const FName& ExtensionName)
+	{
+		if (!ExtensionMap.Contains(ExtensionName))
+		{
+			UE_LOG(LogChaos, Warning, TEXT("Dataflow node detail extension deregistration -- extension not registered : %s"), *ExtensionName.ToString());
+		}
+		else
+		{
+			ExtensionMap.Remove(ExtensionName);
+		}
+	}
+
+	bool FDataflowNodeDetailExtensionRegistry::IsPropertyExtendable(const UClass* InObjectClass, const IPropertyHandle& PropertyHandle) const
+	{
+		for (TMap<FName, TUniquePtr<IDataflowNodeDetailExtension>>::TConstIterator Iter = ExtensionMap.CreateConstIterator(); Iter; ++Iter)
+		{
+			if (Iter.Value() && Iter.Value()->IsPropertyExtendable(InObjectClass, PropertyHandle))
+			{
+				return true;
+			}
+		}
+		return false;
+	}
+
+	void FDataflowNodeDetailExtensionRegistry::ExtendWidgetRow(FDetailWidgetRow& InWidgetRow, const IDetailLayoutBuilder& InDetailBuilder, const UClass* InObjectClass, TSharedPtr<IPropertyHandle> PropertyHandle)
+	{
+		if (!PropertyHandle)
+		{
+			return;
+		}
+
+		for (TMap<FName, TUniquePtr<IDataflowNodeDetailExtension>>::TIterator Iter = ExtensionMap.CreateIterator(); Iter; ++Iter)
+		{
+			if (Iter.Value() && Iter.Value()->IsPropertyExtendable(InObjectClass, *PropertyHandle))
+			{
+				Iter.Value()->ExtendWidgetRow(InWidgetRow, InDetailBuilder, InObjectClass, PropertyHandle);
+			}
+		}
+	}
+
+	bool FDataflowNodeDetailExtensionHandler::IsPropertyExtendable(const UClass* InObjectClass, const IPropertyHandle& PropertyHandle) const
+	{
+		return FDataflowNodeDetailExtensionRegistry::GetInstance().IsPropertyExtendable(InObjectClass, PropertyHandle);
+	}
+
+	void FDataflowNodeDetailExtensionHandler::ExtendWidgetRow(FDetailWidgetRow& InWidgetRow, const IDetailLayoutBuilder& InDetailBuilder, const UClass* InObjectClass, TSharedPtr<IPropertyHandle> PropertyHandle)
+	{
+		FDataflowNodeDetailExtensionRegistry::GetInstance().ExtendWidgetRow(InWidgetRow, InDetailBuilder, InObjectClass, PropertyHandle);
+	}
+}
+
+#undef LOCTEXT_NAMESPACE 
+
diff --git a/Engine/Plugins/Experimental/Dataflow/Source/DataflowEditor/Public/Dataflow/DataflowEditorToolkit.h b/Engine/Plugins/Experimental/Dataflow/Source/DataflowEditor/Public/Dataflow/DataflowEditorToolkit.h
--- a/Engine/Plugins/Experimental/Dataflow/Source/DataflowEditor/Public/Dataflow/DataflowEditorToolkit.h
+++ b/Engine/Plugins/Experimental/Dataflow/Source/DataflowEditor/Public/Dataflow/DataflowEditorToolkit.h
@@ -29,6 +29,10 @@ class UDataflowBaseContent;
 class FDataflowSimulationScene;
 class FDataflowSkeletonView;
 class UDataflowEditor;
+namespace UE::Dataflow
+{
+	class FDataflowNodeDetailExtensionHandler;
+}
 
 
 class DATAFLOWEDITOR_API FDataflowEditorToolkit final : public FBaseCharacterFXEditorToolkit, public FTickableEditorObject, public FNotifyHook, public FGCObject
@@ -154,6 +158,7 @@ private:
 	TSharedPtr<SDataflowSimulationViewport> DataflowSimulationViewport;
 	TSharedPtr<SDataflowGraphEditor> GraphEditor;
 	TSharedPtr<IStructureDetailsView> NodeDetailsEditor;
+	TSharedPtr<UE::Dataflow::FDataflowNodeDetailExtensionHandler> NodeDetailsExtensionHandler;
 	TSharedPtr<FDataflowSkeletonView> SkeletonEditorView;
 	TSharedPtr<IDetailsView> AssetDetailsEditor;
 	TSharedPtr<FDataflowSelectionView> DataflowSelectionView_1;
diff --git a/Engine/Plugins/Experimental/Dataflow/Source/DataflowEditor/Public/Dataflow/DataflowNodeDetailExtension.h b/Engine/Plugins/Experimental/Dataflow/Source/DataflowEditor/Public/Dataflow/DataflowNodeDetailExtension.h
new file mode 100644
--- /dev/null
+++ b/Engine/Plugins/Experimental/Dataflow/Source/DataflowEditor/Public/Dataflow/DataflowNodeDetailExtension.h
@@ -0,0 +1,58 @@
+// Copyright Epic Games, Inc. All Rights Reserved.
+
+#pragma once
+
+#include "UObject/NameTypes.h"
+#include "Templates/SharedPointer.h"
+#include "IDetailPropertyExtensionHandler.h"
+
+namespace UE::Dataflow
+{
+	/**
+	* IDetailPropertyExtensionHandler that can be registered with FDataflowNodeDetailExtensionRegistry to be aggregated and applied to the Dataflow Node Details panel.
+	*/
+	class IDataflowNodeDetailExtension : public IDetailPropertyExtensionHandler
+	{
+	public:
+		virtual ~IDataflowNodeDetailExtension() = default;
+
+		DATAFLOWEDITOR_API virtual FName GetName() const = 0;
+	};
+
+	/**
+	* Registry for IDataflowNodeDetailExtension. Register/Deregister extensions to extend the Dataflow Node Details rows via IDetailPropertyExtensionHandler.
+	* The only reason why this isn't an IDetailPropertyExtensionHandler is because DetailsViews want a SharedPtr of an IDetailPropertyExtensionHandler, but this is a singleton registry.
+	*/
+	class FDataflowNodeDetailExtensionRegistry
+	{
+	public:
+
+		// FLazySingleton
+		static DATAFLOWEDITOR_API FDataflowNodeDetailExtensionRegistry& GetInstance();
+		static DATAFLOWEDITOR_API void TearDown();
+
+		DATAFLOWEDITOR_API void RegisterExtension(TUniquePtr<IDataflowNodeDetailExtension>&& Extension);
+		DATAFLOWEDITOR_API void DeregisterExtension(const FName& ExtensionName);
+
+		// These methods aggregate calls on the registered extensions.
+		DATAFLOWEDITOR_API bool IsPropertyExtendable(const UClass* InObjectClass, const IPropertyHandle& PropertyHandle) const;
+		DATAFLOWEDITOR_API void ExtendWidgetRow(FDetailWidgetRow& InWidgetRow, const IDetailLayoutBuilder& InDetailBuilder, const UClass* InObjectClass, TSharedPtr<IPropertyHandle> PropertyHandle);
+
+	private:
+
+		TMap<FName, TUniquePtr<IDataflowNodeDetailExtension>> ExtensionMap;
+	};
+
+
+	/**
+	* This is the IDetailPropertyExtensionHandler that the DataflowToolkit actually uses.It simply calls the equivalent methods on the FDataflowNodeDetailExtensionRegistry singleton.
+	*/
+	class FDataflowNodeDetailExtensionHandler : public IDetailPropertyExtensionHandler
+	{
+	public:
+		// IDetailPropertyExtensionHandler
+		DATAFLOWEDITOR_API virtual bool IsPropertyExtendable(const UClass* InObjectClass, const IPropertyHandle& PropertyHandle) const override;
+		DATAFLOWEDITOR_API virtual void ExtendWidgetRow(FDetailWidgetRow& InWidgetRow, const IDetailLayoutBuilder& InDetailBuilder, const UClass* InObjectClass, TSharedPtr<IPropertyHandle> PropertyHandle) override;
+	};
+}
+
