Dataflow Editor: add top toolbar buttons for adding nodes and launching interactive tools

diff --git a/Engine/Plugins/ChaosClothAssetEditor/Source/ChaosClothAssetEditorTools/Private/ChaosClothAsset/ChaosClothAssetEditorToolsModule.cpp b/Engine/Plugins/ChaosClothAssetEditor/Source/ChaosClothAssetEditorTools/Private/ChaosClothAsset/ChaosClothAssetEditorToolsModule.cpp
--- a/Engine/Plugins/ChaosClothAssetEditor/Source/ChaosClothAssetEditorTools/Private/ChaosClothAsset/ChaosClothAssetEditorToolsModule.cpp
+++ b/Engine/Plugins/ChaosClothAssetEditor/Source/ChaosClothAssetEditorTools/Private/ChaosClothAsset/ChaosClothAssetEditorToolsModule.cpp
@@ -9,6 +9,8 @@
 #include "ChaosClothAsset/WeightMapNode.h"
 #include "ChaosClothAsset/ClothToolActionCommandBindings.h"
 
+#define LOCTEXT_NAMESPACE "ChaosClothAssetEditorToolsModule"
+
 namespace UE::Chaos::ClothAsset
 {
 	class FChaosClothAssetEditorToolsModule : public IModuleInterface
@@ -20,9 +22,23 @@ namespace UE::Chaos::ClothAsset
 		
 			const TSharedRef<const FClothToolActionCommandBindings> ClothToolActions = MakeShared<FClothToolActionCommandBindings>();
 
-			ToolRegistry.AddNodeToToolMapping(FChaosClothAssetWeightMapNode::StaticType(), NewObject<UClothEditorWeightMapPaintToolBuilder>(), ClothToolActions);
-			ToolRegistry.AddNodeToToolMapping(FChaosClothAssetSelectionNode_v2::StaticType(), NewObject<UClothMeshSelectionToolBuilder>(), ClothToolActions);
-			ToolRegistry.AddNodeToToolMapping(FChaosClothAssetTransferSkinWeightsNode::StaticType(), NewObject<UClothTransferSkinWeightsToolBuilder>(), ClothToolActions);
+			ToolRegistry.AddNodeToToolMapping(FChaosClothAssetWeightMapNode::StaticType(),          // Node type
+				NewObject<UClothEditorWeightMapPaintToolBuilder>(),                                 // Tool builder
+				ClothToolActions,                                                                   // Tool action bindings
+				FSlateIcon(FName("ClothStyle"), FName("ChaosClothAssetEditor.AddWeightMapNode")),   // Button icon (see FChaosClothAssetEditorStyle)
+				LOCTEXT("AddWeightMapNodeButtonText", "Weight Map"));                               // Button text
+
+			ToolRegistry.AddNodeToToolMapping(FChaosClothAssetSelectionNode_v2::StaticType(), 
+				NewObject<UClothMeshSelectionToolBuilder>(), 
+				ClothToolActions, 
+				FSlateIcon(FName("ClothStyle"), FName("ChaosClothAssetEditor.AddMeshSelectionNode")),
+				LOCTEXT("AddSelectionNodeButtonText", "Selection"));
+
+			ToolRegistry.AddNodeToToolMapping(FChaosClothAssetTransferSkinWeightsNode::StaticType(), 
+				NewObject<UClothTransferSkinWeightsToolBuilder>(), 
+				ClothToolActions, 
+				FSlateIcon(FName("ClothStyle"), FName("ChaosClothAssetEditor.AddTransferSkinWeightsNode")),
+				LOCTEXT("AddTransferSkinWeightNodeButtonText", "Transfer Skin Weights"));
 		}
 
 		virtual void ShutdownModule() override
@@ -36,3 +52,5 @@ namespace UE::Chaos::ClothAsset
 }
 
 IMPLEMENT_MODULE(UE::Chaos::ClothAsset::FChaosClothAssetEditorToolsModule, ChaosClothAssetEditorTools)
+
+#undef LOCTEXT_NAMESPACE
diff --git a/Engine/Plugins/Experimental/Dataflow/Source/DataflowEditor/Private/Dataflow/DataflowEditorCommands.cpp b/Engine/Plugins/Experimental/Dataflow/Source/DataflowEditor/Private/Dataflow/DataflowEditorCommands.cpp
--- a/Engine/Plugins/Experimental/Dataflow/Source/DataflowEditor/Private/Dataflow/DataflowEditorCommands.cpp
+++ b/Engine/Plugins/Experimental/Dataflow/Source/DataflowEditor/Private/Dataflow/DataflowEditorCommands.cpp
@@ -106,6 +106,18 @@ void FDataflowEditorCommandsImpl::RegisterCommands()
 			EUserInterfaceActionType::Button,
 			FInputChord()
 		);
+
+		FUICommandInfo::MakeCommandInfo(
+			this->AsShared(),
+			ToolRegistry.GetAddNodeCommandForNode(NodeName),
+			FName("Add_" + NodeName.ToString()),
+			ToolRegistry.GetAddNodeButtonText(NodeName),
+			FText::Format(LOCTEXT("AddDataflowNodeTooltip", "Add a \"{0}\" node to the graph"), FText::FromString(NodeName.ToString())),
+			ToolRegistry.GetAddNodeButtonIcon(NodeName),
+			EUserInterfaceActionType::Button,
+			FInputChord()
+		);
+
 	}
 
 }
diff --git a/Engine/Plugins/Experimental/Dataflow/Source/DataflowEditor/Private/Dataflow/DataflowEditorMode.cpp b/Engine/Plugins/Experimental/Dataflow/Source/DataflowEditor/Private/Dataflow/DataflowEditorMode.cpp
--- a/Engine/Plugins/Experimental/Dataflow/Source/DataflowEditor/Private/Dataflow/DataflowEditorMode.cpp
+++ b/Engine/Plugins/Experimental/Dataflow/Source/DataflowEditor/Private/Dataflow/DataflowEditorMode.cpp
@@ -240,13 +240,18 @@ void UDataflowEditorMode::AddNode(FName NewNodeType)
 	const UEdGraphNode* const NewNode = CreateAndConnectNewNode(NewNodeType, *CurrentlySelectedNode, ConnectionType, ConnectionName);
 	verifyf(NewNode, TEXT("Failed to create a new node: %s"), *NewNodeType.ToString());
 
+	// Wait for FDataflowEditorToolkit::OnNodeSelectionChanged to execute before actually starting the tool
+	bPendingNodeSelectionChanged = true;
+
+	// This will queue the tool to start after FDataflowEditorToolkit::OnNodeSelectionChanged finishes
 	StartToolForSelectedNode(NewNode);
 }
 
 bool UDataflowEditorMode::CanAddNode(FName NewNodeType) const
 {
 	const UEdGraphNode* const CurrentlySelectedNode = GetSingleSelectedNodeWithOutputType(FManagedArrayCollection::StaticType());
-	return (CurrentlySelectedNode != nullptr);
+	const bool bToolActive = GetToolManager()->HasActiveTool(EToolSide::Left);
+	return CurrentlySelectedNode && !bToolActive;
 }
 
 void UDataflowEditorMode::RegisterTools()
@@ -271,10 +276,11 @@ void UDataflowEditorMode::RegisterTools()
 		RegisterDataflowTool(CommandInfo, RegisteredNodeName.ToString() + FString(TEXT("Tool")), Builder, ConstructionViewportToolsContext);
 
 		NodeTypeToToolCommandMap.Add(RegisteredNodeName, CommandInfo);
+
+		// Register "Add Node" commands for buttons in the UI. The EditorToolkit will construct the actual toolbar buttons.
+		NodeTypeToAddNodeCommandMap.Add(RegisteredNodeName, ToolRegistry.GetAddNodeCommandForNode(RegisteredNodeName));
 	}
 
-	// Register "Add Node" commands for buttons in the UI. The EditorToolkit will construct the actual toolbar buttons.
-	NodeTypeToAddNodeCommandMap.Add(FDataflowCollectionAddScalarVertexPropertyNode::StaticType(), CommandInfos.AddWeightMapNode);
 }
 
 bool UDataflowEditorMode::ShouldToolStartBeAllowed(const FString& ToolIdentifier) const
@@ -589,7 +595,7 @@ void UDataflowEditorMode::ModeTick(float DeltaTime)
 		}
 	}
 
-	if (!NodeTypeForPendingToolStart.IsNone() && !GetToolManager()->HasActiveTool(EToolSide::Left))
+	if (!NodeTypeForPendingToolStart.IsNone() && !GetToolManager()->HasActiveTool(EToolSide::Left) && !bPendingNodeSelectionChanged)
 	{
 		const TSharedRef<FUICommandList> CommandList = Toolkit->GetToolkitCommands();
 		const FDataflowEditorCommandsImpl& CommandInfos = FDataflowEditorCommandsImpl::Get();
diff --git a/Engine/Plugins/Experimental/Dataflow/Source/DataflowEditor/Private/Dataflow/DataflowEditorModule.cpp b/Engine/Plugins/Experimental/Dataflow/Source/DataflowEditor/Private/Dataflow/DataflowEditorModule.cpp
--- a/Engine/Plugins/Experimental/Dataflow/Source/DataflowEditor/Private/Dataflow/DataflowEditorModule.cpp
+++ b/Engine/Plugins/Experimental/Dataflow/Source/DataflowEditor/Private/Dataflow/DataflowEditorModule.cpp
@@ -86,7 +86,7 @@ void FDataflowEditorModule::StartupModule()
 
 	UDataflowEditorWeightMapPaintToolBuilder* const ToolBuilder = NewObject<UDataflowEditorWeightMapPaintToolBuilder>();
 	TSharedRef<const FDataflowToolActionCommandBindings> Actions = MakeShared<FDataflowToolActionCommandBindings>();
-	ToolRegistry.AddNodeToToolMapping(FDataflowCollectionAddScalarVertexPropertyNode::StaticType(), ToolBuilder, Actions);
+	ToolRegistry.AddNodeToToolMapping(FDataflowCollectionAddScalarVertexPropertyNode::StaticType(), ToolBuilder, Actions, FSlateIcon(), LOCTEXT("AddWeightMapNodeButtonText", "Scalar Vertex Property"));
 
 	UE::Dataflow::FDataflowConstructionVisualizationRegistry& ConstructionVisualizationRegistry = UE::Dataflow::FDataflowConstructionVisualizationRegistry::GetInstance();
 	ConstructionVisualizationRegistry.RegisterVisualization(MakeUnique<UE::Dataflow::FMeshStatsConstructionVisualization>());
diff --git a/Engine/Plugins/Experimental/Dataflow/Source/DataflowEditor/Private/Dataflow/DataflowEditorToolkit.cpp b/Engine/Plugins/Experimental/Dataflow/Source/DataflowEditor/Private/Dataflow/DataflowEditorToolkit.cpp
--- a/Engine/Plugins/Experimental/Dataflow/Source/DataflowEditor/Private/Dataflow/DataflowEditorToolkit.cpp
+++ b/Engine/Plugins/Experimental/Dataflow/Source/DataflowEditor/Private/Dataflow/DataflowEditorToolkit.cpp
@@ -506,7 +506,7 @@ void FDataflowEditorToolkit::PostInitAssetEditor()
 	FName ParentToolbarName;
 	const FName ToolBarName = GetToolMenuToolbarName(ParentToolbarName);
 	UToolMenu* const AssetToolbar = UToolMenus::Get()->ExtendMenu(ToolBarName);
-	FToolMenuSection& Section = AssetToolbar->FindOrAddSection("ClothTools");
+	FToolMenuSection& Section = AssetToolbar->FindOrAddSection("DataflowTools");
 
 	for (const TPair<FName, TSharedPtr<const FUICommandInfo>>& NodeAndAddCommand : DataflowMode->NodeTypeToAddNodeCommandMap)
 	{
@@ -1052,6 +1052,8 @@ void FDataflowEditorToolkit::OnNodeSelectionChanged(const TSet<UObject*>& InNewS
 			// TODO: Clear and disable View Mode Button. For now set default mode to the built-in 3D view mode.
 			DataflowMode->SetConstructionViewMode(UE::Dataflow::FDataflowConstruction3DViewMode::Name);
 		}
+
+		DataflowMode->SetPendingNodeSelectionChanged(false);
 	}
 }
 
diff --git a/Engine/Plugins/Experimental/Dataflow/Source/DataflowEditor/Private/Dataflow/DataflowToolRegistry.cpp b/Engine/Plugins/Experimental/Dataflow/Source/DataflowEditor/Private/Dataflow/DataflowToolRegistry.cpp
--- a/Engine/Plugins/Experimental/Dataflow/Source/DataflowEditor/Private/Dataflow/DataflowToolRegistry.cpp
+++ b/Engine/Plugins/Experimental/Dataflow/Source/DataflowEditor/Private/Dataflow/DataflowToolRegistry.cpp
@@ -18,16 +18,30 @@ namespace UE::Dataflow
 		TLazySingleton<FDataflowToolRegistry>::TearDown();
 	}
 
-	void FDataflowToolRegistry::AddNodeToToolMapping(const FName& NodeName, TObjectPtr<UInteractiveToolBuilder> ToolBuilder, const TSharedRef<const IDataflowToolActionCommands>& ToolActionRegistry)
+	void FDataflowToolRegistry::AddNodeToToolMapping(const FName& NodeName, TObjectPtr<UInteractiveToolBuilder> ToolBuilder, const TSharedRef<const IDataflowToolActionCommands>& ToolActionCommands, const FSlateIcon& AddNodeButtonIcon, const FText& AddNodeButtonText)
 	{
 		ensureMsgf(!FDataflowEditorCommands::IsRegistered(), TEXT("FDataflowToolRegistry: DataflowEditorCommands have already been registered. \
 				Newly registered Tools may not be available in the Editor. \
 				Ensure that AddNodeToToolMapping is called before the DataflowEditor module is loaded."));
 
 		// FUICommandInfo is uninitialized, it will be created later in FDataflowEditorCommandsImpl::RegisterCommands
-		NodeTypeToToolMap.Add(NodeName, { ToolBuilder, ToolActionRegistry, nullptr });
+		NodeTypeToToolMap.Add(NodeName, FToolInfo{ ToolBuilder, ToolActionCommands, AddNodeButtonIcon, AddNodeButtonText, nullptr, nullptr });
 	}
 
+	PRAGMA_DISABLE_DEPRECATION_WARNINGS
+	void FDataflowToolRegistry::AddNodeToToolMapping(const FName& NodeName,
+		TObjectPtr<UInteractiveToolBuilder> ToolBuilder,
+		const TSharedRef<const IDataflowToolActionCommands>& ToolActionCommands)
+	{
+		ensureMsgf(!FDataflowEditorCommands::IsRegistered(), TEXT("FDataflowToolRegistry: DataflowEditorCommands have already been registered. \
+				Newly registered Tools may not be available in the Editor. \
+				Ensure that AddNodeToToolMapping is called before the DataflowEditor module is loaded."));
+
+		// FUICommandInfo is uninitialized, it will be created later in FDataflowEditorCommandsImpl::RegisterCommands
+		NodeTypeToToolMap.Add(NodeName, FToolInfo{ ToolBuilder, ToolActionCommands, FSlateIcon(), FText(), nullptr, nullptr });
+	}
+	PRAGMA_ENABLE_DEPRECATION_WARNINGS
+
 	void FDataflowToolRegistry::RemoveNodeToToolMapping(const FName& NodeName)
 	{
 		NodeTypeToToolMap.Remove(NodeName);
@@ -40,6 +54,24 @@ namespace UE::Dataflow
 		return NodeNames;
 	}
 	
+	TSharedPtr<FUICommandInfo>& FDataflowToolRegistry::GetAddNodeCommandForNode(const FName& NodeType)
+	{
+		check(NodeTypeToToolMap.Contains(NodeType));
+		return NodeTypeToToolMap[NodeType].AddNodeCommand;
+	}
+
+	const FSlateIcon& FDataflowToolRegistry::GetAddNodeButtonIcon(const FName& NodeType) const
+	{
+		check(NodeTypeToToolMap.Contains(NodeType));
+		return NodeTypeToToolMap[NodeType].AddNodeButtonIcon;
+	}
+
+	const FText& FDataflowToolRegistry::GetAddNodeButtonText(const FName& NodeType) const
+	{
+		check(NodeTypeToToolMap.Contains(NodeType));
+		return NodeTypeToToolMap[NodeType].AddNodeButtonText;
+	}
+
 	TSharedPtr<FUICommandInfo>& FDataflowToolRegistry::GetToolCommandForNode(const FName& NodeName)
 	{
 		check(NodeTypeToToolMap.Contains(NodeName));
diff --git a/Engine/Plugins/Experimental/Dataflow/Source/DataflowEditor/Public/Dataflow/DataflowEditorMode.h b/Engine/Plugins/Experimental/Dataflow/Source/DataflowEditor/Public/Dataflow/DataflowEditorMode.h
--- a/Engine/Plugins/Experimental/Dataflow/Source/DataflowEditor/Public/Dataflow/DataflowEditorMode.h
+++ b/Engine/Plugins/Experimental/Dataflow/Source/DataflowEditor/Public/Dataflow/DataflowEditorMode.h
@@ -168,6 +168,11 @@ private:
 	void InitializeContextObject();
 	void DeleteContextObject();
 
+	void SetPendingNodeSelectionChanged(bool bIsPending)
+	{
+		bPendingNodeSelectionChanged = bIsPending;
+	}
+
 	/** Dataflow construction scene from the toolkit */
 	FDataflowConstructionScene* ConstructionScene = nullptr;
 
@@ -223,6 +228,9 @@ private:
 	// losing selection information, so we save this before shutting the tool down.
 	bool bHadSingleSelectionBeforeToolShutdown = false;
 
+	// Whether we are waiting for a node selection change to finish before starting a tool
+	bool bPendingNodeSelectionChanged = false;
+
 	UDataflowEditor* DataflowEditor = nullptr;
 
 	// Timestamp for telemetry
diff --git a/Engine/Plugins/Experimental/Dataflow/Source/DataflowEditor/Public/Dataflow/DataflowToolRegistry.h b/Engine/Plugins/Experimental/Dataflow/Source/DataflowEditor/Public/Dataflow/DataflowToolRegistry.h
--- a/Engine/Plugins/Experimental/Dataflow/Source/DataflowEditor/Public/Dataflow/DataflowToolRegistry.h
+++ b/Engine/Plugins/Experimental/Dataflow/Source/DataflowEditor/Public/Dataflow/DataflowToolRegistry.h
@@ -5,6 +5,7 @@
 #include "UObject/NameTypes.h"
 #include "Templates/SharedPointer.h"
 #include "UObject/ObjectPtr.h"
+#include "Textures/SlateIcon.h"
 
 class FUICommandInfo;
 class UInteractiveToolBuilder;
@@ -30,11 +31,25 @@ namespace UE::Dataflow
 		static FDataflowToolRegistry& Get();
 		static void TearDown();
 
-		void AddNodeToToolMapping(const FName& NodeName, TObjectPtr<UInteractiveToolBuilder> ToolBuilder, const TSharedRef<const IDataflowToolActionCommands>& ToolActionCommands);
+		void AddNodeToToolMapping(const FName& NodeName, 
+			TObjectPtr<UInteractiveToolBuilder> ToolBuilder, 
+			const TSharedRef<const IDataflowToolActionCommands>& ToolActionCommands, 
+			const FSlateIcon& AddNodeButtonIcon, 
+			const FText& AddNodeButtonText);
+
+		UE_DEPRECATED(5.6, "Please use the version of AddNodeToToolMapping taking AddNodeButtonIcon and AddNodeButtonText")
+		void AddNodeToToolMapping(const FName& NodeName,
+			TObjectPtr<UInteractiveToolBuilder> ToolBuilder,
+			const TSharedRef<const IDataflowToolActionCommands>& ToolActionCommands);
+
 		void RemoveNodeToToolMapping(const FName& NodeName);
 
 		TArray<FName> GetNodeNames() const;
 
+		TSharedPtr<FUICommandInfo>& GetAddNodeCommandForNode(const FName& NodeType);
+		const FSlateIcon& GetAddNodeButtonIcon(const FName& NodeType) const;
+		const FText& GetAddNodeButtonText(const FName& NodeType) const;
+
 		TSharedPtr<FUICommandInfo>& GetToolCommandForNode(const FName& NodeName);
 
 		UInteractiveToolBuilder* GetToolBuilderForNode(const FName& NodeName);
@@ -50,8 +65,11 @@ namespace UE::Dataflow
 			// Specified when registering the tool
 			TObjectPtr<UInteractiveToolBuilder> ToolBuilder;
 			TSharedRef<const IDataflowToolActionCommands> ToolActionCommands;
+			FSlateIcon AddNodeButtonIcon;
+			FText AddNodeButtonText;
 
 			// Constructed automatically in FDataflowEditorCommandsImpl::RegisterCommands
+			TSharedPtr<FUICommandInfo> AddNodeCommand;
 			TSharedPtr<FUICommandInfo> ToolCommand;
 		};
 
