Dataflow:
- Added point rendering
- Added points rendering
- Added FSphere rendering
- Added Tetrahedron rendering
- Fixed tetrahedron data display in CollectionSpreadsheet
- Added methods to FBoundsFacade to compute bounding sphere
- Added FBoundingSphere, FExpandBoundingSphere nodes
- Added FMakeCollectionDataflowNode
- Added FPointsToCollectionDataflowNode, FCollectionToPointsDataflowNode
- Added FVisualizeTetrahedronsDataflowNode node
- Bug fixes

diff --git a/Engine/Plugins/Experimental/Dataflow/Source/DataflowEditor/Private/Dataflow/CollectionSpreadSheetWidget.cpp b/Engine/Plugins/Experimental/Dataflow/Source/DataflowEditor/Private/Dataflow/CollectionSpreadSheetWidget.cpp
--- a/Engine/Plugins/Experimental/Dataflow/Source/DataflowEditor/Private/Dataflow/CollectionSpreadSheetWidget.cpp
+++ b/Engine/Plugins/Experimental/Dataflow/Source/DataflowEditor/Private/Dataflow/CollectionSpreadSheetWidget.cpp
@@ -115,6 +115,11 @@ namespace CollectionSpreadSheetHelpers
 		return FString::Printf(TEXT("%d %d %d"), Value.X, Value.Y, Value.Z);
 	}
 
+	FString AttributeValueToString(FIntVector4 Value)
+	{
+		return FString::Printf(TEXT("%d %d %d %d"), Value[0], Value[1], Value[2], Value[3]);
+	}
+
 	template<typename T>
 	FString AttributeValueToString(const FManagedArrayCollection& InCollection, const FName& InAttributeName, const FName& InGroupName, int32 InIdxColumn)
 	{
@@ -204,6 +209,10 @@ namespace CollectionSpreadSheetHelpers
 			ValueAsString = AttributeValueToString<FTransform3f>(InCollection, InAttributeName, InGroupName, InIdxColumn);
 			break;
 
+		case FManagedArrayCollection::EArrayType::FIntVector4Type:
+			ValueAsString = AttributeValueToString<FIntVector4>(InCollection, InAttributeName, InGroupName, InIdxColumn);
+			break;
+
 		default:
 			//ensure(false);
 			ValueAsString = "<Unknown Data Type>";
diff --git a/Engine/Plugins/Experimental/Dataflow/Source/DataflowEditor/Private/Dataflow/DataflowEngineRendering.cpp b/Engine/Plugins/Experimental/Dataflow/Source/DataflowEditor/Private/Dataflow/DataflowEngineRendering.cpp
--- a/Engine/Plugins/Experimental/Dataflow/Source/DataflowEditor/Private/Dataflow/DataflowEngineRendering.cpp
+++ b/Engine/Plugins/Experimental/Dataflow/Source/DataflowEditor/Private/Dataflow/DataflowEngineRendering.cpp
@@ -488,51 +488,211 @@ namespace UE::Dataflow
 				FBox Default(ForceInit);
 				const FBox& Box = State.GetValue<FBox>(PrimaryOutput, Default);
 
-				const int32 NumVertices = 8;
-				const int32 NumTriangles = 12;
-
-				TArray<FVector3f> Vertices; Vertices.AddUninitialized(NumVertices);
-				TArray<FIntVector> Tris; Tris.AddUninitialized(NumTriangles);
-
-				FVector Min = Box.Min;
-				FVector Max = Box.Max;
-
-				// Add vertices
-				Vertices[0] = FVector3f(Min);
-				Vertices[1] = FVector3f(Max.X, Min.Y, Min.Z);
-				Vertices[2] = FVector3f(Max.X, Max.Y, Min.Z);
-				Vertices[3] = FVector3f(Min.X, Max.Y, Min.Z);
-				Vertices[4] = FVector3f(Min.X, Min.Y, Max.Z);
-				Vertices[5] = FVector3f(Max.X, Min.Y, Max.Z);
-				Vertices[6] = FVector3f(Max);
-				Vertices[7] = FVector3f(Min.X, Max.Y, Max.Z);
-
-				// Add triangles
-				Tris[0] = FIntVector(0, 1, 3); Tris[1] = FIntVector(1, 2, 3);
-				Tris[2] = FIntVector(0, 4, 1); Tris[3] = FIntVector(4, 5, 1);
-				Tris[4] = FIntVector(5, 2, 1); Tris[5] = FIntVector(5, 6, 2);
-				Tris[6] = FIntVector(3, 2, 6); Tris[7] = FIntVector(7, 3, 6);
-				Tris[8] = FIntVector(0, 3, 7); Tris[9] = FIntVector(4, 0, 7);
-				Tris[10] = FIntVector(5, 4, 7); Tris[11] = FIntVector(5, 7, 6);
-
-				TArray<FVector3f> VertexNormals; VertexNormals.AddUninitialized(NumVertices);
-				// TODO: Compute vertex normals
-
-				// Add VertexNormal and VertexColor
-				TArray<FLinearColor> VertexColors; VertexColors.AddUninitialized(Vertices.Num());
-				for (int32 VertexIdx = 0; VertexIdx < VertexNormals.Num(); ++VertexIdx)
+				int32 GeometryIndex = RenderCollection.StartGeometryGroup(State.GetGuid().ToString());
+				RenderCollection.AddBox(FVector3f(Box.Min), FVector3f(Box.Max));
+				RenderCollection.EndGeometryGroup(GeometryIndex);
+			}
+		}
+	};
+
+
+	class FSphereSurfaceRenderCallbacks : public FRenderingFactory::ICallbackInterface
+	{
+		virtual UE::Dataflow::FRenderKey GetRenderKey() const override
+		{
+			return { "SurfaceRender", FName("FSphere") };
+		}
+
+		virtual bool CanRender(const UE::Dataflow::IDataflowConstructionViewMode& ViewMode) const override
+		{
+			return (ViewMode.GetName() == FDataflowConstruction3DViewMode::Name);
+		}
+
+		virtual void Render(GeometryCollection::Facades::FRenderingFacade& RenderCollection, const FGraphRenderingState& State) override
+		{
+			if (State.GetRenderOutputs().Num())
+			{
+				FName PrimaryOutput = State.GetRenderOutputs()[0]; // "Sphere"
+
+				FSphere Default(ForceInit);
+				const FSphere& Sphere = State.GetValue<FSphere>(PrimaryOutput, Default);
+
+				int32 GeometryIndex = RenderCollection.StartGeometryGroup(State.GetGuid().ToString());
+				RenderCollection.AddSphere(FVector3f(Sphere.Center), float(Sphere.W), FLinearColor(FDataflowEditorModule::SurfaceColor));
+				RenderCollection.EndGeometryGroup(GeometryIndex);
+			}
+		}
+	};
+
+	class FPointRenderCallbacks : public FRenderingFactory::ICallbackInterface
+	{
+		virtual UE::Dataflow::FRenderKey GetRenderKey() const override
+		{
+			return { "PointRender", FName("FVector") };
+		}
+
+		virtual bool CanRender(const UE::Dataflow::IDataflowConstructionViewMode& ViewMode) const override
+		{
+			return (ViewMode.GetName() == FDataflowConstruction3DViewMode::Name);
+		}
+
+		virtual void Render(GeometryCollection::Facades::FRenderingFacade& RenderCollection, const FGraphRenderingState& State) override
+		{
+			if (State.GetRenderOutputs().Num())
+			{
+				FName PrimaryOutput = State.GetRenderOutputs()[0]; // "Point"
+
+				FVector Default;
+				const FVector& Point = State.GetValue<FVector>(PrimaryOutput, Default);
+
+
+
+				int32 GeometryIndex = RenderCollection.StartGeometryGroup(State.GetGuid().ToString());
+				RenderCollection.AddPoint(FVector3f(Point));
+				RenderCollection.EndGeometryGroup(GeometryIndex);
+			}
+		}
+	};
+
+	class FPointsRenderCallbacks : public FRenderingFactory::ICallbackInterface
+	{
+		virtual UE::Dataflow::FRenderKey GetRenderKey() const override
+		{
+			return { "PointsRender", FName("TArray<FVector>") };
+		}
+
+		virtual bool CanRender(const UE::Dataflow::IDataflowConstructionViewMode& ViewMode) const override
+		{
+			return (ViewMode.GetName() == FDataflowConstruction3DViewMode::Name);
+		}
+
+		virtual void Render(GeometryCollection::Facades::FRenderingFacade& RenderCollection, const FGraphRenderingState& State) override
+		{
+			if (State.GetRenderOutputs().Num())
+			{
+				FName PrimaryOutput = State.GetRenderOutputs()[0]; // "Points"
+
+				TArray<FVector> Default;
+				const TArray<FVector>& Points = State.GetValue<TArray<FVector>>(PrimaryOutput, Default);
+
+				TArray<FVector3f> PointsArr; PointsArr.AddUninitialized(Points.Num());
+				for (int32 Idx = 0; Idx < Points.Num(); ++Idx)
 				{
-					VertexNormals[VertexIdx] = FVector3f(0.f);
-					VertexColors[VertexIdx] = FLinearColor(FDataflowEditorModule::SurfaceColor);
+					PointsArr[Idx] = (FVector3f)Points[Idx];
 				}
 
 				int32 GeometryIndex = RenderCollection.StartGeometryGroup(State.GetGuid().ToString());
-				RenderCollection.AddSurface(MoveTemp(Vertices), MoveTemp(Tris), MoveTemp(VertexNormals), MoveTemp(VertexColors));
+				RenderCollection.AddPoints(MoveTemp(PointsArr));
 				RenderCollection.EndGeometryGroup(GeometryIndex);
 			}
 		}
 	};
 
+	void RenderTetrahedronGeometryCollection(GeometryCollection::Facades::FRenderingFacade& RenderCollection, const UE::Dataflow::FGraphRenderingState& State, TArray<FLinearColor>* VertexColorOverride = nullptr)
+	{
+		auto ToD = [](FVector3f V) { return FVector3d(V.X, V.Y, V.Z); };
+		auto ToF = [](FVector3d V) { return FVector3f(V.X, V.Y, V.Z); };
+
+		FManagedArrayCollection Default;
+		FName PrimaryOutput = State.GetRenderOutputs()[0]; // "Collection"
+		const FManagedArrayCollection& Collection = State.GetValue<FManagedArrayCollection>(PrimaryOutput, Default);
+
+		const TManagedArray<int32>& BoneIndex = Collection.GetAttribute<int32>("BoneMap", FGeometryCollection::VerticesGroup);
+		const TManagedArray<int32>& Parents = Collection.GetAttribute<int32>(FTransformCollection::ParentAttribute, FTransformCollection::TransformGroup);
+		const TManagedArray<FTransform3f>& Transforms = Collection.GetAttribute<FTransform3f>(FTransformCollection::TransformAttribute, FTransformCollection::TransformGroup);
+		const TManagedArray<FString>& BoneNames = Collection.GetAttribute<FString>("BoneName", FGeometryCollection::TransformGroup);
+		const TManagedArray<FVector3f>& Vertex = Collection.GetAttribute<FVector3f>("Vertex", FGeometryCollection::VerticesGroup);
+		const TManagedArray<int32>& VertexStart = Collection.GetAttribute<int32>("VertexStart", FGeometryCollection::GeometryGroup);
+		const TManagedArray<int32>& VertexCount = Collection.GetAttribute<int32>("VertexCount", FGeometryCollection::GeometryGroup);
+		const TManagedArray<int32>& TetrahedronStart = Collection.GetAttribute<int32>("TetrahedronStart", FGeometryCollection::GeometryGroup);
+		const TManagedArray<int32>& TetrahedronCount = Collection.GetAttribute<int32>("TetrahedronCount", FGeometryCollection::GeometryGroup);
+		const TManagedArray<FIntVector4>& Tetrahedrons = Collection.GetAttribute<FIntVector4>("Tetrahedron", "Tetrahedral");
+
+		TArray<FTransform> M;
+		GeometryCollectionAlgo::GlobalMatrices(Transforms, Parents, M);
+
+		for (int GeometryIdx = 0; GeometryIdx < Collection.NumElements(FGeometryCollection::GeometryGroup); GeometryIdx++)
+		{
+			TArray<FVector3f> VerticesInCollectionSpace; VerticesInCollectionSpace.AddUninitialized(VertexCount[GeometryIdx]);
+			TArray<FVector3f> SplitVertices;
+
+			// Transform vertices to Collection space
+			for (int32 VertexIdx = VertexStart[GeometryIdx]; VertexIdx < VertexStart[GeometryIdx] + VertexCount[GeometryIdx]; ++VertexIdx)
+			{
+				VerticesInCollectionSpace[VertexIdx] = ToF(M[GeometryIdx].TransformPosition(ToD(Vertex[VertexIdx])));
+			}
+
+			TArray<FIntVector4> Tetras; Tetras.AddUninitialized(TetrahedronCount[GeometryIdx]);
+
+			for (int32 TetrahedronIdx = TetrahedronStart[GeometryIdx]; TetrahedronIdx < TetrahedronStart[GeometryIdx] + TetrahedronCount[GeometryIdx]; ++TetrahedronIdx)
+			{
+				const FIntVector4& Tetra = Tetrahedrons[TetrahedronIdx];
+				const int32 VtxStart = SplitVertices.Num();
+
+				SplitVertices.Add(VerticesInCollectionSpace[Tetra[0]]);
+				SplitVertices.Add(VerticesInCollectionSpace[Tetra[1]]);
+				SplitVertices.Add(VerticesInCollectionSpace[Tetra[2]]);
+				SplitVertices.Add(VerticesInCollectionSpace[Tetra[3]]);
+
+				Tetras[TetrahedronIdx] = { VtxStart, VtxStart + 1, VtxStart + 2, VtxStart + 3 };
+			}
+
+			if (SplitVertices.Num() && Tetras.Num())
+			{
+				FString GeometryName = State.GetGuid().ToString(); GeometryName.AppendChar('.').AppendInt(GeometryIdx);
+				if (BoneIndex[VertexStart[GeometryIdx]] != INDEX_NONE)
+				{
+					GeometryName = BoneNames[BoneIndex[VertexStart[GeometryIdx]]];
+				}
+				int32 GeometryIndex = RenderCollection.StartGeometryGroup(GeometryName);
+				RenderCollection.AddTetrahedrons(MoveTemp(SplitVertices), MoveTemp(Tetras));
+				RenderCollection.EndGeometryGroup(GeometryIndex);
+			}
+		}
+	}
+
+	class FGeometryCollectionTetrahedronRenderCallbacks : public FRenderingFactory::ICallbackInterface
+	{
+		virtual UE::Dataflow::FRenderKey GetRenderKey() const override
+		{
+			return { "TetrahedronRender", FGeometryCollection::StaticType() };
+		}
+
+		virtual bool CanRender(const UE::Dataflow::IDataflowConstructionViewMode& ViewMode) const override
+		{
+			return (ViewMode.GetName() == FDataflowConstruction3DViewMode::Name);
+		}
+
+		virtual void Render(GeometryCollection::Facades::FRenderingFacade& RenderCollection, const FGraphRenderingState& State) override
+		{
+			if (State.GetRenderOutputs().Num())
+			{
+				FManagedArrayCollection Default;
+				FName PrimaryOutput = State.GetRenderOutputs()[0]; // "Collection"
+				const FManagedArrayCollection& Collection = State.GetValue<FManagedArrayCollection>(PrimaryOutput, Default);
+
+				const bool bFoundIndices = Collection.FindAttributeTyped<FIntVector>("Indices", FGeometryCollection::FacesGroup) != nullptr;
+				const bool bFoundVertices = Collection.FindAttributeTyped<FVector3f>("Vertex", FGeometryCollection::VerticesGroup) != nullptr;
+				const bool bFoundTransforms = Collection.FindAttributeTyped<FTransform3f>(FTransformCollection::TransformAttribute, FTransformCollection::TransformGroup) != nullptr;
+				const bool bFoundBoneMap = Collection.FindAttributeTyped<int32>("BoneMap", FGeometryCollection::VerticesGroup) != nullptr;
+				const bool bFoundParents = Collection.FindAttributeTyped<int32>(FTransformCollection::ParentAttribute, FTransformCollection::TransformGroup) != nullptr;
+				bool bFoundRenderData = bFoundIndices && bFoundVertices && bFoundTransforms && bFoundBoneMap && bFoundParents
+					&& Collection.NumElements(FTransformCollection::TransformGroup) > 0;
+
+				const bool bFoundVertexStart = Collection.FindAttributeTyped<int32>("VertexStart", FGeometryCollection::GeometryGroup) != nullptr;
+				const bool bFoundVertexCount = Collection.FindAttributeTyped<int32>("VertexCount", FGeometryCollection::GeometryGroup) != nullptr;
+				const bool bFoundFaceStart = Collection.FindAttributeTyped<int32>("FaceStart", FGeometryCollection::GeometryGroup) != nullptr;
+				const bool bFoundFaceCount = Collection.FindAttributeTyped<int32>("FaceCount", FGeometryCollection::GeometryGroup) != nullptr;
+				bool bFoundGeometryAttributes = bFoundVertexStart && bFoundVertexCount && bFoundFaceStart && bFoundFaceCount
+					&& Collection.NumElements(FGeometryCollection::GeometryGroup) > 0;
+
+				if (bFoundRenderData && bFoundGeometryAttributes)
+				{
+					RenderTetrahedronGeometryCollection(RenderCollection, State);
+				}
+			}
+		}
+	};
 
 	class FFieldVolumeRenderCallbacks : public FRenderingFactory::ICallbackInterface
 	{
@@ -597,7 +757,10 @@ namespace UE::Dataflow
 		FRenderingFactory::GetInstance()->RegisterCallbacks(MakeUnique<FGeometryCollectionSurfaceWeightsRenderCallbacks>());
 		FRenderingFactory::GetInstance()->RegisterCallbacks(MakeUnique<FDynamicMesh3SurfaceRenderCallbacks>());
 		FRenderingFactory::GetInstance()->RegisterCallbacks(MakeUnique<FBoxSurfaceRenderCallbacks>());
+		FRenderingFactory::GetInstance()->RegisterCallbacks(MakeUnique<FSphereSurfaceRenderCallbacks>());
+		FRenderingFactory::GetInstance()->RegisterCallbacks(MakeUnique<FPointRenderCallbacks>());
+		FRenderingFactory::GetInstance()->RegisterCallbacks(MakeUnique<FPointsRenderCallbacks>());
+		FRenderingFactory::GetInstance()->RegisterCallbacks(MakeUnique<FGeometryCollectionTetrahedronRenderCallbacks>());
 		FRenderingFactory::GetInstance()->RegisterCallbacks(MakeUnique<FFieldVolumeRenderCallbacks>());
 	}
-
 }
diff --git a/Engine/Plugins/Experimental/GeometryCollectionPlugin/Source/GeometryCollectionNodes/Private/Dataflow/GeometryCollectionMakeNodes.cpp b/Engine/Plugins/Experimental/GeometryCollectionPlugin/Source/GeometryCollectionNodes/Private/Dataflow/GeometryCollectionMakeNodes.cpp
--- a/Engine/Plugins/Experimental/GeometryCollectionPlugin/Source/GeometryCollectionNodes/Private/Dataflow/GeometryCollectionMakeNodes.cpp
+++ b/Engine/Plugins/Experimental/GeometryCollectionPlugin/Source/GeometryCollectionNodes/Private/Dataflow/GeometryCollectionMakeNodes.cpp
@@ -57,6 +57,7 @@ namespace UE::Dataflow
 		DATAFLOW_NODE_REGISTER_CREATION_FACTORY(FMakeTransformDataflowNode);
 		DATAFLOW_NODE_REGISTER_CREATION_FACTORY(FMakeQuaternionDataflowNode);
 		DATAFLOW_NODE_REGISTER_CREATION_FACTORY(FMakeFloatArrayDataflowNode);
+		DATAFLOW_NODE_REGISTER_CREATION_FACTORY(FMakeCollectionDataflowNode);
 
 		// Generators
 		DATAFLOW_NODE_REGISTER_CREATION_FACTORY_NODE_COLORS_BY_CATEGORY("Generators", FLinearColor(.4f, 0.8f, 0.f), CDefaultNodeBodyTintColor);
@@ -181,3 +182,10 @@ void FMakeFloatArrayDataflowNode::Evaluate(UE::Dataflow::FContext& Context, cons
 	}
 }
 
+void FMakeCollectionDataflowNode::Evaluate(UE::Dataflow::FContext& Context, const FDataflowOutput* Out) const
+{
+	if (Out->IsA(&Collection))
+	{
+		SetValue(Context, FManagedArrayCollection(), &Collection);
+	}
+}
diff --git a/Engine/Plugins/Experimental/GeometryCollectionPlugin/Source/GeometryCollectionNodes/Private/Dataflow/GeometryCollectionNodes.cpp b/Engine/Plugins/Experimental/GeometryCollectionPlugin/Source/GeometryCollectionNodes/Private/Dataflow/GeometryCollectionNodes.cpp
--- a/Engine/Plugins/Experimental/GeometryCollectionPlugin/Source/GeometryCollectionNodes/Private/Dataflow/GeometryCollectionNodes.cpp
+++ b/Engine/Plugins/Experimental/GeometryCollectionPlugin/Source/GeometryCollectionNodes/Private/Dataflow/GeometryCollectionNodes.cpp
@@ -82,6 +82,11 @@ namespace UE::Dataflow
 		DATAFLOW_NODE_REGISTER_CREATION_FACTORY(FSelectionToVertexListDataflowNode);
 		DATAFLOW_NODE_REGISTER_CREATION_FACTORY(FBranchFloatDataflowNode);
 		DATAFLOW_NODE_REGISTER_CREATION_FACTORY(FBranchIntDataflowNode);
+		DATAFLOW_NODE_REGISTER_CREATION_FACTORY(FBoundingSphereDataflowNode);
+		DATAFLOW_NODE_REGISTER_CREATION_FACTORY(FExpandBoundingSphereDataflowNode);
+		DATAFLOW_NODE_REGISTER_CREATION_FACTORY(FVisualizeTetrahedronsDataflowNode);
+		DATAFLOW_NODE_REGISTER_CREATION_FACTORY(FPointsToCollectionDataflowNode);
+		DATAFLOW_NODE_REGISTER_CREATION_FACTORY(FCollectionToPointsDataflowNode);
 
 		// GeometryCollection
 		DATAFLOW_NODE_REGISTER_CREATION_FACTORY_NODE_COLORS_BY_CATEGORY("GeometryCollection", FLinearColor(0.55f, 0.45f, 1.0f), CDefaultNodeBodyTintColor);
@@ -1178,3 +1183,206 @@ void FBranchIntDataflowNode::Evaluate(UE::Dataflow::FContext& Context, const FDa
 	}
 }
 
+void FBoundingSphereDataflowNode::Evaluate(UE::Dataflow::FContext& Context, const FDataflowOutput* Out) const
+{
+	if (Out->IsA(&BoundingSphere))
+	{
+		if (IsConnected(&Collection))
+		{
+			const FManagedArrayCollection& InCollection = GetValue<FManagedArrayCollection>(Context, &Collection);
+
+			GeometryCollection::Facades::FBoundsFacade BoundsFacade(InCollection);
+			const FSphere& BoundingSphereInCollectionSpace = BoundsFacade.GetBoundingSphereInCollectionSpace();
+
+			SetValue(Context, BoundingSphereInCollectionSpace, &BoundingSphere);
+			return;
+		}
+
+		SetValue(Context, FSphere(), &BoundingSphere);
+	}
+}
+
+void FExpandBoundingSphereDataflowNode::Evaluate(UE::Dataflow::FContext& Context, const FDataflowOutput* Out) const
+{
+	FSphere InSphere = GetValue(Context, &BoundingSphere);
+
+	if (Out->IsA(&Center))
+	{
+		SetValue(Context, InSphere.Center, &Center);
+	}
+	else if (Out->IsA(&Radius))
+	{
+		SetValue(Context, (float)InSphere.W, &Radius);
+	}
+	else if (Out->IsA(&Volume))
+	{
+		SetValue(Context, (float)InSphere.GetVolume(), &Volume);
+	}
+}
+
+void FVisualizeTetrahedronsDataflowNode::Evaluate(UE::Dataflow::FContext& Context, const FDataflowOutput* Out) const
+{
+	if (Out->IsA(&Collection))
+	{
+		SafeForwardInput(Context, &Collection, &Collection);
+	}
+	else if (Out->IsA(&Vertices))
+	{
+		if (IsConnected(&Collection))
+		{
+			FManagedArrayCollection InCollection = GetValue(Context, &Collection);
+
+			if (InCollection.HasAttribute("Vertex", FGeometryCollection::VerticesGroup) &&
+				InCollection.HasAttribute("TransformToGeometryIndex", FGeometryCollection::TransformGroup) &&
+				InCollection.HasAttribute("VertexStart", FGeometryCollection::GeometryGroup) &&
+				InCollection.HasAttribute("VertexCount", FGeometryCollection::GeometryGroup))
+			{
+				GeometryCollection::Facades::FCollectionTransformFacade TransformFacade(InCollection);
+				if (TransformFacade.IsValid())
+				{
+					TArray<FTransform> CollectionSpaceTransforms = TransformFacade.ComputeCollectionSpaceTransforms();
+
+					const TManagedArray<FVector3f>& Vertex = InCollection.GetAttribute<FVector3f>("Vertex", FGeometryCollection::VerticesGroup);
+					const TManagedArray<int32>& TransformToGeometryIndex = InCollection.GetAttribute<int32>("TransformToGeometryIndex", FGeometryCollection::TransformGroup);
+					const TManagedArray<int32>& VertexStartArr = InCollection.GetAttribute<int32>("VertexStart", FGeometryCollection::GeometryGroup);
+					const TManagedArray<int32>& VertexCountArr = InCollection.GetAttribute<int32>("VertexCount", FGeometryCollection::GeometryGroup);
+
+					TArray<FVector> VerticesInCollectionSpace;
+					VerticesInCollectionSpace.AddUninitialized(Vertex.Num());
+
+					for (int32 TransformIndex = 0; TransformIndex < CollectionSpaceTransforms.Num(); ++TransformIndex)
+					{
+						const FTransform CollectionSpaceTransform = CollectionSpaceTransforms[TransformIndex];
+						const int32 GeoIndex = TransformToGeometryIndex[TransformIndex];
+						const int32 VertexStart = VertexStartArr[GeoIndex];
+						const int32 VertexCount = VertexCountArr[GeoIndex];
+
+						for (int32 VertexIdx = VertexStart; VertexIdx < VertexStart + VertexCount; ++VertexIdx)
+						{
+							VerticesInCollectionSpace[VertexIdx] = CollectionSpaceTransform.TransformPosition((FVector)Vertex[VertexIdx]);
+						}
+					}
+
+					SetValue(Context, MoveTemp(VerticesInCollectionSpace), &Vertices);
+					return;
+				}
+			}
+		}
+
+		SetValue(Context, TArray<FVector>(), &Vertices);
+	}
+}
+
+void FPointsToCollectionDataflowNode::Evaluate(UE::Dataflow::FContext& Context, const FDataflowOutput* Out) const
+{
+	if (Out->IsA(&Collection))
+	{
+		if (IsConnected(&Collection) && IsConnected(&Points))
+		{
+			FManagedArrayCollection InCollection = GetValue(Context, &Collection);
+			const TArray<FVector> InPoints = GetValue(Context, &Points);
+
+			const int32 NumTransforms = InCollection.NumElements(FGeometryCollection::TransformGroup);
+			const int32 NumGeoms = InCollection.NumElements(FGeometryCollection::GeometryGroup);
+			const int32 NumVertices = InCollection.NumElements(FGeometryCollection::VerticesGroup);
+
+			// Add new element to groups
+			InCollection.AddElements(1, FGeometryCollection::TransformGroup);
+			InCollection.AddElements(1, FGeometryCollection::GeometryGroup);
+			InCollection.AddElements(InPoints.Num(), FGeometryCollection::VerticesGroup);
+
+			TManagedArray<FTransform>& Transform = InCollection.AddAttribute<FTransform>("Transform", FGeometryCollection::TransformGroup);
+			TManagedArray<FString>& BoneName = InCollection.AddAttribute<FString>("BoneName", FGeometryCollection::TransformGroup);
+			TManagedArray<FLinearColor>& BoneColor = InCollection.AddAttribute<FLinearColor>("BoneColor", FGeometryCollection::TransformGroup);
+			TManagedArray<int32>& Parent = InCollection.AddAttribute<int32>("Parent", FGeometryCollection::TransformGroup);
+			TManagedArray<TSet<int32>>& Children = InCollection.AddAttribute<TSet<int32>>("Children", FGeometryCollection::TransformGroup);
+			TManagedArray<int32>& TransformToGeometryIndex = InCollection.AddAttribute<int32>("TransformToGeometryIndex", FGeometryCollection::TransformGroup);
+			TManagedArray<FVector3f>& Vertex = InCollection.AddAttribute<FVector3f>("Vertex", FGeometryCollection::VerticesGroup);
+			TManagedArray<int32>& BoneMap = InCollection.AddAttribute<int32>("BoneMap", FGeometryCollection::VerticesGroup);
+			TManagedArray<int32>& TransformIndex = InCollection.AddAttribute<int32>("TransformIndex", FGeometryCollection::GeometryGroup);
+			TManagedArray<FBox>& BoundingBox = InCollection.AddAttribute<FBox>("BoundingBox", FGeometryCollection::GeometryGroup);
+			TManagedArray<int32>& VertexStart = InCollection.AddAttribute<int32>("VertexStart", FGeometryCollection::GeometryGroup);
+			TManagedArray<int32>& VertexCount = InCollection.AddAttribute<int32>("VertexCount", FGeometryCollection::GeometryGroup);
+
+			Transform[NumTransforms] = FTransform::Identity;
+			BoneName[NumTransforms] = FString(TEXT("Points"));
+			BoneColor[NumTransforms] = FLinearColor(0.02f, 0.01f, 0.1f, 1.0f);
+			Parent[NumTransforms] = -1;
+			Children[NumTransforms] = TSet<int32>();
+			TransformToGeometryIndex[NumTransforms] = NumGeoms;
+
+			for (int32 VertexIdx = 0; VertexIdx < InPoints.Num(); ++VertexIdx)
+			{
+				Vertex[NumVertices + VertexIdx] = (FVector3f)InPoints[VertexIdx];
+			}
+
+			TransformIndex[NumGeoms] = NumTransforms;
+			VertexStart[NumGeoms] = NumVertices;
+			VertexCount[NumGeoms] = InPoints.Num();
+
+
+			GeometryCollection::Facades::FBoundsFacade BoundsFacade(InCollection);
+			const FBox& BoundingBoxOfPoints = BoundsFacade.ComputeBoundingBox(InPoints);
+
+			BoundingBox[NumGeoms] = BoundingBoxOfPoints;
+
+			SetValue(Context, MoveTemp(InCollection), &Collection);
+			return;
+		}
+
+		SetValue(Context, FManagedArrayCollection(), &Collection);
+	}
+}
+
+void FCollectionToPointsDataflowNode::Evaluate(UE::Dataflow::FContext& Context, const FDataflowOutput* Out) const
+{
+	if (Out->IsA(&Collection))
+	{
+		SafeForwardInput(Context, &Collection, &Collection);
+	}
+	else if (Out->IsA(&Points))
+	{
+		if (IsConnected(&Collection))
+		{
+			FManagedArrayCollection InCollection = GetValue(Context, &Collection);
+
+			if (InCollection.HasAttribute("Vertex", FGeometryCollection::VerticesGroup) &&
+				InCollection.HasAttribute("TransformToGeometryIndex", FGeometryCollection::TransformGroup) &&
+				InCollection.HasAttribute("VertexStart", FGeometryCollection::GeometryGroup) &&
+				InCollection.HasAttribute("VertexCount", FGeometryCollection::GeometryGroup))
+			{
+				GeometryCollection::Facades::FCollectionTransformFacade TransformFacade(InCollection);
+				if (TransformFacade.IsValid())
+				{
+					TArray<FTransform> CollectionSpaceTransforms = TransformFacade.ComputeCollectionSpaceTransforms();
+
+					const TManagedArray<FVector3f>& Vertex = InCollection.GetAttribute<FVector3f>("Vertex", FGeometryCollection::VerticesGroup);
+					const TManagedArray<int32>& TransformToGeometryIndex = InCollection.GetAttribute<int32>("TransformToGeometryIndex", FGeometryCollection::TransformGroup);
+					const TManagedArray<int32>& VertexStartArr = InCollection.GetAttribute<int32>("VertexStart", FGeometryCollection::GeometryGroup);
+					const TManagedArray<int32>& VertexCountArr = InCollection.GetAttribute<int32>("VertexCount", FGeometryCollection::GeometryGroup);
+
+					TArray<FVector> VerticesInCollectionSpace;
+					VerticesInCollectionSpace.AddUninitialized(Vertex.Num());
+
+					for (int32 TransformIndex = 0; TransformIndex < CollectionSpaceTransforms.Num(); ++TransformIndex)
+					{
+						const FTransform CollectionSpaceTransform = CollectionSpaceTransforms[TransformIndex];
+						const int32 GeoIndex = TransformToGeometryIndex[TransformIndex];
+						const int32 VertexStart = VertexStartArr[GeoIndex];
+						const int32 VertexCount = VertexCountArr[GeoIndex];
+
+						for (int32 VertexIdx = VertexStart; VertexIdx < VertexStart + VertexCount; ++VertexIdx)
+						{
+							VerticesInCollectionSpace[VertexIdx] = CollectionSpaceTransform.TransformPosition((FVector)Vertex[VertexIdx]);
+						}
+					}
+
+					SetValue(Context, MoveTemp(VerticesInCollectionSpace), &Points);
+					return;
+				}
+			}
+		}
+
+		SetValue(Context, TArray<FVector>(), &Points);
+	}
+}
\ No newline at end of file
diff --git a/Engine/Plugins/Experimental/GeometryCollectionPlugin/Source/GeometryCollectionNodes/Public/Dataflow/GeometryCollectionMakeNodes.h b/Engine/Plugins/Experimental/GeometryCollectionPlugin/Source/GeometryCollectionNodes/Public/Dataflow/GeometryCollectionMakeNodes.h
--- a/Engine/Plugins/Experimental/GeometryCollectionPlugin/Source/GeometryCollectionNodes/Public/Dataflow/GeometryCollectionMakeNodes.h
+++ b/Engine/Plugins/Experimental/GeometryCollectionPlugin/Source/GeometryCollectionNodes/Public/Dataflow/GeometryCollectionMakeNodes.h
@@ -57,6 +57,7 @@ struct FMakePointsDataflowNode : public FDataflowNode
 {
 	GENERATED_USTRUCT_BODY()
 	DATAFLOW_NODE_DEFINE_INTERNAL(FMakePointsDataflowNode, "MakePoints", "Generators|Point", "")
+	DATAFLOW_NODE_RENDER_TYPE("PointsRender", FName("TArray<FVector>"), "Points")
 
 public:
 	UPROPERTY(EditAnywhere, Category = "Point")
@@ -142,6 +143,7 @@ struct FMakeSphereDataflowNode : public FDataflowNode
 {
 	GENERATED_USTRUCT_BODY()
 	DATAFLOW_NODE_DEFINE_INTERNAL(FMakeSphereDataflowNode, "MakeSphere", "Generators|Sphere", "")
+	DATAFLOW_NODE_RENDER_TYPE("SurfaceRender", FName("FSphere"), "Sphere")
 
 public:
 	UPROPERTY(EditAnywhere, Category = "Sphere", meta = (DataflowInput));
@@ -370,7 +372,7 @@ public:
 
 /**
  *
- * Description for this node
+ * M
  *
  */
 USTRUCT()
@@ -404,6 +406,31 @@ public:
 
 };
 
+/**
+ *
+ * Make an empty ManagedArrayCollection
+ *
+ */
+USTRUCT()
+struct FMakeCollectionDataflowNode : public FDataflowNode
+{
+	GENERATED_USTRUCT_BODY()
+	DATAFLOW_NODE_DEFINE_INTERNAL(FMakeCollectionDataflowNode, "MakeCollection", "Generators|Collection", "")
+
+private:
+	UPROPERTY(meta = (DataflowOutput));
+	FManagedArrayCollection Collection;
+
+public:
+	FMakeCollectionDataflowNode(const UE::Dataflow::FNodeParameters& InParam, FGuid InGuid = FGuid::NewGuid())
+		: FDataflowNode(InParam, InGuid)
+	{
+		RegisterOutputConnection(&Collection);
+	}
+
+	virtual void Evaluate(UE::Dataflow::FContext& Context, const FDataflowOutput* Out) const override;
+};
+
 namespace UE::Dataflow
 {
 	void GeometryCollectionMakeNodes();
diff --git a/Engine/Plugins/Experimental/GeometryCollectionPlugin/Source/GeometryCollectionNodes/Public/Dataflow/GeometryCollectionNodes.h b/Engine/Plugins/Experimental/GeometryCollectionPlugin/Source/GeometryCollectionNodes/Public/Dataflow/GeometryCollectionNodes.h
--- a/Engine/Plugins/Experimental/GeometryCollectionPlugin/Source/GeometryCollectionNodes/Public/Dataflow/GeometryCollectionNodes.h
+++ b/Engine/Plugins/Experimental/GeometryCollectionPlugin/Source/GeometryCollectionNodes/Public/Dataflow/GeometryCollectionNodes.h
@@ -197,6 +197,36 @@ public:
 
 };
 
+/**
+ *
+ * Description for this node
+ *
+ */
+USTRUCT(meta = (DataflowGeometryCollection))
+struct FBoundingSphereDataflowNode : public FDataflowNode
+{
+	GENERATED_USTRUCT_BODY()
+	DATAFLOW_NODE_DEFINE_INTERNAL(FBoundingSphereDataflowNode, "BoundingSphere", "Utilities|Sphere", "")
+	DATAFLOW_NODE_RENDER_TYPE("SurfaceRender", FName("FSphere"), "BoundingSphere")
+
+public:
+	UPROPERTY(meta = (DataflowInput))
+	FManagedArrayCollection Collection;
+
+	UPROPERTY(meta = (DataflowOutput))
+	FSphere BoundingSphere = FSphere(ForceInit);
+
+	FBoundingSphereDataflowNode(const UE::Dataflow::FNodeParameters& InParam, FGuid InGuid = FGuid::NewGuid())
+		: FDataflowNode(InParam, InGuid)
+	{
+		RegisterInputConnection(&Collection);
+		RegisterOutputConnection(&BoundingSphere);
+	}
+
+	virtual void Evaluate(UE::Dataflow::FContext& Context, const FDataflowOutput* Out) const override;
+
+};
+
 UENUM()
 enum class EBoxLengthMeasurementMethod : uint8
 {
@@ -310,7 +340,42 @@ public:
 
 };
 
+/**
+ *
+ * Expands data of FSphere
+ *
+ */
+USTRUCT()
+struct FExpandBoundingSphereDataflowNode : public FDataflowNode
+{
+	GENERATED_USTRUCT_BODY()
+	DATAFLOW_NODE_DEFINE_INTERNAL(FExpandBoundingSphereDataflowNode, "ExpandBoundingSphere", "Utilities|Sphere", "")
+	DATAFLOW_NODE_RENDER_TYPE("PointRender", FName("FVector"), "Center")
+
+public:
+	UPROPERTY(meta = (DataflowInput))
+	FSphere BoundingSphere = FSphere(ForceInit);;
+
+	UPROPERTY(meta = (DataflowOutput))
+	FVector Center = FVector(0.0);
+
+	UPROPERTY(meta = (DataflowOutput))
+	float Radius = 0.f;
+
+	UPROPERTY(meta = (DataflowOutput))
+	float Volume = 0.f;
+
+	FExpandBoundingSphereDataflowNode(const UE::Dataflow::FNodeParameters& InParam, FGuid InGuid = FGuid::NewGuid())
+		: FDataflowNode(InParam, InGuid)
+	{
+		RegisterInputConnection(&BoundingSphere);
+		RegisterOutputConnection(&Center);
+		RegisterOutputConnection(&Radius);
+		RegisterOutputConnection(&Volume);
+	}
 
+	virtual void Evaluate(UE::Dataflow::FContext& Context, const FDataflowOutput* Out) const override;
+};
 
 
 /**
@@ -1665,6 +1730,109 @@ public:
  //USTRUCT(meta = (Experimental))
  //struct FLogStringDataflowNode_v2 : public FDataflowNode
 
+/**
+ *
+ * Visualize tetrahedrons in a collection
+ *
+ */
+USTRUCT(meta = (DataflowGeometryCollection))
+struct FVisualizeTetrahedronsDataflowNode : public FDataflowNode
+{
+	GENERATED_USTRUCT_BODY()
+	DATAFLOW_NODE_DEFINE_INTERNAL(FVisualizeTetrahedronsDataflowNode, "VisualizeTetrahedrons", "Flesh|Utilities", "")
+	DATAFLOW_NODE_RENDER_TYPE_START()
+	DATAFLOW_NODE_RENDER_TYPE_ADD("TetrahedronRender", FGeometryCollection::StaticType(), "Collection")
+	DATAFLOW_NODE_RENDER_TYPE_ADD("PointsRender", FName("TArray<FVector>"), "Vertices")
+	DATAFLOW_NODE_RENDER_TYPE_END()
+
+private:
+	/** Collection for the custom attribute */
+	UPROPERTY(meta = (DataflowInput, DataflowOutput, DataflowPassthrough = "Collection", DataflowIntrinsic))
+	FManagedArrayCollection Collection;
+
+	UPROPERTY(meta = (DataflowOutput))
+	TArray<FVector> Vertices;
+
+public:
+	FVisualizeTetrahedronsDataflowNode(const UE::Dataflow::FNodeParameters& InParam, FGuid InGuid = FGuid::NewGuid())
+		: FDataflowNode(InParam, InGuid)
+	{
+		RegisterInputConnection(&Collection);
+		RegisterOutputConnection(&Collection, &Collection);
+		RegisterOutputConnection(&Vertices);
+	}
+
+	virtual void Evaluate(UE::Dataflow::FContext& Context, const FDataflowOutput* Out) const override;
+
+};
+
+/**
+ *
+ * Add point cloud to a collection as vertices
+ *
+ */
+USTRUCT(meta = (DataflowGeometryCollection))
+struct FPointsToCollectionDataflowNode : public FDataflowNode
+{
+	GENERATED_USTRUCT_BODY()
+	DATAFLOW_NODE_DEFINE_INTERNAL(FPointsToCollectionDataflowNode, "PointsToCollection", "GeometryCollection|Utilities", "")
+
+private:
+	/** Collection to add the points to */
+	UPROPERTY(meta = (DataflowInput, DataflowOutput, DataflowPassthrough = "Collection", DataflowIntrinsic))
+	FManagedArrayCollection Collection;
+
+	/** Points to add to the collection */
+	UPROPERTY(EditAnyWhere, Category = "Points", meta = (DataflowInput))
+	TArray<FVector> Points;
+
+public:
+	FPointsToCollectionDataflowNode(const UE::Dataflow::FNodeParameters& InParam, FGuid InGuid = FGuid::NewGuid())
+		: FDataflowNode(InParam, InGuid)
+	{
+		RegisterInputConnection(&Collection);
+		RegisterInputConnection(&Points);
+		RegisterOutputConnection(&Collection, &Collection);
+	}
+
+	virtual void Evaluate(UE::Dataflow::FContext& Context, const FDataflowOutput* Out) const override;
+};
+
+/**
+ *
+ * Get vertices from a collection as a point cloud
+ *
+ */
+USTRUCT(meta = (DataflowGeometryCollection))
+struct FCollectionToPointsDataflowNode : public FDataflowNode
+{
+	GENERATED_USTRUCT_BODY()
+	DATAFLOW_NODE_DEFINE_INTERNAL(FCollectionToPointsDataflowNode, "CollectionToPoints", "GeometryCollection|Utilities", "")
+	DATAFLOW_NODE_RENDER_TYPE_START()
+	DATAFLOW_NODE_RENDER_TYPE_ADD("PointsRender", FName("TArray<FVector>"), "Points")
+	DATAFLOW_NODE_RENDER_TYPE_END()
+
+private:
+	/** Collection storing the points */
+	UPROPERTY(meta = (DataflowInput, DataflowOutput, DataflowPassthrough = "Collection", DataflowIntrinsic))
+	FManagedArrayCollection Collection;
+
+	/** Points from the collection */
+	UPROPERTY(meta = (DataflowOutput))
+	TArray<FVector> Points;
+
+public:
+	FCollectionToPointsDataflowNode(const UE::Dataflow::FNodeParameters& InParam, FGuid InGuid = FGuid::NewGuid())
+		: FDataflowNode(InParam, InGuid)
+	{
+		RegisterInputConnection(&Collection);
+		RegisterOutputConnection(&Collection, &Collection);
+		RegisterOutputConnection(&Points);
+	}
+
+	virtual void Evaluate(UE::Dataflow::FContext& Context, const FDataflowOutput* Out) const override;
+};
+
 namespace UE::Dataflow
 {
 	void GeometryCollectionEngineNodes();
diff --git a/Engine/Source/Runtime/Experimental/Chaos/Private/GeometryCollection/Facades/CollectionBoundsFacade.cpp b/Engine/Source/Runtime/Experimental/Chaos/Private/GeometryCollection/Facades/CollectionBoundsFacade.cpp
--- a/Engine/Source/Runtime/Experimental/Chaos/Private/GeometryCollection/Facades/CollectionBoundsFacade.cpp
+++ b/Engine/Source/Runtime/Experimental/Chaos/Private/GeometryCollection/Facades/CollectionBoundsFacade.cpp
@@ -202,6 +202,213 @@ namespace GeometryCollection::Facades
 		return Vertices;
 	}
 
+	FSphere FBoundsFacade::ComputeBoundingSphere(const FBox& InBoundingBox)
+	{
+		FSphere BoundingSphere;
+
+		BoundingSphere.Center = InBoundingBox.GetCenter();
+		BoundingSphere.W = 0.5 * (InBoundingBox.Max - InBoundingBox.Min).Length();
+
+		return BoundingSphere;
+	}
+
+	void FBoundsFacade::ComputeBoundingSphereImpl(const TArray<FVector>& InVertices, FSphere& OutSphere) const
+	{
+		FBox Box;
+		FVector MinIx[3] = { FVector(0), FVector(0), FVector(0) };
+		FVector MaxIx[3] = { FVector(0), FVector(0), FVector(0) };
+
+		bool bFirstVertex = true;
+
+		for (int32 Idx = 0; Idx < InVertices.Num(); ++Idx)
+		{
+			FVector Point = InVertices[Idx];
+			if (bFirstVertex)
+			{
+				// First, find AABB, remembering furthest points in each dir.
+				Box.Min = Point;
+				Box.Max = Box.Min;
+
+				MinIx[0] = Point;
+				MinIx[1] = Point;
+				MinIx[2] = Point;
+
+				MaxIx[0] = Point;
+				MaxIx[1] = Point;
+				MaxIx[2] = Point;
+				bFirstVertex = false;
+				continue;
+			}
+
+			// X //
+			if (Point.X < Box.Min.X)
+			{
+				Box.Min.X = Point.X;
+				MinIx[0] = Point;
+			}
+			else if (Point.X > Box.Max.X)
+			{
+				Box.Max.X = Point.X;
+				MaxIx[0] = Point;
+			}
+
+			// Y //
+			if (Point.Y < Box.Min.Y)
+			{
+				Box.Min.Y = Point.Y;
+				MinIx[1] = Point;
+			}
+			else if (Point.Y > Box.Max.Y)
+			{
+				Box.Max.Y = Point.Y;
+				MaxIx[1] = Point;
+			}
+
+			// Z //
+			if (Point.Z < Box.Min.Z)
+			{
+				Box.Min.Z = Point.Z;
+				MinIx[2] = Point;
+			}
+			else if (Point.Z > Box.Max.Z)
+			{
+				Box.Max.Z = Point.Z;
+				MaxIx[2] = Point;
+			}
+		}
+
+		const FVector Extremes[3] = { (MaxIx[0] - MinIx[0]), (MaxIx[1] - MinIx[1]), (MaxIx[2] - MinIx[2]) };
+
+		// Now find extreme points furthest apart, and initial center and radius of sphere.
+		double MaxDist2 = 0.0;
+		for (int32 i = 0; i < 3; ++i)
+		{
+			const double TmpDist2 = Extremes[i].SizeSquared();
+			if (TmpDist2 > MaxDist2)
+			{
+				MaxDist2 = TmpDist2;
+				OutSphere.Center = (MinIx[i] + (0.5f * Extremes[i]));
+				OutSphere.W = 0.f;
+			}
+		}
+
+		const FVector Extents = FVector(Extremes[0].X, Extremes[1].Y, Extremes[2].Z);
+
+		// radius and radius squared
+		double Radius = 0.5f * Extents.GetMax();
+		double Radius2 = FMath::Square(Radius);
+
+		// Now check each point lies within this sphere. If not - expand it a bit.
+		for (int32 Idx = 0; Idx < InVertices.Num(); ++Idx)
+		{
+			const FVector CenterToPoint = InVertices[Idx] - OutSphere.Center;
+			const double CenterToPoint2 = CenterToPoint.SizeSquared();
+
+			// If this point is outside our current bounding sphere's radius
+			if (CenterToPoint2 > Radius2)
+			{
+				// ..expand radius just enough to include this point.
+				const double PointRadius = FMath::Sqrt(CenterToPoint2);
+				Radius = 0.5f * (Radius + PointRadius);
+				Radius2 = FMath::Square(Radius);
+
+				OutSphere.Center += ((PointRadius - Radius) / PointRadius * CenterToPoint);
+			}
+		}
+
+		OutSphere.W = Radius;
+	}
+
+	void FBoundsFacade::ComputeBoundingSphereImpl2(const TArray<FVector>& InVertices, FSphere& OutSphere) const
+	{
+		const FBox BoundingBox = GetBoundingBoxInCollectionSpace();
+
+		OutSphere.Center = BoundingBox.GetCenter();
+		OutSphere.W = 0.0f;
+
+		for (int32 Idx = 0; Idx < InVertices.Num(); ++Idx)
+		{
+			const double Dist2 = FVector::DistSquared(InVertices[Idx], OutSphere.Center);
+			if (Dist2 > OutSphere.W)
+			{
+				OutSphere.W = Dist2;
+			}
+		}
+		OutSphere.W = FMath::Sqrt(OutSphere.W);
+	}
+
+	void FBoundsFacade::ComputeBoundingSphere(const TArray<FVector>& InVertices, FSphere& OutSphere) const
+	{
+		FSphere Sphere, Sphere2, BestSphere;
+
+		ComputeBoundingSphereImpl(InVertices, Sphere);
+		ComputeBoundingSphereImpl2(InVertices, Sphere2);
+
+		if (Sphere.W < Sphere2.W)
+			BestSphere = Sphere;
+		else
+			BestSphere = Sphere2;
+
+		// Don't use if radius is zero.
+		if (BestSphere.W <= 0.f)
+		{
+			OutSphere = FSphere(0);
+		}
+
+		OutSphere = BestSphere;
+	}
+
+	FSphere FBoundsFacade::GetBoundingSphereInCollectionSpace()
+	{
+		FSphere Sphere(0);
+
+		if (VertexAttribute.IsValid())
+		{
+			GeometryCollection::Facades::FCollectionTransformFacade TransformFacade(ConstCollection);
+			if (TransformFacade.IsValid())
+			{
+				TArray<FTransform> CollectionSpaceTransforms = TransformFacade.ComputeCollectionSpaceTransforms();
+
+				const TManagedArray<FVector3f>& Vertex = VertexAttribute.Get();
+				const TManagedArray<int32>& VertexStartArr = VertexStartAttribute.Get();
+				const TManagedArray<int32>& VertexCountArr = VertexCountAttribute.Get();
+
+				TArray<FVector> VerticesInCollectionSpace;
+				VerticesInCollectionSpace.AddUninitialized(Vertex.Num());
+
+				for (int32 TransformIndex = 0; TransformIndex < CollectionSpaceTransforms.Num(); ++TransformIndex)
+				{
+					const FTransform CollectionSpaceTransform = CollectionSpaceTransforms[TransformIndex];
+					const int32 GeoIndex = TransformToGeometryIndexAttribute[TransformIndex];
+					const int32 VertexStart = VertexStartArr[GeoIndex];
+					const int32 VertexCount = VertexCountArr[GeoIndex];
+
+					for (int32 VertexIdx = VertexStart; VertexIdx < VertexStart + VertexCount; ++VertexIdx)
+					{
+						VerticesInCollectionSpace[VertexIdx] = CollectionSpaceTransform.TransformPosition((FVector)Vertex[VertexIdx]);
+					}
+				}
+
+				ComputeBoundingSphere(VerticesInCollectionSpace, Sphere);
+			}
+		}
+
+		return Sphere;
+	}
+
+	FBox FBoundsFacade::ComputeBoundingBox(const TArray<FVector> InPoints)
+	{
+		FBox BoundingBox;
+		BoundingBox.Init();
+
+		for (int32 PointIdx = 0; PointIdx < InPoints.Num(); ++PointIdx)
+		{
+			BoundingBox += InPoints[PointIdx];
+		}
+
+		return BoundingBox;
+	}
+
 }; // GeometryCollection::Facades
 
 
diff --git a/Engine/Source/Runtime/Experimental/Chaos/Private/GeometryCollection/Facades/CollectionRenderingFacade.cpp b/Engine/Source/Runtime/Experimental/Chaos/Private/GeometryCollection/Facades/CollectionRenderingFacade.cpp
--- a/Engine/Source/Runtime/Experimental/Chaos/Private/GeometryCollection/Facades/CollectionRenderingFacade.cpp
+++ b/Engine/Source/Runtime/Experimental/Chaos/Private/GeometryCollection/Facades/CollectionRenderingFacade.cpp
@@ -121,6 +121,71 @@ namespace GeometryCollection::Facades
 		return 0;
 	}
 
+	void FRenderingFacade::AddPoint(const FVector3f& InPoint)
+	{
+		constexpr int32 NumVertices = 15;
+		constexpr int32 NumTriangles = 9;
+
+		TArray<FVector3f> Vertices; Vertices.AddUninitialized(NumVertices);
+		TArray<FIntVector> Tris; Tris.AddUninitialized(NumTriangles);
+
+		constexpr float Extension = 1.f;
+
+		// Add vertices
+		Vertices[0] = FVector3f(InPoint) + FVector3f(-Extension, 0.f, 0.f);
+		Vertices[1] = FVector3f(InPoint) + FVector3f(Extension, 0.f, 0.f);
+		FVector3f Dir = Vertices[1] - Vertices[0];
+		FVector3f DirAdd = Dir;
+		DirAdd.Y += Extension * 0.1f;
+		FVector3f OrthogonalDir = (Dir ^ DirAdd).GetSafeNormal();
+		OrthogonalDir *= 0.02f;
+		Vertices[2] = FVector3f(InPoint) + OrthogonalDir;
+
+		Vertices[3] = FVector3f(InPoint) + FVector3f(0.f, -Extension, 0.f);
+		Vertices[4] = FVector3f(InPoint) + FVector3f(0.f, Extension, 0.f);
+		Dir = Vertices[4] - Vertices[0];
+		DirAdd = Dir;
+		DirAdd.X += Extension * 0.1f;
+		OrthogonalDir = (Dir ^ DirAdd).GetSafeNormal();
+		OrthogonalDir *= 0.02f;
+		Vertices[5] = FVector3f(InPoint) + OrthogonalDir;
+
+		Vertices[6] = FVector3f(InPoint) + FVector3f(0.f, 0.f, -Extension);
+		Vertices[7] = FVector3f(InPoint) + FVector3f(0.f, 0.f, Extension);
+		Dir = Vertices[7] - Vertices[0];
+		DirAdd = Dir;
+		DirAdd.X += Extension * 0.1f;
+		OrthogonalDir = (Dir ^ DirAdd).GetSafeNormal();
+		OrthogonalDir *= 0.02f;
+		Vertices[8] = FVector3f(InPoint) + OrthogonalDir;
+
+		constexpr float SmallExtension = .5f * Extension;
+
+		Vertices[9] = FVector3f(InPoint) + FVector3f(-SmallExtension, 0.f, 0.f);
+		Vertices[10] = FVector3f(InPoint) + FVector3f(SmallExtension, 0.f, 0.f);
+		Vertices[11] = FVector3f(InPoint) + FVector3f(0.f, -SmallExtension, 0.f);
+		Vertices[12] = FVector3f(InPoint) + FVector3f(0.f, SmallExtension, 0.f);
+		Vertices[13] = FVector3f(InPoint) + FVector3f(0.f, 0.f, -SmallExtension);
+		Vertices[14] = FVector3f(InPoint) + FVector3f(0.f, 0.f, SmallExtension);
+
+		AddTriangle(Chaos::FTriangle(Vertices[0], Vertices[1], Vertices[2]));
+		AddTriangle(Chaos::FTriangle(Vertices[3], Vertices[4], Vertices[5]));
+		AddTriangle(Chaos::FTriangle(Vertices[6], Vertices[7], Vertices[8]));
+		AddTriangle(Chaos::FTriangle(Vertices[13], Vertices[9], Vertices[14]));
+		AddTriangle(Chaos::FTriangle(Vertices[14], Vertices[10], Vertices[13]));
+		AddTriangle(Chaos::FTriangle(Vertices[9], Vertices[11], Vertices[10]));
+		AddTriangle(Chaos::FTriangle(Vertices[10], Vertices[12], Vertices[9]));
+		AddTriangle(Chaos::FTriangle(Vertices[13], Vertices[11], Vertices[12]));
+		AddTriangle(Chaos::FTriangle(Vertices[12], Vertices[11], Vertices[14]));
+	}
+
+	void FRenderingFacade::AddPoints(TArray<FVector3f>&& InPoints)
+	{
+		for (int32 Idx = 0; Idx < InPoints.Num(); ++Idx)
+		{
+			AddPoint(InPoints[Idx]);
+		}
+	}
 
 	void FRenderingFacade::AddTriangle(const Chaos::FTriangle& InTriangle)
 	{
@@ -142,6 +207,156 @@ namespace GeometryCollection::Facades
 		}
 	}
 
+	void FRenderingFacade::AddBox(const FBox& InBox)
+	{
+		AddBox((FVector3f)InBox.Min, (FVector3f)InBox.Max);
+	}
+
+	void FRenderingFacade::AddBox(const FVector3f& InMinVertex, const FVector3f& InMaxVertex)
+	{
+		constexpr int32 NumVertices = 8;
+		constexpr int32 NumTriangles = 12;
+
+		TArray<FVector3f> Vertices; Vertices.AddUninitialized(NumVertices);
+		TArray<FIntVector> Tris; Tris.AddUninitialized(NumTriangles);
+
+		// Add vertices
+		Vertices[0] = InMinVertex;
+		Vertices[1] = FVector3f(InMaxVertex.X, InMinVertex.Y, InMinVertex.Z);
+		Vertices[2] = FVector3f(InMaxVertex.X, InMaxVertex.Y, InMinVertex.Z);
+		Vertices[3] = FVector3f(InMinVertex.X, InMaxVertex.Y, InMinVertex.Z);
+		Vertices[4] = FVector3f(InMinVertex.X, InMinVertex.Y, InMaxVertex.Z);
+		Vertices[5] = FVector3f(InMaxVertex.X, InMinVertex.Y, InMaxVertex.Z);
+		Vertices[6] = InMaxVertex;
+		Vertices[7] = FVector3f(InMinVertex.X, InMaxVertex.Y, InMaxVertex.Z);
+
+		AddTriangle(Chaos::FTriangle(Vertices[0], Vertices[1], Vertices[3]));
+		AddTriangle(Chaos::FTriangle(Vertices[1], Vertices[2], Vertices[3]));
+		AddTriangle(Chaos::FTriangle(Vertices[0], Vertices[4], Vertices[1]));
+		AddTriangle(Chaos::FTriangle(Vertices[4], Vertices[5], Vertices[1]));
+		AddTriangle(Chaos::FTriangle(Vertices[5], Vertices[2], Vertices[1]));
+		AddTriangle(Chaos::FTriangle(Vertices[5], Vertices[6], Vertices[2]));
+		AddTriangle(Chaos::FTriangle(Vertices[3], Vertices[2], Vertices[6]));
+		AddTriangle(Chaos::FTriangle(Vertices[7], Vertices[3], Vertices[6]));
+		AddTriangle(Chaos::FTriangle(Vertices[0], Vertices[3], Vertices[7]));
+		AddTriangle(Chaos::FTriangle(Vertices[4], Vertices[0], Vertices[7]));
+		AddTriangle(Chaos::FTriangle(Vertices[5], Vertices[4], Vertices[7]));
+		AddTriangle(Chaos::FTriangle(Vertices[5], Vertices[7], Vertices[6]));
+	}
+
+	void FRenderingFacade::AddSphere(const FSphere& InSphere, const FLinearColor& InColor)
+	{
+		AddSphere((FVector3f)InSphere.Center, (float)InSphere.W, InColor);
+	}
+
+	void FRenderingFacade::AddSphere(const FVector3f& InCenter, const float InRadius, const FLinearColor& InColor)
+	{
+		constexpr int32 NumVertices = 26;
+		constexpr int32 NumTriangles = 48;
+
+		TArray<FVector3f> Vertices; Vertices.AddUninitialized(NumVertices);
+		TArray<FIntVector> Tris; Tris.AddUninitialized(NumTriangles);
+
+		// Add vertices
+		Vertices[0] = FVector3f(0.f, 0.f, 1.f);
+		Vertices[1] = FVector3f(0.f, 0.f, -1.f);
+		Vertices[2] = FVector3f(0.707107f, 0.f, 0.707107f);
+		Vertices[3] = FVector3f(0.5f, -0.5f, 0.707107f);
+		Vertices[4] = FVector3f(0.f, -0.707107f, 0.707107f);
+		Vertices[5] = FVector3f(-0.5f, -0.5f, 0.707107f);
+		Vertices[6] = FVector3f(-0.707107f, 0.f, 0.707107f);
+		Vertices[7] = FVector3f(-0.5f, 0.5f, 0.707107f);
+		Vertices[8] = FVector3f(0.f, 0.707107f, 0.707107f);
+		Vertices[9] = FVector3f(0.5f, 0.5f, 0.707107f);
+		Vertices[10] = FVector3f(1.f, 0.f, 0.f);
+		Vertices[11] = FVector3f(0.707107f, -0.707107f, 0.f);
+		Vertices[12] = FVector3f(0.f, -1.f, 0.f);
+		Vertices[13] = FVector3f(-0.707107f, -0.707107f, 0.f);
+		Vertices[14] = FVector3f(-1.f, 0.f, 0.f);
+		Vertices[15] = FVector3f(-0.707107f, 0.707107f, 0.f);
+		Vertices[16] = FVector3f(0.f, 1.f, 0.f);
+		Vertices[17] = FVector3f(0.707107f, 0.707107f, 0.f);
+		Vertices[18] = FVector3f(0.707107f, 0.f, -0.707107f);
+		Vertices[19] = FVector3f(0.5f, -0.5f, -0.707107f);
+		Vertices[20] = FVector3f(0.f, -0.707107f, -0.707107f);
+		Vertices[21] = FVector3f(-0.5f, -0.5f, -0.707107f);
+		Vertices[22] = FVector3f(-0.707107f, 0.f, -0.707107f);
+		Vertices[23] = FVector3f(-0.5f, 0.5f, -0.707107f);
+		Vertices[24] = FVector3f(0.f, 0.707107f, -0.707107f);
+		Vertices[25] = FVector3f(0.5f, 0.5f, -0.707107f);
+
+		for (int32 Idx = 0; Idx < NumVertices; ++Idx)
+		{
+			Vertices[Idx] *= InRadius;
+			Vertices[Idx] += FVector3f(InCenter);
+		}
+
+		// Add triangles
+		Tris[0] = FIntVector(0, 2, 3); Tris[1] = FIntVector(0, 3, 4);
+		Tris[2] = FIntVector(0, 4, 5); Tris[3] = FIntVector(0, 5, 6);
+		Tris[4] = FIntVector(0, 6, 7); Tris[5] = FIntVector(0, 7, 8);
+		Tris[6] = FIntVector(0, 8, 9); Tris[7] = FIntVector(0, 9, 2);
+		Tris[8] = FIntVector(2, 10, 11); Tris[9] = FIntVector(2, 11, 3);
+		Tris[10] = FIntVector(3, 11, 12); Tris[11] = FIntVector(3, 12, 4);
+		Tris[12] = FIntVector(4, 12, 13); Tris[13] = FIntVector(4, 13, 5);
+		Tris[14] = FIntVector(5, 13, 14); Tris[15] = FIntVector(5, 14, 6);
+		Tris[16] = FIntVector(6, 14, 15); Tris[17] = FIntVector(6, 15, 7);
+		Tris[18] = FIntVector(7, 15, 16); Tris[19] = FIntVector(7, 16, 8);
+		Tris[20] = FIntVector(8, 16, 17); Tris[21] = FIntVector(8, 17, 9);
+		Tris[22] = FIntVector(9, 17, 10); Tris[23] = FIntVector(9, 10, 2);
+		Tris[24] = FIntVector(10, 18, 19); Tris[25] = FIntVector(10, 19, 11);
+		Tris[26] = FIntVector(11, 19, 20); Tris[27] = FIntVector(11, 20, 12);
+		Tris[28] = FIntVector(12, 20, 21); Tris[29] = FIntVector(12, 21, 13);
+		Tris[30] = FIntVector(13, 21, 22); Tris[31] = FIntVector(13, 22, 14);
+		Tris[32] = FIntVector(14, 22, 23); Tris[33] = FIntVector(14, 23, 15);
+		Tris[34] = FIntVector(15, 23, 24); Tris[35] = FIntVector(15, 24, 16);
+		Tris[36] = FIntVector(16, 24, 25); Tris[37] = FIntVector(16, 25, 17);
+		Tris[38] = FIntVector(17, 25, 18); Tris[39] = FIntVector(17, 18, 10);
+		Tris[40] = FIntVector(18, 1, 19); Tris[41] = FIntVector(19, 1, 20);
+		Tris[42] = FIntVector(20, 1, 21); Tris[43] = FIntVector(21, 1, 22);
+		Tris[44] = FIntVector(22, 1, 23); Tris[45] = FIntVector(23, 1, 24);
+		Tris[46] = FIntVector(24, 1, 25); Tris[47] = FIntVector(25, 1, 18);
+
+		TArray<FVector3f> VertexNormals; VertexNormals.AddUninitialized(NumVertices);
+		// TODO: Compute vertex normals
+
+		// Add VertexNormal and VertexColor
+		TArray<FLinearColor> VertexColors; VertexColors.AddUninitialized(Vertices.Num());
+		for (int32 VertexIdx = 0; VertexIdx < VertexNormals.Num(); ++VertexIdx)
+		{
+			VertexNormals[VertexIdx] = FVector3f(0.f);
+			VertexColors[VertexIdx] = FLinearColor(InColor);
+		}
+
+		AddSurface(MoveTemp(Vertices), MoveTemp(Tris), MoveTemp(VertexNormals), MoveTemp(VertexColors));
+	}
+
+	void FRenderingFacade::AddTetrahedron(const TArray<FVector3f>& InVertices, const FIntVector4& InIndices)
+	{
+		constexpr int32 NumVertices = 4;
+		constexpr int32 NumTriangles = 4;
+
+		TArray<FVector3f> Vertices; Vertices.AddUninitialized(NumVertices);
+		TArray<FIntVector> Tris; Tris.AddUninitialized(NumTriangles);
+
+		Vertices[0] = InVertices[InIndices[0]];
+		Vertices[1] = InVertices[InIndices[1]];
+		Vertices[2] = InVertices[InIndices[2]];
+		Vertices[3] = InVertices[InIndices[3]];
+
+		AddTriangle(Chaos::FTriangle(Vertices[0], Vertices[1], Vertices[2]));
+		AddTriangle(Chaos::FTriangle(Vertices[0], Vertices[3], Vertices[1]));
+		AddTriangle(Chaos::FTriangle(Vertices[2], Vertices[0], Vertices[3]));
+		AddTriangle(Chaos::FTriangle(Vertices[3], Vertices[1], Vertices[2]));
+	}
+
+	void FRenderingFacade::AddTetrahedrons(TArray<FVector3f>&& InVertices, TArray<FIntVector4>&& InIndices)
+	{
+		for (int32 Idx = 0; Idx < InIndices.Num(); ++Idx)
+		{
+			AddTetrahedron(InVertices, InIndices[Idx]);
+		}
+	}
 
 	void FRenderingFacade::AddSurface(TArray<FVector3f>&& InVertices, TArray<FIntVector>&& InIndices, TArray<FVector3f>&& InNormals, TArray<FLinearColor>&& InColors)
 	{
diff --git a/Engine/Source/Runtime/Experimental/Chaos/Public/GeometryCollection/Facades/CollectionBoundsFacade.h b/Engine/Source/Runtime/Experimental/Chaos/Public/GeometryCollection/Facades/CollectionBoundsFacade.h
--- a/Engine/Source/Runtime/Experimental/Chaos/Public/GeometryCollection/Facades/CollectionBoundsFacade.h
+++ b/Engine/Source/Runtime/Experimental/Chaos/Public/GeometryCollection/Facades/CollectionBoundsFacade.h
@@ -55,6 +55,13 @@ namespace GeometryCollection::Facades
 
 		/** TransformToGeometryIndex Access */
 		const TManagedArray< int32 >& GetTransformToGeometryIndex() const { return TransformToGeometryIndexAttribute.Get(); }
+
+		/** BoundingSphere for the whole collection in Collection space */
+		CHAOS_API FSphere GetBoundingSphereInCollectionSpace();
+
+		/** Compute boundingBox for a point cloud */
+		CHAOS_API FBox ComputeBoundingBox(const TArray<FVector> InPoints);
+
 	protected:
 
 		/** Transform based bounds evaluation, where the vertices are nested within a transform */
@@ -63,6 +70,16 @@ namespace GeometryCollection::Facades
 		/** Vertex based bounds evaluation, where the vertices are NOT nested within a transform */
 		void UpdateVertexBasedBoundingBox();
 
+	private:
+		/** Compute a bounding sphere of a box */
+		FSphere ComputeBoundingSphere(const FBox& InBoundingBox);
+
+		/** Compute a bounding sphere of a point cloud 
+		* These are different implementations and we use the better result */
+		void ComputeBoundingSphereImpl(const TArray<FVector>& InVertices, FSphere& OutSphere) const;
+		void ComputeBoundingSphereImpl2(const TArray<FVector>& InVertices, FSphere& OutSphere) const;
+		void ComputeBoundingSphere(const TArray<FVector>& InVertices, FSphere& OutSphere) const ;
+
 	private:
 		const FManagedArrayCollection& ConstCollection;
 		FManagedArrayCollection* Collection = nullptr;
diff --git a/Engine/Source/Runtime/Experimental/Chaos/Public/GeometryCollection/Facades/CollectionRenderingFacade.h b/Engine/Source/Runtime/Experimental/Chaos/Public/GeometryCollection/Facades/CollectionRenderingFacade.h
--- a/Engine/Source/Runtime/Experimental/Chaos/Public/GeometryCollection/Facades/CollectionRenderingFacade.h
+++ b/Engine/Source/Runtime/Experimental/Chaos/Public/GeometryCollection/Facades/CollectionRenderingFacade.h
@@ -49,15 +49,35 @@ namespace GeometryCollection::Facades
 		/**Number of triangles to render.*/
 		CHAOS_API int32 NumTriangles() const;
 
+		/**Add a point to the rendering view.*/
+		CHAOS_API void AddPoint(const FVector3f& InPoint);
+
+		/**Add point cloud to the rendering view.*/
+		CHAOS_API void AddPoints(TArray<FVector3f>&& InPoints);
+
 		/**Add a triangle to the rendering view.*/
 		CHAOS_API void AddTriangle(const Chaos::FTriangle& InTriangle);
 
+		/**Add a box to the rendering view.*/
+		CHAOS_API void AddBox(const FBox& InBox);
+		CHAOS_API void AddBox(const FVector3f& InMinVertex, const FVector3f& InMaxVertex);
+
+		/**Add a sphere to the rendering view.*/
+		CHAOS_API void AddSphere(const FSphere& InSphere, const FLinearColor& InColor);
+		CHAOS_API void AddSphere(const FVector3f& InCenter, const float InRadius, const FLinearColor& InColor);
+
 		/** Add a surface to the rendering view.*/
 		CHAOS_API void AddSurface(TArray<FVector3f>&& InVertices, TArray<FIntVector>&& InIndices, TArray<FVector3f>&& InNormals, TArray<FLinearColor>&& InColors);
 
 		/** Add a surface with UV layers, Material IDs and material paths to the rendering view.*/
 		CHAOS_API void AddSurface(TArray<FVector3f>&& InVertices, TArray<FIntVector>&& InIndices, TArray<FVector3f>&& InNormals, TArray<FLinearColor>&& InColors, TArray<TArray<FVector2f>>&& InUVs, TArray<int32>&& InMaterialIDs, TArray<FString>&& MaterialPaths);
 
+		/**Add a tetrahedron to the rendering view.*/
+		CHAOS_API void AddTetrahedron(const TArray<FVector3f>& InVertices, const FIntVector4& InIndices);
+
+		/**Add a tetrahedrons to the rendering view.*/
+		CHAOS_API void AddTetrahedrons(TArray<FVector3f>&& InVertices, TArray<FIntVector4>&& InIndices);
+
 		/** GetIndices */
 		const TManagedArray< FIntVector >& GetIndices() const { return IndicesAttribute.Get(); }
 
diff --git a/Engine/Source/Runtime/Experimental/Dataflow/Core/Public/Dataflow/DataflowNode.h b/Engine/Source/Runtime/Experimental/Dataflow/Core/Public/Dataflow/DataflowNode.h
--- a/Engine/Source/Runtime/Experimental/Dataflow/Core/Public/Dataflow/DataflowNode.h
+++ b/Engine/Source/Runtime/Experimental/Dataflow/Core/Public/Dataflow/DataflowNode.h
@@ -668,6 +668,16 @@ private:
 		Array.Add({ A, B, {__VA_ARGS__,} });														\
 		return Array;}
 
+#define DATAFLOW_NODE_RENDER_TYPE_START() \
+	virtual TArray<::UE::Dataflow::FRenderingParameter> GetRenderParametersImpl() const {		\
+		TArray<::UE::Dataflow::FRenderingParameter> Array;								
+
+#define DATAFLOW_NODE_RENDER_TYPE_ADD(A, B, ...)												\
+		Array.Add({ A, B, {__VA_ARGS__,} });	
+
+#define DATAFLOW_NODE_RENDER_TYPE_END()	\
+		return Array;}
+
 #define DATAFLOW_NODE_DEFINE_INTERNAL(TYPE, DISPLAY_NAME, CATEGORY, TAGS)			\
 public:																				\
 	static FName StaticType() {return #TYPE;}										\
