Dataflow Editor: auto-focus the simulation viewport when editor starts. Also fix for attaching two different functions to a static FViewport delegate.

diff --git a/Engine/Plugins/Experimental/Dataflow/Source/DataflowEditor/Private/Dataflow/DataflowEditorMode.cpp b/Engine/Plugins/Experimental/Dataflow/Source/DataflowEditor/Private/Dataflow/DataflowEditorMode.cpp
--- a/Engine/Plugins/Experimental/Dataflow/Source/DataflowEditor/Private/Dataflow/DataflowEditorMode.cpp
+++ b/Engine/Plugins/Experimental/Dataflow/Source/DataflowEditor/Private/Dataflow/DataflowEditorMode.cpp
@@ -84,6 +84,8 @@ void UDataflowEditorMode::Enter()
 {
 	UBaseCharacterFXEditorMode::Enter();
 
+	FViewport::ViewportResizedEvent.AddUObject(this, &UDataflowEditorMode::ViewportResized);
+
 	// Register gizmo ContextObject for use inside interactive tools
 	UE::TransformGizmoUtil::RegisterTransformGizmoContextObject(GetInteractiveToolsContext());
 
@@ -505,15 +507,15 @@ void UDataflowEditorMode::RefocusConstructionViewportClient()
 void UDataflowEditorMode::RefocusSimulationViewportClient()
 {
 	TSharedPtr<FDataflowSimulationViewportClient, ESPMode::ThreadSafe> PinnedVC = SimulationViewportClient.Pin();
-	if (PinnedVC.IsValid())
+	if (PinnedVC.IsValid() && SimulationScene)
 	{
 		// This will happen in FocusViewportOnBox anyways; do it now to get a consistent end result
 		PinnedVC->ToggleOrbitCamera(false);
 
-		const FBox SceneBounds = SceneBoundingBox();
+		const FBox SceneBounds = SimulationScene->GetBoundingBox();
 
-		// 3D space
-		PinnedVC->SetInitialViewTransform(ELevelViewportType::LVT_Perspective, FVector(0, 150, 200), FRotator(0, 0, 0), DEFAULT_ORTHOZOOM);
+		// Set up camera for an angled view by default 
+		PinnedVC->SetInitialViewTransform(ELevelViewportType::LVT_Perspective, FVector(0.0, 0.0, 0.0), FRotator(-15.0, -40.0, 0.0), DEFAULT_ORTHOZOOM);
 
 		constexpr bool bInstant = true;
 		PinnedVC->FocusViewportOnBox(SceneBounds, bInstant);
@@ -646,25 +648,38 @@ void UDataflowEditorMode::ModeTick(float DeltaTime)
 	}
 }
 
-void UDataflowEditorMode::ConstructionViewportResized(FViewport* ConstructionViewport, uint32 /*Unused*/)
+void UDataflowEditorMode::ViewportResized(FViewport* Viewport, uint32 /*Unused*/)
 {
-	// We'd like to call RefocusConstructionViewportClient() when the viewport is first created, however in Ortho mode the
-	// viewport needs to have non-zero size for FocusViewportOnBox() to work properly. So we wait until the viewport is resized here.
-	if (bShouldFocusConstructionView && ConstructionViewport && ConstructionViewport->GetSizeXY().X > 0 && ConstructionViewport->GetSizeXY().Y > 0)
+	// We'd like to call Refocus*ViewportClient() when the viewport is first created, however the viewport needs to have non-zero size for FocusViewportOnBox() to work properly. 
+	// So we wait until the viewport is resized and call it here.
+
+	// Construction
+
+	if (TSharedPtr<FDataflowConstructionViewportClient> PinnedConstructionViewportClient = ConstructionViewportClient.Pin())
 	{
-		RefocusConstructionViewportClient();
-		bShouldFocusConstructionView = false;
+		if (PinnedConstructionViewportClient->Viewport == Viewport)
+		{
+			if (bShouldFocusConstructionView && Viewport->GetSizeXY().X > 0 && Viewport->GetSizeXY().Y > 0)
+			{
+				RefocusConstructionViewportClient();
+				bShouldFocusConstructionView = false;
+				return;
+			}
+		}
 	}
-}
 
-void UDataflowEditorMode::SimulationViewportResized(FViewport* SimulationViewport, uint32 /*Unused*/)
-{
-	// We'd like to call RefocusConstructionViewportClient() when the viewport is first created, however in Ortho mode the
-	// viewport needs to have non-zero size for FocusViewportOnBox() to work properly. So we wait until the viewport is resized here.
-	if (bShouldFocusSimulationView && SimulationViewport && SimulationViewport->GetSizeXY().X > 0 && SimulationViewport->GetSizeXY().Y > 0)
+	// Simulation
+
+	if (TSharedPtr<FDataflowSimulationViewportClient> PinnedSimulationViewportClient = SimulationViewportClient.Pin())
 	{
-		RefocusSimulationViewportClient();
-		bShouldFocusSimulationView = false;
+		if (PinnedSimulationViewportClient->Viewport == Viewport)
+		{
+			if (bShouldFocusSimulationView && Viewport->GetSizeXY().X > 0 && Viewport->GetSizeXY().Y > 0)
+			{
+				RefocusSimulationViewportClient();
+				bShouldFocusSimulationView = false;
+			}
+		}
 	}
 }
 
@@ -902,26 +917,12 @@ void UDataflowEditorMode::SetConstructionViewportClient(TWeakPtr<FDataflowConstr
 	{
 		VC->SetConstructionViewMode(ConstructionViewMode);
 		VC->SetToolCommandList(ToolCommandList);
-
-		if (VC->Viewport)
-		{
-			VC->Viewport->ViewportResizedEvent.AddUObject(this, &UDataflowEditorMode::ConstructionViewportResized);
-		}
 	}
 }
 
 void UDataflowEditorMode::SetSimulationViewportClient(TWeakPtr<FDataflowSimulationViewportClient, ESPMode::ThreadSafe> InViewportClient)
 {
 	SimulationViewportClient = InViewportClient;
-
-	TSharedPtr<FDataflowSimulationViewportClient> VC = SimulationViewportClient.Pin();
-	if (VC.IsValid())
-	{
-		if (VC->Viewport)
-		{
-			VC->Viewport->ViewportResizedEvent.AddUObject(this, &UDataflowEditorMode::SimulationViewportResized);
-		}
-	}
 }
 
 void UDataflowEditorMode::InitializeContextObject()
diff --git a/Engine/Plugins/Experimental/Dataflow/Source/DataflowEditor/Private/Dataflow/DataflowEditorToolkit.cpp b/Engine/Plugins/Experimental/Dataflow/Source/DataflowEditor/Private/Dataflow/DataflowEditorToolkit.cpp
--- a/Engine/Plugins/Experimental/Dataflow/Source/DataflowEditor/Private/Dataflow/DataflowEditorToolkit.cpp
+++ b/Engine/Plugins/Experimental/Dataflow/Source/DataflowEditor/Private/Dataflow/DataflowEditorToolkit.cpp
@@ -528,6 +528,9 @@ void FDataflowEditorToolkit::InitializeEdMode(UBaseCharacterFXEditorMode* EdMode
 	// generated by the tools
 	DataflowMode->SetDataflowConstructionScene(static_cast<FDataflowConstructionScene*>(ObjectScene.Get()));
 
+	checkf(SimulationScene, TEXT("Expected SimulationScene to have been created in FDataflowEditorToolkit constructor"));
+	DataflowMode->SetDataflowSimlationScene(SimulationScene.Get());
+
 	// Set of the graph editor to be able to add nodes
 	DataflowMode->SetDataflowGraphEditor(GraphEditor);
 	TArray<TObjectPtr<UObject>> ObjectsToEdit;
diff --git a/Engine/Plugins/Experimental/Dataflow/Source/DataflowEditor/Private/Dataflow/DataflowSimulationScene.cpp b/Engine/Plugins/Experimental/Dataflow/Source/DataflowEditor/Private/Dataflow/DataflowSimulationScene.cpp
--- a/Engine/Plugins/Experimental/Dataflow/Source/DataflowEditor/Private/Dataflow/DataflowSimulationScene.cpp
+++ b/Engine/Plugins/Experimental/Dataflow/Source/DataflowEditor/Private/Dataflow/DataflowSimulationScene.cpp
@@ -288,6 +288,35 @@ int32 FDataflowSimulationScene::GetPreviewLOD() const
 	return CurrentPreviewLOD;
 }
 
+FBox FDataflowSimulationScene::GetBoundingBox() const
+{
+	FBox SceneBounds(ForceInitToZero);
+	if (DataflowModeManager.IsValid())
+	{
+		USelection* const SelectedComponents = DataflowModeManager->GetSelectedComponents();
+
+		TArray<TWeakObjectPtr<UObject>> SelectedObjects;
+		const int32 NumSelected = SelectedComponents->GetSelectedObjects(SelectedObjects);
+
+		if (NumSelected > 0)
+		{
+			for (const TWeakObjectPtr<const UObject> SelectedObject : SelectedObjects)
+			{
+				if (const UPrimitiveComponent* const SelectedComponent = Cast<const UPrimitiveComponent>(SelectedObject))
+				{
+					SceneBounds += SelectedComponent->Bounds.GetBox();
+				}
+			}
+		}
+		else
+		{
+			constexpr bool bIncludeNonColliding = true;
+			SceneBounds += PreviewActor->GetComponentsBoundingBox(bIncludeNonColliding);
+		}
+	}
+	return SceneBounds;
+}
+
 void FDataflowSimulationScene::TickDataflowScene(const float DeltaSeconds)
 {
 	float DeltaTime = DeltaSeconds;
diff --git a/Engine/Plugins/Experimental/Dataflow/Source/DataflowEditor/Public/Dataflow/DataflowEditorMode.h b/Engine/Plugins/Experimental/Dataflow/Source/DataflowEditor/Public/Dataflow/DataflowEditorMode.h
--- a/Engine/Plugins/Experimental/Dataflow/Source/DataflowEditor/Public/Dataflow/DataflowEditorMode.h
+++ b/Engine/Plugins/Experimental/Dataflow/Source/DataflowEditor/Public/Dataflow/DataflowEditorMode.h
@@ -214,9 +214,7 @@ private:
 	bool bShouldFocusConstructionView = true;
 	bool bShouldFocusSimulationView = true;
 
-	//@todo(Dataflow) : Move to Construction and Simulaiton
-	void ConstructionViewportResized(FViewport* ConstructionViewport, uint32 Unused);
-	void SimulationViewportResized(FViewport* ConstructionViewport, uint32 Unused);
+	void ViewportResized(FViewport* Viewport, uint32 Unused);
 
 	// Dataflow node type whose corresponding tool should be started on the next Tick
 	FName NodeTypeForPendingToolStart;
diff --git a/Engine/Plugins/Experimental/Dataflow/Source/DataflowEditor/Public/Dataflow/DataflowEditorPreviewSceneBase.h b/Engine/Plugins/Experimental/Dataflow/Source/DataflowEditor/Public/Dataflow/DataflowEditorPreviewSceneBase.h
--- a/Engine/Plugins/Experimental/Dataflow/Source/DataflowEditor/Public/Dataflow/DataflowEditorPreviewSceneBase.h
+++ b/Engine/Plugins/Experimental/Dataflow/Source/DataflowEditor/Public/Dataflow/DataflowEditorPreviewSceneBase.h
@@ -40,7 +40,7 @@ public:
 	const TSharedPtr<FAssetEditorModeManager>& GetDataflowModeManager() const { return DataflowModeManager; }
 	
 	/** Build the scene bounding box */
-	FBox GetBoundingBox() const;
+	virtual FBox GetBoundingBox() const;
 
 	/** Tick data flow scene */
 	virtual void TickDataflowScene(const float DeltaSeconds) {}
diff --git a/Engine/Plugins/Experimental/Dataflow/Source/DataflowEditor/Public/Dataflow/DataflowSimulationScene.h b/Engine/Plugins/Experimental/Dataflow/Source/DataflowEditor/Public/Dataflow/DataflowSimulationScene.h
--- a/Engine/Plugins/Experimental/Dataflow/Source/DataflowEditor/Public/Dataflow/DataflowSimulationScene.h
+++ b/Engine/Plugins/Experimental/Dataflow/Source/DataflowEditor/Public/Dataflow/DataflowSimulationScene.h
@@ -104,6 +104,9 @@ public:
 	/** Check if the preview scene can run simulation */
 	virtual bool CanRunSimulation() const { return true; }
 
+	/** Get the bounding box for the selected components or the entire sim scene */
+	virtual FBox GetBoundingBox() const override;
+
 	/** Get the scene description used in the preview scene widget */
 	UDataflowSimulationSceneDescription* GetPreviewSceneDescription() const { return SceneDescription; }
 
