Dataflow Editor: ignore non-alt left-mouse-drag in both viewports, similar to the Chaos Cloth Editor. We do this by adding a common base class that defines a ClickDragInputBehavior that is there to simply prevent the default drag operation from starting

diff --git a/Engine/Plugins/Experimental/Dataflow/Source/DataflowEditor/Private/Dataflow/DataflowConstructionViewportClient.cpp b/Engine/Plugins/Experimental/Dataflow/Source/DataflowEditor/Private/Dataflow/DataflowConstructionViewportClient.cpp
--- a/Engine/Plugins/Experimental/Dataflow/Source/DataflowEditor/Private/Dataflow/DataflowConstructionViewportClient.cpp
+++ b/Engine/Plugins/Experimental/Dataflow/Source/DataflowEditor/Private/Dataflow/DataflowConstructionViewportClient.cpp
@@ -19,13 +19,14 @@
 #include "Selection.h"
 #include "SGraphPanel.h"
 #include "SNodePanel.h"
+#include "BaseBehaviors/ClickDragBehavior.h"
 
 #define LOCTEXT_NAMESPACE "DataflowConstructionViewportClient"
 
 FDataflowConstructionViewportClient::FDataflowConstructionViewportClient(FEditorModeTools* InModeTools,
                                                              FPreviewScene* InPreviewScene,  const bool bCouldTickScene,
                                                              const TWeakPtr<SEditorViewport> InEditorViewportWidget)
-	: FEditorViewportClient(InModeTools, InPreviewScene, InEditorViewportWidget)
+	: FDataflowEditorViewportClientBase(InModeTools, InPreviewScene, bCouldTickScene, InEditorViewportWidget)
 {
 	// We want our near clip plane to be quite close so that we can zoom in further.
 	OverrideNearClipPlane(KINDA_SMALL_NUMBER);
@@ -47,11 +48,6 @@ void FDataflowConstructionViewportClient::SetToolCommandList(TWeakPtr<FUICommand
 	ToolCommandList = InToolCommandList;
 }
 
-//const UInputBehaviorSet* FDataflowConstructionViewportClient::GetInputBehaviors() const
-//{
-//	return BehaviorSet;
-//}
-
 void FDataflowConstructionViewportClient::Tick(float DeltaSeconds)
 {
 	Super::Tick(DeltaSeconds);
@@ -86,7 +82,11 @@ bool FDataflowConstructionViewportClient::InputKey(const FInputKeyEventArgs& Eve
 void FDataflowConstructionViewportClient::ProcessClick(FSceneView& View, HHitProxy* HitProxy, FKey Key, EInputEvent Event, uint32 HitX, uint32 HitY)
 {
 	Super::ProcessClick(View, HitProxy, Key, Event, HitX, HitY);
+	OnViewportClicked(HitProxy);
+}
 
+void FDataflowConstructionViewportClient::OnViewportClicked(HHitProxy* HitProxy)
+{
 	auto EnableToolForSelectedNode = [&](USelection* SelectedComponents)
 	{
 		if (TSharedPtr<FDataflowEditorToolkit> DataflowEditorToolkit = DataflowEditorToolkitPtr.Pin())
diff --git a/Engine/Plugins/Experimental/Dataflow/Source/DataflowEditor/Private/Dataflow/DataflowEditorViewportClientBase.cpp b/Engine/Plugins/Experimental/Dataflow/Source/DataflowEditor/Private/Dataflow/DataflowEditorViewportClientBase.cpp
new file mode 100644
--- /dev/null
+++ b/Engine/Plugins/Experimental/Dataflow/Source/DataflowEditor/Private/Dataflow/DataflowEditorViewportClientBase.cpp
@@ -0,0 +1,51 @@
+// Copyright Epic Games, Inc. All Rights Reserved.
+
+#include "Dataflow/DataflowEditorViewportClientBase.h"
+#include "EditorModeManager.h"
+#include "EdModeInteractiveToolsContext.h"
+#include "BaseBehaviors/ClickDragBehavior.h"
+
+#define LOCTEXT_NAMESPACE "DataflowEditorViewportClientBase"
+
+FDataflowEditorViewportClientBase::FDataflowEditorViewportClientBase(FEditorModeTools* InModeTools,
+                                                             FPreviewScene* InPreviewScene,  const bool bCouldTickScene,
+                                                             const TWeakPtr<SEditorViewport> InEditorViewportWidget)
+	: FEditorViewportClient(InModeTools, InPreviewScene, InEditorViewportWidget)
+{
+	BehaviorSet = NewObject<UInputBehaviorSet>();
+
+	// The ClickDragBehavior is used to intercept non-alt left-mouse-button drag inputs, but still allow single-click for select/deselect operation
+	const TObjectPtr<ULocalClickDragInputBehavior> ClickDragBehavior = NewObject<ULocalClickDragInputBehavior>();
+	ClickDragBehavior->Initialize();
+
+	// We'll have the priority of our viewport behaviors be lower (i.e. higher numerically) than both the gizmo default and the tool default
+	constexpr int ViewportBehaviorPriority = FMath::Max(FInputCapturePriority::DEFAULT_GIZMO_PRIORITY, FInputCapturePriority::DEFAULT_TOOL_PRIORITY) + 1;
+	ClickDragBehavior->SetDefaultPriority(ViewportBehaviorPriority);
+
+	ClickDragBehavior->ModifierCheckFunc = [](const FInputDeviceState& InputState) 
+	{
+		return !FInputDeviceState::IsAltKeyDown(InputState);
+	};
+
+	ClickDragBehavior->CanBeginClickDragFunc = [](const FInputDeviceRay& InputDeviceRay)
+	{
+		return FInputRayHit(TNumericLimits<float>::Max()); // bHit is true. Depth is max to lose the standard tiebreaker.
+	};
+
+	ClickDragBehavior->OnClickPressFunc = [this](const FInputDeviceRay& ClickPos)
+	{
+		HHitProxy* const HitProxy = Viewport->GetHitProxy(ClickPos.ScreenPosition[0], ClickPos.ScreenPosition[1]);
+		OnViewportClicked(HitProxy);
+	};
+
+	BehaviorSet->Add(ClickDragBehavior);
+
+	ModeTools->GetInteractiveToolsContext()->InputRouter->RegisterSource(this);
+}
+
+const UInputBehaviorSet* FDataflowEditorViewportClientBase::GetInputBehaviors() const
+{
+	return BehaviorSet;
+}
+
+#undef LOCTEXT_NAMESPACE 
diff --git a/Engine/Plugins/Experimental/Dataflow/Source/DataflowEditor/Private/Dataflow/DataflowSimulationViewportClient.cpp b/Engine/Plugins/Experimental/Dataflow/Source/DataflowEditor/Private/Dataflow/DataflowSimulationViewportClient.cpp
--- a/Engine/Plugins/Experimental/Dataflow/Source/DataflowEditor/Private/Dataflow/DataflowSimulationViewportClient.cpp
+++ b/Engine/Plugins/Experimental/Dataflow/Source/DataflowEditor/Private/Dataflow/DataflowSimulationViewportClient.cpp
@@ -22,7 +22,7 @@
 FDataflowSimulationViewportClient::FDataflowSimulationViewportClient(FEditorModeTools* InModeTools,
                                                              FPreviewScene* InPreviewScene,  const bool bCouldTickScene,
                                                              const TWeakPtr<SEditorViewport> InEditorViewportWidget)
-	: FEditorViewportClient(InModeTools, InPreviewScene, InEditorViewportWidget)
+	: FDataflowEditorViewportClientBase(InModeTools, InPreviewScene, bCouldTickScene, InEditorViewportWidget)
 {
 	// We want our near clip plane to be quite close so that we can zoom in further.
 	OverrideNearClipPlane(KINDA_SMALL_NUMBER);
@@ -44,11 +44,6 @@ void FDataflowSimulationViewportClient::SetToolCommandList(TWeakPtr<FUICommandLi
 	ToolCommandList = InToolCommandList;
 }
 
-//const UInputBehaviorSet* FDataflowSimulationViewportClient::GetInputBehaviors() const
-//{
-//	return BehaviorSet;
-//}
-
 void FDataflowSimulationViewportClient::Tick(float DeltaSeconds)
 {
 	Super::Tick(DeltaSeconds);
@@ -61,7 +56,11 @@ void FDataflowSimulationViewportClient::Tick(float DeltaSeconds)
 void FDataflowSimulationViewportClient::ProcessClick(FSceneView& View, HHitProxy* HitProxy, FKey Key, EInputEvent Event, uint32 HitX, uint32 HitY)
 {
 	Super::ProcessClick(View, HitProxy, Key, Event, HitX, HitY);
+	OnViewportClicked(HitProxy);
+}
 
+void FDataflowSimulationViewportClient::OnViewportClicked(HHitProxy* HitProxy)
+{
 	USelection* SelectedComponents = ModeTools->GetSelectedComponents();
 	
 	TArray<UPrimitiveComponent*> PreviouslySelectedComponents;
diff --git a/Engine/Plugins/Experimental/Dataflow/Source/DataflowEditor/Public/Dataflow/DataflowConstructionViewportClient.h b/Engine/Plugins/Experimental/Dataflow/Source/DataflowEditor/Public/Dataflow/DataflowConstructionViewportClient.h
--- a/Engine/Plugins/Experimental/Dataflow/Source/DataflowEditor/Public/Dataflow/DataflowConstructionViewportClient.h
+++ b/Engine/Plugins/Experimental/Dataflow/Source/DataflowEditor/Public/Dataflow/DataflowConstructionViewportClient.h
@@ -2,7 +2,7 @@
 #pragma once
 
 #include "CoreMinimal.h"
-#include "EditorViewportClient.h"
+#include "Dataflow/DataflowEditorViewportClientBase.h"
 #include "Dataflow/DataflowNodeParameters.h"
 #include "Dataflow/DataflowComponentSelectionState.h"
 #include "Dataflow/DataflowContent.h"
@@ -24,10 +24,10 @@ namespace UE::Dataflow
 	class IDataflowConstructionViewMode;
 }
 
-class DATAFLOWEDITOR_API FDataflowConstructionViewportClient : public FEditorViewportClient
+class DATAFLOWEDITOR_API FDataflowConstructionViewportClient : public FDataflowEditorViewportClientBase
 {
 public:
-	using Super = FEditorViewportClient;
+	using Super = FDataflowEditorViewportClientBase;
 
 	FDataflowConstructionViewportClient(FEditorModeTools* InModeTools, FPreviewScene* InPreviewScene,  const bool bCouldTickScene,
 								  const TWeakPtr<SEditorViewport> InEditorViewportWidget = nullptr);
@@ -56,6 +56,9 @@ public:
 
 private:
 
+	// FDataflowEditorViewportClientBase
+	virtual void OnViewportClicked(HHitProxy* HitProxy) override;
+
 	// FEditorViewportClient interface
 	virtual void Tick(float DeltaSeconds) override;
 	virtual bool InputKey(const FInputKeyEventArgs& EventArgs) override;
diff --git a/Engine/Plugins/Experimental/Dataflow/Source/DataflowEditor/Public/Dataflow/DataflowEditorViewportClientBase.h b/Engine/Plugins/Experimental/Dataflow/Source/DataflowEditor/Public/Dataflow/DataflowEditorViewportClientBase.h
new file mode 100644
--- /dev/null
+++ b/Engine/Plugins/Experimental/Dataflow/Source/DataflowEditor/Public/Dataflow/DataflowEditorViewportClientBase.h
@@ -0,0 +1,28 @@
+// Copyright Epic Games, Inc. All Rights Reserved.
+#pragma once
+
+#include "EditorViewportClient.h"
+#include "InputBehaviorSet.h"
+
+class FEditorModeTools;
+class FPreviewScene;
+class SEditorViewport;
+class HHitProxy;
+
+class DATAFLOWEDITOR_API FDataflowEditorViewportClientBase : public FEditorViewportClient, public IInputBehaviorSource
+{
+public:
+
+	FDataflowEditorViewportClientBase(FEditorModeTools* InModeTools, FPreviewScene* InPreviewScene,  const bool bCouldTickScene, const TWeakPtr<SEditorViewport> InEditorViewportWidget = nullptr);
+
+protected:
+
+	virtual void OnViewportClicked(HHitProxy* HitProxy) = 0;
+
+private:
+
+	// IInputBehaviorSource
+	virtual const UInputBehaviorSet* GetInputBehaviors() const override;
+
+	TObjectPtr<UInputBehaviorSet> BehaviorSet;
+};
diff --git a/Engine/Plugins/Experimental/Dataflow/Source/DataflowEditor/Public/Dataflow/DataflowSimulationViewportClient.h b/Engine/Plugins/Experimental/Dataflow/Source/DataflowEditor/Public/Dataflow/DataflowSimulationViewportClient.h
--- a/Engine/Plugins/Experimental/Dataflow/Source/DataflowEditor/Public/Dataflow/DataflowSimulationViewportClient.h
+++ b/Engine/Plugins/Experimental/Dataflow/Source/DataflowEditor/Public/Dataflow/DataflowSimulationViewportClient.h
@@ -2,7 +2,7 @@
 #pragma once
 
 #include "CoreMinimal.h"
-#include "EditorViewportClient.h"
+#include "Dataflow/DataflowEditorViewportClientBase.h"
 #include "Dataflow/DataflowNodeParameters.h"
 #include "Dataflow/DataflowComponentSelectionState.h"
 #include "Dataflow/DataflowContent.h"
@@ -16,12 +16,11 @@ class FTransformGizmoDataBinder;
 class FDataflowPreviewSceneBase;
 class UInputBehaviorSet;
 
-class DATAFLOWEDITOR_API FDataflowSimulationViewportClient : public FEditorViewportClient //, public IInputBehaviorSource
+class DATAFLOWEDITOR_API FDataflowSimulationViewportClient : public FDataflowEditorViewportClientBase
 {
-//@todo(brice) : Add BehaviorUI support
 
 public:
-	using Super = FEditorViewportClient;
+	using Super = FDataflowEditorViewportClientBase;
 
 	FDataflowSimulationViewportClient(FEditorModeTools* InModeTools, FPreviewScene* InPreviewScene,  const bool bCouldTickScene,
 								  const TWeakPtr<SEditorViewport> InEditorViewportWidget = nullptr);
@@ -41,6 +40,9 @@ public:
 
 private:
 
+	// FDataflowEditorViewportClientBase
+	virtual void OnViewportClicked(HHitProxy* HitProxy) override;
+
 	// FEditorViewportClient interface
 	virtual void Tick(float DeltaSeconds) override;
 	virtual void ProcessClick(FSceneView& View, HHitProxy* HitProxy, FKey Key, EInputEvent Event, uint32 HitX, uint32 HitY) override;
