ChaosFlesh: added skeletal mesh check when loading an animation in dataflow simulation. Now produces an error if skeletal mesh  in the preview scene has different skeleton than the asset.

diff --git a/Engine/Plugins/Experimental/ChaosFlesh/Source/ChaosFleshEngine/Private/ChaosFlesh/ChaosDeformableTetrahedralComponent.cpp b/Engine/Plugins/Experimental/ChaosFlesh/Source/ChaosFleshEngine/Private/ChaosFlesh/ChaosDeformableTetrahedralComponent.cpp
--- a/Engine/Plugins/Experimental/ChaosFlesh/Source/ChaosFleshEngine/Private/ChaosFlesh/ChaosDeformableTetrahedralComponent.cpp
+++ b/Engine/Plugins/Experimental/ChaosFlesh/Source/ChaosFleshEngine/Private/ChaosFlesh/ChaosDeformableTetrahedralComponent.cpp
@@ -721,6 +721,35 @@ TOptional<TArray<int32>> UDeformableTetrahedralComponent::GetMeshImportVertexMap
 #endif
 }
 
+bool UDeformableTetrahedralComponent::IsSkeletalMeshAnimationCompatible(const USkeletalMeshComponent* SkeletalMeshComponent) const
+{
+	if (SkeletalMeshComponent)
+	{
+		if (const UFleshAsset* FleshAsset = GetRestCollection())
+		{
+			if (const FFleshCollection* Rest = FleshAsset->GetCollection())
+			{
+				GeometryCollection::Facades::FTransformSource TransformSource(*Rest);
+				if (TransformSource.IsValid())
+				{
+					if (const USkeletalMesh* SkeletalMesh = SkeletalMeshComponent->GetSkeletalMeshAsset())
+					{
+						if (const USkeleton* Skeleton = SkeletalMesh->GetSkeleton())
+						{
+							TSet<int32> Roots = TransformSource.GetTransformSource(Skeleton->GetName(), Skeleton->GetGuid().ToString());
+							if (!Roots.IsEmpty() && ensureMsgf(Roots.Num() == 1, TEXT("Error: Only supports a single root per skeleton.(%s)"), *Skeleton->GetName()))
+							{
+								return true;
+							}
+						}
+					}
+				}
+			}
+		}
+	}
+	return false;
+}
+
 TArray<FVector> UDeformableTetrahedralComponent::GetSkeletalMeshEmbeddedPositions(
 	const ChaosDeformableBindingOption Format,
 	const FTransform TargetDeformationSkeletonOffset,
diff --git a/Engine/Plugins/Experimental/ChaosFlesh/Source/ChaosFleshEngine/Public/ChaosFlesh/ChaosDeformableTetrahedralComponent.h b/Engine/Plugins/Experimental/ChaosFlesh/Source/ChaosFleshEngine/Public/ChaosFlesh/ChaosDeformableTetrahedralComponent.h
--- a/Engine/Plugins/Experimental/ChaosFlesh/Source/ChaosFleshEngine/Public/ChaosFlesh/ChaosDeformableTetrahedralComponent.h
+++ b/Engine/Plugins/Experimental/ChaosFlesh/Source/ChaosFleshEngine/Public/ChaosFlesh/ChaosDeformableTetrahedralComponent.h
@@ -141,6 +141,8 @@ public:
 	virtual TArray<FVector3f> GetGeometryCachePositions(const USkeletalMesh* SkeletalMesh) const override;
 
 	virtual TOptional<TArray<int32>> GetMeshImportVertexMap(const USkinnedAsset& SkinnedMeshAsset) const override;
+
+	virtual bool IsSkeletalMeshAnimationCompatible(const USkeletalMeshComponent* SkeletalMeshComponent) const override;
 	//~ End IDataflowGeometryCachable Interface
 
 	virtual void SetMaterial(int32 Index, UMaterialInterface* InMaterial) override
diff --git a/Engine/Plugins/Experimental/Dataflow/Source/DataflowEditor/Private/Dataflow/DataflowSimulationControls.cpp b/Engine/Plugins/Experimental/Dataflow/Source/DataflowEditor/Private/Dataflow/DataflowSimulationControls.cpp
--- a/Engine/Plugins/Experimental/Dataflow/Source/DataflowEditor/Private/Dataflow/DataflowSimulationControls.cpp
+++ b/Engine/Plugins/Experimental/Dataflow/Source/DataflowEditor/Private/Dataflow/DataflowSimulationControls.cpp
@@ -15,7 +15,7 @@
 #include "Engine/World.h"
 #include "BonePose.h"
 #include "AnimationRuntime.h"
-
+#include "Dataflow/Interfaces/DataflowInterfaceGeometryCachable.h"
 #define LOCTEXT_NAMESPACE "DataflowSimulationGenerator"
 
 namespace UE::Dataflow
@@ -100,6 +100,15 @@ namespace UE::Dataflow
 			TInlineComponentArray<UPrimitiveComponent*> PrimComponents;
 			PreviewActor->GetComponents(PrimComponents);
 
+			TArray<IDataflowGeometryCachable*> GeometryCachables;
+			for (UPrimitiveComponent* PrimComponent : PrimComponents)
+			{
+				if (IDataflowGeometryCachable* SkeletalMeshComponent = Cast<IDataflowGeometryCachable>(PrimComponent))
+				{
+					GeometryCachables.Add(SkeletalMeshComponent);
+				}
+			}
+
 			for(UPrimitiveComponent* PrimComponent : PrimComponents)
 			{
 				if(USkeletalMeshComponent* SkeletalMeshComponent = Cast<USkeletalMeshComponent>(PrimComponent))
@@ -110,6 +119,13 @@ namespace UE::Dataflow
 					
 					if(UAnimSingleNodeInstance* AnimNodeInstance = SkeletalMeshComponent->GetSingleNodeInstance())
 					{
+						for (IDataflowGeometryCachable* GeometryCachable : GeometryCachables)
+						{
+							if (!GeometryCachable->IsSkeletalMeshAnimationCompatible(SkeletalMeshComponent))
+							{
+								UE_LOG(LogChaosSimulation, Warning, TEXT("Asset is not compatible with the skeletal mesh [%s] for animation updates, check if Skeletons match"), *SkeletalMeshComponent->GetSkeletalMeshAsset()->GetName());
+							}
+						}
 						// Setup the animation instance
 						AnimNodeInstance->SetAnimationAsset(SkeletalMeshComponent->AnimationData.AnimToPlay);
 						AnimNodeInstance->InitializeAnimation();
@@ -169,7 +185,6 @@ namespace UE::Dataflow
 		{
 			TInlineComponentArray<UPrimitiveComponent*> PrimComponents;
 			PreviewActor->GetComponents(PrimComponents);
-
 			// Update all the animation time 
 			for(UPrimitiveComponent* PrimComponent : PrimComponents)
 			{
diff --git a/Engine/Source/Runtime/Experimental/Dataflow/Simulation/Public/Dataflow/Interfaces/DataflowInterfaceGeometryCachable.h b/Engine/Source/Runtime/Experimental/Dataflow/Simulation/Public/Dataflow/Interfaces/DataflowInterfaceGeometryCachable.h
--- a/Engine/Source/Runtime/Experimental/Dataflow/Simulation/Public/Dataflow/Interfaces/DataflowInterfaceGeometryCachable.h
+++ b/Engine/Source/Runtime/Experimental/Dataflow/Simulation/Public/Dataflow/Interfaces/DataflowInterfaceGeometryCachable.h
@@ -29,4 +29,7 @@ public:
 
 	/** Get mapping from final mesh vertex index to raw import vertex index */
 	virtual TOptional<TArray<int32>> GetMeshImportVertexMap(const USkinnedAsset& SkinnedMeshAsset) const = 0;
+
+	/** Returns true if the asset supports animation update from a skeletal mesh  */
+	virtual bool IsSkeletalMeshAnimationCompatible(const USkeletalMeshComponent* SkeletalMeshComponent) const { return true; };
 };
