Dataflow Editor: Fix grid rendering in Simulation Viewport

diff --git a/Engine/Plugins/Experimental/Dataflow/Source/DataflowEditor/Private/Dataflow/DataflowEditorToolkit.cpp b/Engine/Plugins/Experimental/Dataflow/Source/DataflowEditor/Private/Dataflow/DataflowEditorToolkit.cpp
--- a/Engine/Plugins/Experimental/Dataflow/Source/DataflowEditor/Private/Dataflow/DataflowEditorToolkit.cpp
+++ b/Engine/Plugins/Experimental/Dataflow/Source/DataflowEditor/Private/Dataflow/DataflowEditorToolkit.cpp
@@ -739,6 +739,8 @@ void FDataflowEditorToolkit::CreateSimulationViewportClient()
 		}
 		return Viewport;
 	};
+
+	SimulationViewportClient->RegisterDelegates();
 }
 
 //~ End FBaseAssetToolkit overrides
diff --git a/Engine/Plugins/Experimental/Dataflow/Source/DataflowEditor/Private/Dataflow/DataflowSimulationViewportClient.cpp b/Engine/Plugins/Experimental/Dataflow/Source/DataflowEditor/Private/Dataflow/DataflowSimulationViewportClient.cpp
--- a/Engine/Plugins/Experimental/Dataflow/Source/DataflowEditor/Private/Dataflow/DataflowSimulationViewportClient.cpp
+++ b/Engine/Plugins/Experimental/Dataflow/Source/DataflowEditor/Private/Dataflow/DataflowSimulationViewportClient.cpp
@@ -18,6 +18,7 @@
 #include "Selection.h"
 #include "SGraphPanel.h"
 #include "SNodePanel.h"
+#include "AssetViewerSettings.h"
 
 FDataflowSimulationViewportClient::FDataflowSimulationViewportClient(FEditorModeTools* InModeTools,
                                                              FPreviewScene* InPreviewScene,  const bool bCouldTickScene,
@@ -34,6 +35,11 @@ FDataflowSimulationViewportClient::FDataflowSimulationViewportClient(FEditorMode
 	bEnableSceneTicking = bCouldTickScene;
 }
 
+FDataflowSimulationViewportClient::~FDataflowSimulationViewportClient()
+{
+	UAssetViewerSettings::Get()->OnAssetViewerSettingsChanged().Remove(OnAssetViewerSettingsChangedDelegateHandle);
+}
+
 void FDataflowSimulationViewportClient::SetDataflowEditorToolkit(TWeakPtr<FDataflowEditorToolkit> InDataflowEditorToolkitPtr)
 {
 	DataflowEditorToolkitPtr = InDataflowEditorToolkitPtr;
@@ -90,6 +96,36 @@ void FDataflowSimulationViewportClient::OnViewportClicked(HHitProxy* HitProxy)
 	}
 }
 
+void FDataflowSimulationViewportClient::RegisterDelegates()
+{
+	// Remove any existing delegate in case this function is called twice
+	UAssetViewerSettings::Get()->OnAssetViewerSettingsChanged().Remove(OnAssetViewerSettingsChangedDelegateHandle);
+
+	auto SyncDrawGridSettings = [this]()
+	{
+		checkf(PreviewScene, TEXT("Expected valid PreviewScene pointer"));
+		const int32 CurrentProfileIndex = PreviewScene->GetCurrentProfileIndex();
+
+		const UAssetViewerSettings* const DefaultSettings = UAssetViewerSettings::Get();
+
+		checkf(DefaultSettings && DefaultSettings->Profiles.IsValidIndex(CurrentProfileIndex), TEXT("Invalid default settings pointer or current profile index"));
+
+		const FPreviewSceneProfile& Profile = DefaultSettings->Profiles[CurrentProfileIndex];
+		EngineShowFlags.Grid = Profile.bShowGrid;
+		DrawHelper.bDrawGrid = Profile.bShowGrid;
+	};
+
+	// Run the lambda once to get things in sync
+	SyncDrawGridSettings();
+
+	OnAssetViewerSettingsChangedDelegateHandle = UAssetViewerSettings::Get()->OnAssetViewerSettingsChanged().AddLambda([this, SyncDrawGridSettings](const FName& InPropertyName)
+	{
+		if ((InPropertyName == GET_MEMBER_NAME_CHECKED(FPreviewSceneProfile, bShowGrid)) || (InPropertyName == NAME_None))		// NAME_None is passed in when the current profile changes
+		{
+			SyncDrawGridSettings();
+		}
+	});
+}
 
 void FDataflowSimulationViewportClient::AddReferencedObjects(FReferenceCollector& Collector)
 {
diff --git a/Engine/Plugins/Experimental/Dataflow/Source/DataflowEditor/Public/Dataflow/DataflowEditorViewportClientBase.h b/Engine/Plugins/Experimental/Dataflow/Source/DataflowEditor/Public/Dataflow/DataflowEditorViewportClientBase.h
--- a/Engine/Plugins/Experimental/Dataflow/Source/DataflowEditor/Public/Dataflow/DataflowEditorViewportClientBase.h
+++ b/Engine/Plugins/Experimental/Dataflow/Source/DataflowEditor/Public/Dataflow/DataflowEditorViewportClientBase.h
@@ -15,6 +15,9 @@ public:
 
 	FDataflowEditorViewportClientBase(FEditorModeTools* InModeTools, FPreviewScene* InPreviewScene,  const bool bCouldTickScene, const TWeakPtr<SEditorViewport> InEditorViewportWidget = nullptr);
 
+	// Call this after construction to initialize callbacks when settings change
+	virtual void RegisterDelegates() {}
+
 protected:
 
 	virtual void OnViewportClicked(HHitProxy* HitProxy) = 0;
diff --git a/Engine/Plugins/Experimental/Dataflow/Source/DataflowEditor/Public/Dataflow/DataflowSimulationViewportClient.h b/Engine/Plugins/Experimental/Dataflow/Source/DataflowEditor/Public/Dataflow/DataflowSimulationViewportClient.h
--- a/Engine/Plugins/Experimental/Dataflow/Source/DataflowEditor/Public/Dataflow/DataflowSimulationViewportClient.h
+++ b/Engine/Plugins/Experimental/Dataflow/Source/DataflowEditor/Public/Dataflow/DataflowSimulationViewportClient.h
@@ -25,6 +25,11 @@ public:
 	FDataflowSimulationViewportClient(FEditorModeTools* InModeTools, FPreviewScene* InPreviewScene,  const bool bCouldTickScene,
 								  const TWeakPtr<SEditorViewport> InEditorViewportWidget = nullptr);
 
+	virtual ~FDataflowSimulationViewportClient();
+
+	// FDataflowEditorViewportClientBase
+	virtual void RegisterDelegates() override;
+
 	/** Set the data flow toolkit used to create the client*/
 	void SetDataflowEditorToolkit(TWeakPtr<FDataflowEditorToolkit> DataflowToolkit);
 	
@@ -49,6 +54,8 @@ private:
 	virtual void Draw(const FSceneView* View, FPrimitiveDrawInterface* PDI) override;
 	virtual void DrawCanvas(FViewport& InViewport, FSceneView& View, FCanvas& Canvas) override;
 
+	FDelegateHandle OnAssetViewerSettingsChangedDelegateHandle;
+
 	/** Toolkit used to create the viewport client */
 	TWeakPtr<FDataflowEditorToolkit> DataflowEditorToolkitPtr = nullptr;
 
