Dataflow : fix action menu filtering issue when pulling from an anytype anytype pin
- using now the type defined in the input of the dataflow node rather than the pin itself which is less "precise" about the type

diff --git a/Engine/Plugins/Experimental/Dataflow/Source/DataflowEditor/Private/Dataflow/DataflowSchema.cpp b/Engine/Plugins/Experimental/Dataflow/Source/DataflowEditor/Private/Dataflow/DataflowSchema.cpp
--- a/Engine/Plugins/Experimental/Dataflow/Source/DataflowEditor/Private/Dataflow/DataflowSchema.cpp
+++ b/Engine/Plugins/Experimental/Dataflow/Source/DataflowEditor/Private/Dataflow/DataflowSchema.cpp
@@ -116,26 +116,46 @@ bool UDataflowSchema::CanPinBeConnectedToNode(const UEdGraphPin* Pin, const UE::
 	}
 	if (NodeParameters.DefaultNodeObject)
 	{
-		const FName PinType = Pin->PinType.PinCategory;
-		if (Pin->Direction == EEdGraphPinDirection::EGPD_Output)
+		// get the type from the dataflow input/output as the pin type may not be precise enough when using anytypes
+		TSharedPtr<const FDataflowNode> DataflowNode;
+		if (UDataflowEdNode* EdNode = Cast<UDataflowEdNode>(Pin->GetOwningNode()))
 		{
-			const TArray<FDataflowInput*> Inputs = NodeParameters.DefaultNodeObject->GetInputs();
-			for (const FDataflowInput* Input : Inputs)
+			DataflowNode = EdNode->GetDataflowNode();
+		}
+		if (DataflowNode)
+		{
+			if (Pin->Direction == EEdGraphPinDirection::EGPD_Output)
 			{
-				if (Input && Input->SupportsType(PinType))
+				if (const FDataflowOutput* PinOutput = DataflowNode->FindOutput(Pin->PinName))
 				{
-					return true;
+					const TArray<FDataflowInput*> Inputs = NodeParameters.DefaultNodeObject->GetInputs();
+					for (const FDataflowInput* Input : Inputs)
+					{
+						if (Input && (
+							Input->SupportsType(PinOutput->GetType()) || 
+							PinOutput->SupportsType(Input->GetType())
+							))
+						{
+							return true;
+						}
+					}
 				}
 			}
-		}
-		else if (Pin->Direction == EEdGraphPinDirection::EGPD_Input)
-		{
-			const TArray<FDataflowOutput*> Outputs = NodeParameters.DefaultNodeObject->GetOutputs();
-			for (const FDataflowOutput* Output : Outputs)
+			else if (Pin->Direction == EEdGraphPinDirection::EGPD_Input)
 			{
-				if (Output && Output->SupportsType(PinType))
+				if (const FDataflowInput* PinInput = DataflowNode->FindInput(Pin->PinName))
 				{
-					return true;
+					const TArray<FDataflowOutput*> Outputs = NodeParameters.DefaultNodeObject->GetOutputs();
+					for (const FDataflowOutput* Output : Outputs)
+					{
+						if (Output && (
+							Output->SupportsType(PinInput->GetType()) ||
+							PinInput->SupportsType(Output->GetType())
+							))
+						{
+							return true;
+						}
+					}
 				}
 			}
 		}
