Dataflow : change double clicking on node to check for CTRL key to open source code

diff --git a/Engine/Plugins/Experimental/Dataflow/Source/DataflowEditor/Private/Dataflow/DataflowSNode.cpp b/Engine/Plugins/Experimental/Dataflow/Source/DataflowEditor/Private/Dataflow/DataflowSNode.cpp
--- a/Engine/Plugins/Experimental/Dataflow/Source/DataflowEditor/Private/Dataflow/DataflowSNode.cpp
+++ b/Engine/Plugins/Experimental/Dataflow/Source/DataflowEditor/Private/Dataflow/DataflowSNode.cpp
@@ -276,17 +276,20 @@ void SDataflowEdNode::UpdateErrorInfo()
 
 FReply SDataflowEdNode::OnMouseButtonDoubleClick(const FGeometry& InMyGeometry, const FPointerEvent& InMouseEvent)
 {
-	if (GraphNode)
+	if (InMouseEvent.GetModifierKeys().IsControlDown())
 	{
-		if (UDataflowEdNode* DataflowNode = Cast<UDataflowEdNode>(GraphNode))
+		if (GraphNode)
 		{
-			if (TSharedPtr<UE::Dataflow::FGraph> Graph = DataflowNode->GetDataflowGraph())
+			if (UDataflowEdNode* DataflowNode = Cast<UDataflowEdNode>(GraphNode))
 			{
-				if (TSharedPtr<FDataflowNode> Node = Graph->FindBaseNode(DataflowNode->GetDataflowNodeGuid()))
+				if (TSharedPtr<UE::Dataflow::FGraph> Graph = DataflowNode->GetDataflowGraph())
 				{
-					if (FSourceCodeNavigation::CanNavigateToStruct(Node->TypedScriptStruct()))
+					if (TSharedPtr<FDataflowNode> Node = Graph->FindBaseNode(DataflowNode->GetDataflowNodeGuid()))
 					{
-						FSourceCodeNavigation::NavigateToStruct(Node->TypedScriptStruct());
+						if (FSourceCodeNavigation::CanNavigateToStruct(Node->TypedScriptStruct()))
+						{
+							FSourceCodeNavigation::NavigateToStruct(Node->TypedScriptStruct());
+						}
 					}
 				}
 			}
