Dataflow Editor viewport clients:
- Move callback that responds to changes in the preview scene profile up to the common base class
- Respect the "Advanced Features" checkbox from the preview scene profile in this callback
- Remove duplicate BehaviorSet variables from the derived classes

diff --git a/Engine/Plugins/Experimental/Dataflow/Source/DataflowEditor/Private/Dataflow/DataflowConstructionViewportClient.cpp b/Engine/Plugins/Experimental/Dataflow/Source/DataflowEditor/Private/Dataflow/DataflowConstructionViewportClient.cpp
--- a/Engine/Plugins/Experimental/Dataflow/Source/DataflowEditor/Private/Dataflow/DataflowConstructionViewportClient.cpp
+++ b/Engine/Plugins/Experimental/Dataflow/Source/DataflowEditor/Private/Dataflow/DataflowConstructionViewportClient.cpp
@@ -241,12 +241,6 @@ void FDataflowConstructionViewportClient::SetConstructionViewMode(const UE::Data
 }
 
 
-void FDataflowConstructionViewportClient::AddReferencedObjects(FReferenceCollector& Collector)
-{
-	Super::AddReferencedObjects(Collector);
-	Collector.AddReferencedObject(BehaviorSet);
-}
-
 void FDataflowConstructionViewportClient::Draw(const FSceneView* View, FPrimitiveDrawInterface* PDI)
 {
 	FEditorViewportClient::Draw(View, PDI);
diff --git a/Engine/Plugins/Experimental/Dataflow/Source/DataflowEditor/Private/Dataflow/DataflowEditorToolkit.cpp b/Engine/Plugins/Experimental/Dataflow/Source/DataflowEditor/Private/Dataflow/DataflowEditorToolkit.cpp
--- a/Engine/Plugins/Experimental/Dataflow/Source/DataflowEditor/Private/Dataflow/DataflowEditorToolkit.cpp
+++ b/Engine/Plugins/Experimental/Dataflow/Source/DataflowEditor/Private/Dataflow/DataflowEditorToolkit.cpp
@@ -676,6 +676,7 @@ TSharedPtr<FEditorViewportClient> FDataflowEditorToolkit::CreateEditorViewportCl
 	EditorModeManager.Get(), ObjectScene.Get(), true);
 	LocalConstructionClient->SetDataflowEditorToolkit(StaticCastSharedRef<FDataflowEditorToolkit>(
 		const_cast<FDataflowEditorToolkit*>(this)->AsShared()));
+
 	return LocalConstructionClient;
 }
 
@@ -705,8 +706,6 @@ void FDataflowEditorToolkit::CreateSimulationViewportClient()
 		}
 		return Viewport;
 	};
-
-	SimulationViewportClient->RegisterDelegates();
 }
 
 //~ End FBaseAssetToolkit overrides
diff --git a/Engine/Plugins/Experimental/Dataflow/Source/DataflowEditor/Private/Dataflow/DataflowEditorViewportClientBase.cpp b/Engine/Plugins/Experimental/Dataflow/Source/DataflowEditor/Private/Dataflow/DataflowEditorViewportClientBase.cpp
--- a/Engine/Plugins/Experimental/Dataflow/Source/DataflowEditor/Private/Dataflow/DataflowEditorViewportClientBase.cpp
+++ b/Engine/Plugins/Experimental/Dataflow/Source/DataflowEditor/Private/Dataflow/DataflowEditorViewportClientBase.cpp
@@ -1,9 +1,11 @@
 // Copyright Epic Games, Inc. All Rights Reserved.
 
 #include "Dataflow/DataflowEditorViewportClientBase.h"
+#include "Dataflow/DataflowEditorPreviewSceneBase.h"
 #include "EditorModeManager.h"
 #include "EdModeInteractiveToolsContext.h"
 #include "BaseBehaviors/ClickDragBehavior.h"
+#include "AssetViewerSettings.h"
 
 #define LOCTEXT_NAMESPACE "DataflowEditorViewportClientBase"
 
@@ -41,6 +43,15 @@ FDataflowEditorViewportClientBase::FDataflowEditorViewportClientBase(FEditorMode
 	BehaviorSet->Add(ClickDragBehavior);
 
 	ModeTools->GetInteractiveToolsContext()->InputRouter->RegisterSource(this);
+
+	DataflowPreviewScene = static_cast<FDataflowPreviewSceneBase*>(InPreviewScene);
+
+	RegisterDelegates();
+}
+
+FDataflowEditorViewportClientBase::~FDataflowEditorViewportClientBase()
+{
+	DeregisterDelegates();
 }
 
 const UInputBehaviorSet* FDataflowEditorViewportClientBase::GetInputBehaviors() const
@@ -48,4 +59,57 @@ const UInputBehaviorSet* FDataflowEditorViewportClientBase::GetInputBehaviors()
 	return BehaviorSet;
 }
 
+void FDataflowEditorViewportClientBase::AddReferencedObjects(FReferenceCollector& Collector)
+{
+	Super::AddReferencedObjects(Collector);
+	Collector.AddReferencedObject(BehaviorSet);
+}
+
+void FDataflowEditorViewportClientBase::RegisterDelegates()
+{
+	// Remove any existing delegate in case this function is called twice
+	UAssetViewerSettings::Get()->OnAssetViewerSettingsChanged().Remove(OnAssetViewerSettingsChangedDelegateHandle);
+
+	auto SyncEngineShowFlags = [this]()
+	{
+		checkf(DataflowPreviewScene, TEXT("Expected valid DataflowPreviewScene pointer"));
+		const int32 CurrentProfileIndex = DataflowPreviewScene->GetCurrentProfileIndex();
+
+		const UAssetViewerSettings* const DefaultSettings = UAssetViewerSettings::Get();
+
+		checkf(DefaultSettings && DefaultSettings->Profiles.IsValidIndex(CurrentProfileIndex), TEXT("Invalid default settings pointer or current profile index"));
+
+		const FPreviewSceneProfile& Profile = DefaultSettings->Profiles[CurrentProfileIndex];
+		EngineShowFlags.Grid = Profile.bShowGrid;
+		DrawHelper.bDrawGrid = Profile.bShowGrid;
+
+		if (Profile.bPostProcessingEnabled)
+		{
+			EngineShowFlags.EnableAdvancedFeatures();
+		}
+		else
+		{
+			EngineShowFlags.DisableAdvancedFeatures();
+		}
+	};
+
+	// Run the lambda once to get things in sync
+	SyncEngineShowFlags();
+
+	OnAssetViewerSettingsChangedDelegateHandle = UAssetViewerSettings::Get()->OnAssetViewerSettingsChanged().AddLambda([this, SyncEngineShowFlags](const FName& InPropertyName)
+	{
+		if ((InPropertyName == GET_MEMBER_NAME_CHECKED(FPreviewSceneProfile, bShowGrid)) ||
+			(InPropertyName == GET_MEMBER_NAME_CHECKED(FPreviewSceneProfile, bPostProcessingEnabled)) ||
+			(InPropertyName == NAME_None))		// NAME_None is passed in when the current profile changes
+		{
+			SyncEngineShowFlags();
+		}
+	});
+}
+
+void FDataflowEditorViewportClientBase::DeregisterDelegates()
+{
+	UAssetViewerSettings::Get()->OnAssetViewerSettingsChanged().Remove(OnAssetViewerSettingsChangedDelegateHandle);
+}
+
 #undef LOCTEXT_NAMESPACE 
diff --git a/Engine/Plugins/Experimental/Dataflow/Source/DataflowEditor/Private/Dataflow/DataflowSimulationViewportClient.cpp b/Engine/Plugins/Experimental/Dataflow/Source/DataflowEditor/Private/Dataflow/DataflowSimulationViewportClient.cpp
--- a/Engine/Plugins/Experimental/Dataflow/Source/DataflowEditor/Private/Dataflow/DataflowSimulationViewportClient.cpp
+++ b/Engine/Plugins/Experimental/Dataflow/Source/DataflowEditor/Private/Dataflow/DataflowSimulationViewportClient.cpp
@@ -18,7 +18,6 @@
 #include "Selection.h"
 #include "SGraphPanel.h"
 #include "SNodePanel.h"
-#include "AssetViewerSettings.h"
 
 FDataflowSimulationViewportClient::FDataflowSimulationViewportClient(FEditorModeTools* InModeTools,
                                                              FPreviewScene* InPreviewScene,  const bool bCouldTickScene,
@@ -31,15 +30,9 @@ FDataflowSimulationViewportClient::FDataflowSimulationViewportClient(FEditorMode
 	EngineShowFlags.SetSelectionOutline(true);
 	EngineShowFlags.EnableAdvancedFeatures();
 
-	PreviewScene = static_cast<FDataflowPreviewSceneBase*>(InPreviewScene);
 	bEnableSceneTicking = bCouldTickScene;
 }
 
-FDataflowSimulationViewportClient::~FDataflowSimulationViewportClient()
-{
-	UAssetViewerSettings::Get()->OnAssetViewerSettingsChanged().Remove(OnAssetViewerSettingsChangedDelegateHandle);
-}
-
 void FDataflowSimulationViewportClient::SetDataflowEditorToolkit(TWeakPtr<FDataflowEditorToolkit> InDataflowEditorToolkitPtr)
 {
 	DataflowEditorToolkitPtr = InDataflowEditorToolkitPtr;
@@ -53,9 +46,9 @@ void FDataflowSimulationViewportClient::SetToolCommandList(TWeakPtr<FUICommandLi
 void FDataflowSimulationViewportClient::Tick(float DeltaSeconds)
 {
 	Super::Tick(DeltaSeconds);
-	if (PreviewScene)
+	if (DataflowPreviewScene)
 	{
-		PreviewScene->TickDataflowScene(DeltaSeconds);
+		DataflowPreviewScene->TickDataflowScene(DeltaSeconds);
 	}
 }
 
@@ -96,42 +89,6 @@ void FDataflowSimulationViewportClient::OnViewportClicked(HHitProxy* HitProxy)
 	}
 }
 
-void FDataflowSimulationViewportClient::RegisterDelegates()
-{
-	// Remove any existing delegate in case this function is called twice
-	UAssetViewerSettings::Get()->OnAssetViewerSettingsChanged().Remove(OnAssetViewerSettingsChangedDelegateHandle);
-
-	auto SyncDrawGridSettings = [this]()
-	{
-		checkf(PreviewScene, TEXT("Expected valid PreviewScene pointer"));
-		const int32 CurrentProfileIndex = PreviewScene->GetCurrentProfileIndex();
-
-		const UAssetViewerSettings* const DefaultSettings = UAssetViewerSettings::Get();
-
-		checkf(DefaultSettings && DefaultSettings->Profiles.IsValidIndex(CurrentProfileIndex), TEXT("Invalid default settings pointer or current profile index"));
-
-		const FPreviewSceneProfile& Profile = DefaultSettings->Profiles[CurrentProfileIndex];
-		EngineShowFlags.Grid = Profile.bShowGrid;
-		DrawHelper.bDrawGrid = Profile.bShowGrid;
-	};
-
-	// Run the lambda once to get things in sync
-	SyncDrawGridSettings();
-
-	OnAssetViewerSettingsChangedDelegateHandle = UAssetViewerSettings::Get()->OnAssetViewerSettingsChanged().AddLambda([this, SyncDrawGridSettings](const FName& InPropertyName)
-	{
-		if ((InPropertyName == GET_MEMBER_NAME_CHECKED(FPreviewSceneProfile, bShowGrid)) || (InPropertyName == NAME_None))		// NAME_None is passed in when the current profile changes
-		{
-			SyncDrawGridSettings();
-		}
-	});
-}
-
-void FDataflowSimulationViewportClient::AddReferencedObjects(FReferenceCollector& Collector)
-{
-	Super::AddReferencedObjects(Collector);
-	Collector.AddReferencedObject(BehaviorSet);
-}
 
 void FDataflowSimulationViewportClient::Draw(const FSceneView* View, FPrimitiveDrawInterface* PDI)
 {
diff --git a/Engine/Plugins/Experimental/Dataflow/Source/DataflowEditor/Public/Dataflow/DataflowConstructionViewportClient.h b/Engine/Plugins/Experimental/Dataflow/Source/DataflowEditor/Public/Dataflow/DataflowConstructionViewportClient.h
--- a/Engine/Plugins/Experimental/Dataflow/Source/DataflowEditor/Public/Dataflow/DataflowConstructionViewportClient.h
+++ b/Engine/Plugins/Experimental/Dataflow/Source/DataflowEditor/Public/Dataflow/DataflowConstructionViewportClient.h
@@ -48,7 +48,6 @@ public:
 	void SetToolCommandList(TWeakPtr<FUICommandList> ToolCommandList);
 
 	// FGCObject Interface
-	virtual void AddReferencedObjects(FReferenceCollector& Collector) override;
 	virtual FString GetReferencerName() const override { return TEXT("FDataflowConstructionViewportClient"); }
 
 	DECLARE_MULTICAST_DELEGATE_OneParam(FOnSelectionChangedMulticast, const TArray<UPrimitiveComponent*>&)
@@ -78,9 +77,6 @@ private:
 	/** Construction view mode */
 	const UE::Dataflow::IDataflowConstructionViewMode* ConstructionViewMode = nullptr;
 
-	/** Behavior set for the behavior UI */
-	TObjectPtr<UInputBehaviorSet> BehaviorSet;
-	
 	/** Flag to enable scene ticking from the client */
 	bool bEnableSceneTicking = false;
 
diff --git a/Engine/Plugins/Experimental/Dataflow/Source/DataflowEditor/Public/Dataflow/DataflowEditorViewportClientBase.h b/Engine/Plugins/Experimental/Dataflow/Source/DataflowEditor/Public/Dataflow/DataflowEditorViewportClientBase.h
--- a/Engine/Plugins/Experimental/Dataflow/Source/DataflowEditor/Public/Dataflow/DataflowEditorViewportClientBase.h
+++ b/Engine/Plugins/Experimental/Dataflow/Source/DataflowEditor/Public/Dataflow/DataflowEditorViewportClientBase.h
@@ -8,24 +8,40 @@ class FEditorModeTools;
 class FPreviewScene;
 class SEditorViewport;
 class HHitProxy;
+class FDataflowPreviewSceneBase;
 
 class DATAFLOWEDITOR_API FDataflowEditorViewportClientBase : public FEditorViewportClient, public IInputBehaviorSource
 {
 public:
+	using Super = FEditorViewportClient;
 
 	FDataflowEditorViewportClientBase(FEditorModeTools* InModeTools, FPreviewScene* InPreviewScene,  const bool bCouldTickScene, const TWeakPtr<SEditorViewport> InEditorViewportWidget = nullptr);
-
-	// Call this after construction to initialize callbacks when settings change
-	virtual void RegisterDelegates() {}
+	
+	virtual ~FDataflowEditorViewportClientBase();
 
 protected:
 
 	virtual void OnViewportClicked(HHitProxy* HitProxy) = 0;
 
+	/** Dataflow preview scene from the toolkit */
+	FDataflowPreviewSceneBase* DataflowPreviewScene = nullptr;
+
 private:
 
 	// IInputBehaviorSource
 	virtual const UInputBehaviorSet* GetInputBehaviors() const override;
 
+	// FGCObject
+	virtual void AddReferencedObjects(FReferenceCollector& Collector) override;
+
+	/** Called from constructor/destructor to manage callbacks when settings change */
+	void RegisterDelegates();
+	void DeregisterDelegates();
+
+	/** Behaviors available to the current viewport */
 	TObjectPtr<UInputBehaviorSet> BehaviorSet;
+
+	/** Handle for callback when preview scene profile settings change */
+	FDelegateHandle OnAssetViewerSettingsChangedDelegateHandle;
+
 };
diff --git a/Engine/Plugins/Experimental/Dataflow/Source/DataflowEditor/Public/Dataflow/DataflowSimulationViewportClient.h b/Engine/Plugins/Experimental/Dataflow/Source/DataflowEditor/Public/Dataflow/DataflowSimulationViewportClient.h
--- a/Engine/Plugins/Experimental/Dataflow/Source/DataflowEditor/Public/Dataflow/DataflowSimulationViewportClient.h
+++ b/Engine/Plugins/Experimental/Dataflow/Source/DataflowEditor/Public/Dataflow/DataflowSimulationViewportClient.h
@@ -25,10 +25,7 @@ public:
 	FDataflowSimulationViewportClient(FEditorModeTools* InModeTools, FPreviewScene* InPreviewScene,  const bool bCouldTickScene,
 								  const TWeakPtr<SEditorViewport> InEditorViewportWidget = nullptr);
 
-	virtual ~FDataflowSimulationViewportClient();
-
-	// FDataflowEditorViewportClientBase
-	virtual void RegisterDelegates() override;
+	virtual ~FDataflowSimulationViewportClient() = default;
 
 	/** Set the data flow toolkit used to create the client*/
 	void SetDataflowEditorToolkit(TWeakPtr<FDataflowEditorToolkit> DataflowToolkit);
@@ -40,7 +37,6 @@ public:
 	void SetToolCommandList(TWeakPtr<FUICommandList> ToolCommandList);
 
 	// FGCObject Interface
-	virtual void AddReferencedObjects(FReferenceCollector& Collector) override;
 	virtual FString GetReferencerName() const override { return TEXT("FDataflowSimulationViewportClient"); }
 
 private:
@@ -54,19 +50,11 @@ private:
 	virtual void Draw(const FSceneView* View, FPrimitiveDrawInterface* PDI) override;
 	virtual void DrawCanvas(FViewport& InViewport, FSceneView& View, FCanvas& Canvas) override;
 
-	FDelegateHandle OnAssetViewerSettingsChangedDelegateHandle;
-
 	/** Toolkit used to create the viewport client */
 	TWeakPtr<FDataflowEditorToolkit> DataflowEditorToolkitPtr = nullptr;
 
-	/** Dataflow preview scene from the toolkit */
-	FDataflowPreviewSceneBase* PreviewScene = nullptr;
-
 	// @todo(brice) : Is this needed?
 	TWeakPtr<FUICommandList> ToolCommandList;
-
-	/** Behavior set for the behavior UI */
-	TObjectPtr<UInputBehaviorSet> BehaviorSet;
 	
 	/** Flag to enable scene ticking from the client */
 	bool bEnableSceneTicking = false;
