Dataflow : editor utility function to modify the dataflow asset using transactions
- currently implement most of the variable related functions

diff --git a/Engine/Plugins/Experimental/Dataflow/Source/DataflowEditor/Private/Dataflow/DataflowAssetEditUtils.cpp b/Engine/Plugins/Experimental/Dataflow/Source/DataflowEditor/Private/Dataflow/DataflowAssetEditUtils.cpp
new file mode 100644
--- /dev/null
+++ b/Engine/Plugins/Experimental/Dataflow/Source/DataflowEditor/Private/Dataflow/DataflowAssetEditUtils.cpp
@@ -0,0 +1,309 @@
+// Copyright Epic Games, Inc. All Rights Reserved.
+
+#include "Dataflow/DataflowAssetEditUtils.h"
+
+#include "HAL/PlatformApplicationMisc.h"
+#include "Dataflow/DataflowObject.h"
+#include "ScopedTransaction.h"
+#include "PropertyBagDetails.h"
+
+
+#define LOCTEXT_NAMESPACE "DataflowAssetEditUtils"
+
+namespace UE::Dataflow::EditAssetUtils
+{
+	namespace Private
+	{
+		static const TCHAR* DataflowVariableClipboardPrefix = TEXT("DataflowVariable_");
+
+		enum class EChangeResult
+		{
+			None,
+			Changed,
+			Cancel,
+		};
+
+		/** 
+		* Change a dataflow asset with a transaction
+		* The InFunction return paramater determines if the asset will be modified or not 
+		* if modification happens, a PostEditChangeProperty notification is sent
+		*/
+		static void ChangeDataflowAssetWithTransaction(UDataflow* DataflowAsset, const FText &TransactionName, TFunctionRef<EChangeResult(UDataflow&)> InFunction)
+		{
+			if (DataflowAsset)
+			{
+				FScopedTransaction Transaction(TransactionName);
+
+				const EChangeResult Result = InFunction(*DataflowAsset);
+				switch (Result)
+				{
+				case EChangeResult::Cancel:
+					Transaction.Cancel();
+					break;
+
+				case EChangeResult::Changed:
+				{
+					DataflowAsset->Modify();
+					FPropertyChangedEvent PropertyChangedEvent(nullptr);
+					if (UClass* DataflowClass = DataflowAsset->GetClass())
+					{
+						FProperty* MemberProperty = DataflowClass->FindPropertyByName(GET_MEMBER_NAME_CHECKED(UDataflow, Variables));
+						PropertyChangedEvent.SetActiveMemberProperty(MemberProperty);
+					}
+					DataflowAsset->PostEditChangeProperty(PropertyChangedEvent);
+					break;
+				}
+				case EChangeResult::None:
+				default:
+					break;
+				}
+			}
+		}
+
+		/** returns true if the variable was modified */
+		static bool ModifyVariable(UDataflow& DataflowAsset, FName Variable, TFunctionRef<void(FPropertyBagPropertyDesc& PropertyDesc)> InFunction)
+		{
+			bool bModified = false;
+
+			TArray<FPropertyBagPropertyDesc> NewPropertyDescs;
+			NewPropertyDescs.Append(DataflowAsset.Variables.GetPropertyBagStruct()->GetPropertyDescs());
+			for (FPropertyBagPropertyDesc& PropertyDesc : NewPropertyDescs)
+			{
+				if (PropertyDesc.Name == Variable)
+				{
+					InFunction(PropertyDesc);
+					bModified = true;
+				}
+			}
+
+			if (bModified)
+			{
+				if (const UPropertyBag* NewBagStruct = UPropertyBag::GetOrCreateFromDescs(NewPropertyDescs))
+				{
+					DataflowAsset.Variables.MigrateToNewBagStruct(NewBagStruct);
+					return true;
+				}
+			}
+			return false;
+		}
+
+		/** Generate a Dataflow asset variable unique name from a BaseName*/
+		static FName GenerateUniqueVariableName(const UDataflow& DataflowAsset, FName BaseName)
+		{
+			int32 Counter = 1;
+			FName UniqueName{ BaseName };
+			const FString BasenameStr{ BaseName.ToString() };
+			while (true)
+			{
+				if (DataflowAsset.Variables.FindPropertyDescByName(UniqueName) == nullptr)
+				{
+					break; // found an available name exit 
+				}
+				UniqueName = FName(FString::Format(TEXT("{0}_{1}"), { BasenameStr, FString::FromInt(Counter++)}));
+			}
+			return UniqueName;
+		}
+	}
+
+	//////////////////////////////////////////////////////////////////////////////////////////////////////////
+	// 
+	// VARIABLES FUNCTIONS
+	// 
+	//////////////////////////////////////////////////////////////////////////////////////////////////////////
+	FName AddNewVariable(UDataflow* DataflowAsset, FName BaseName)
+	{
+		const FText TransactionName{ LOCTEXT("AddNewDataflowVariable", "Add New Dataflow Variable") };
+
+		FName UniqueVariableName;
+
+		auto AddNewVariableInternal =
+			[&BaseName, &UniqueVariableName](UDataflow& DataflowAsset) -> Private::EChangeResult
+			{
+				UniqueVariableName = Private::GenerateUniqueVariableName(DataflowAsset, BaseName);
+				DataflowAsset.Variables.AddProperty(UniqueVariableName, EPropertyBagPropertyType::Int32);
+				return Private::EChangeResult::Changed;
+			};
+
+		Private::ChangeDataflowAssetWithTransaction(DataflowAsset, TransactionName, AddNewVariableInternal);
+
+		return UniqueVariableName;
+	}
+
+	//////////////////////////////////////////////////////////////////////////////////////////////////////////
+	void DeleteVariable(UDataflow* DataflowAsset, FName VariableName)
+	{
+		const FText TransactionName{ LOCTEXT("DeleteDataflowVariable", "Delete Dataflow Variable") };
+
+		auto DeleteVariableInternal =
+			[&VariableName](UDataflow& DataflowAsset) -> Private::EChangeResult
+			{
+				DataflowAsset.Variables.RemovePropertyByName(VariableName);
+				return Private::EChangeResult::Changed;
+			};
+
+		Private::ChangeDataflowAssetWithTransaction(DataflowAsset, TransactionName, DeleteVariableInternal);
+	}
+
+	//////////////////////////////////////////////////////////////////////////////////////////////////////////
+	void RenameVariable(UDataflow* DataflowAsset, FName OldVariableName, FName NewVariableName)
+	{
+		const FText TransactionName{ LOCTEXT("RenameDataflowVariable", "Rename Dataflow Variable") };
+
+		auto SetVariableNameInternal =
+			[OldVariableName, NewVariableName](UDataflow& DataflowAsset) -> Private::EChangeResult
+			{
+				auto ChangePropertyNameLambda =
+					[NewVariableName](FPropertyBagPropertyDesc& PropertyDesc)
+					{
+						PropertyDesc.Name = NewVariableName;
+					};
+
+				const bool bModified = Private::ModifyVariable(DataflowAsset, OldVariableName, ChangePropertyNameLambda);
+				return (bModified) ? Private::EChangeResult::Changed : Private::EChangeResult::None;
+			};
+
+		Private::ChangeDataflowAssetWithTransaction(DataflowAsset, TransactionName, SetVariableNameInternal);
+	}
+
+	//////////////////////////////////////////////////////////////////////////////////////////////////////////
+	FName DuplicateVariable(UDataflow* DataflowAsset, FName VariableName)
+	{
+		const FText TransactionName = FText::Format(LOCTEXT("DuplicateDataflowVariable", "Duplicate Dataflow Variable: {0}"), FText::FromName(VariableName));
+		
+		FName NewVariableName;
+
+		auto DuplicateVariableInternal =
+			[VariableName, &NewVariableName](UDataflow& DataflowAsset) -> Private::EChangeResult
+			{
+				if (const FPropertyBagPropertyDesc* PropertyDescPtr = DataflowAsset.Variables.FindPropertyDescByName(VariableName))
+				{
+					NewVariableName = Private::GenerateUniqueVariableName(DataflowAsset, VariableName);
+
+					// make sure the name is unique and the GUID is invalidated to avoid copying it as is 
+					FPropertyBagPropertyDesc NewDesc(*PropertyDescPtr);
+					NewDesc.Name = NewVariableName;
+					NewDesc.ID.Invalidate();
+					DataflowAsset.Variables.AddProperties(MakeConstArrayView(&NewDesc, 1));
+					return Private::EChangeResult::Changed;
+				}
+				return Private::EChangeResult::None;
+			};
+
+		Private::ChangeDataflowAssetWithTransaction(DataflowAsset, TransactionName, DuplicateVariableInternal);
+
+		return NewVariableName;
+	}
+
+
+	//////////////////////////////////////////////////////////////////////////////////////////////////////////
+	void SetVariableType(UDataflow* DataflowAsset, FName VariableName, const FEdGraphPinType& PinType)
+	{
+		const FText TransactionName{ LOCTEXT("ChangeDataflowVariableType", "Change Dataflow Variable Type") };
+
+		auto SetVariableTypeInternal =
+			[VariableName, &PinType](UDataflow& DataflowAsset) -> Private::EChangeResult
+			{
+				auto ChangePropertyTypeLambda =
+					[VariableName, &PinType](FPropertyBagPropertyDesc& PropertyDesc)
+					{
+						UE::StructUtils::SetPropertyDescFromPin(PropertyDesc, PinType);
+					};
+
+				const bool bModified = Private::ModifyVariable(DataflowAsset, VariableName, ChangePropertyTypeLambda);
+				return (bModified) ? Private::EChangeResult::Changed : Private::EChangeResult::None;
+			};
+
+		Private::ChangeDataflowAssetWithTransaction(DataflowAsset, TransactionName, SetVariableTypeInternal);
+	}
+
+	//////////////////////////////////////////////////////////////////////////////////////////////////////////
+	void SetVariableValue(UDataflow* DataflowAsset, FName VariableName, const FInstancedPropertyBag & SourceBag)
+	{
+		const FText TransactionName{ LOCTEXT("ChangeDataflowVariableValue", "Change Dataflow Variable Value") };
+
+		auto SetVariableValueInternal =
+			[VariableName, &SourceBag](UDataflow& DataflowAsset) -> Private::EChangeResult
+			{
+				if (const FPropertyBagPropertyDesc* SourceDesc = SourceBag.FindPropertyDescByName(VariableName))
+				{
+					if (SourceDesc->CachedProperty)
+					{
+						const void* SourcePropertyAddr = SourceDesc->CachedProperty->ContainerPtrToValuePtr<void>(SourceBag.GetValue().GetMemory());
+						EPropertyBagResult Result = DataflowAsset.Variables.SetValue(VariableName, SourceDesc->CachedProperty, SourcePropertyAddr);
+						if (Result == EPropertyBagResult::Success)
+						{
+							return Private::EChangeResult::Changed;
+						}
+					}
+				}
+				return Private::EChangeResult::None;
+			};
+
+		// do we need a transction in that case ? 
+		Private::ChangeDataflowAssetWithTransaction(DataflowAsset, TransactionName, SetVariableValueInternal);
+	}
+
+	//////////////////////////////////////////////////////////////////////////////////////////////////////////
+	void CopyVariableToClipboard(UDataflow* DataflowAsset, FName VariableName)
+	{
+		if (DataflowAsset)
+		{
+			// no transaction needed in that case as we write to an external system 
+			if (const FPropertyBagPropertyDesc* PropertyDescPtr = DataflowAsset->Variables.FindPropertyDescByName(VariableName))
+			{
+				FString ClipboardPayload;
+
+				FPropertyBagPropertyDesc::StaticStruct()->ExportText(ClipboardPayload, PropertyDescPtr, PropertyDescPtr, nullptr, 0, nullptr, false);
+
+				if (!ClipboardPayload.IsEmpty())
+				{
+					ClipboardPayload = Private::DataflowVariableClipboardPrefix + ClipboardPayload;
+					FPlatformApplicationMisc::ClipboardCopy(ClipboardPayload.GetCharArray().GetData());
+				}
+			}
+		}
+	}
+
+	//////////////////////////////////////////////////////////////////////////////////////////////////////////
+	FName PasteVariableFromClipboard(UDataflow* DataflowAsset)
+	{
+		FString ClipboardPayload;
+		FPlatformApplicationMisc::ClipboardPaste(ClipboardPayload);
+
+		if (!ensure(ClipboardPayload.StartsWith(Private::DataflowVariableClipboardPrefix, ESearchCase::CaseSensitive)))
+		{
+			return NAME_None;
+		}
+
+		FStringOutputDevice Errors;
+		const TCHAR* ImportPayload = ClipboardPayload.GetCharArray().GetData() + FCString::Strlen(Private::DataflowVariableClipboardPrefix);
+
+		FPropertyBagPropertyDesc PropertyDesc;
+		FPropertyBagPropertyDesc::StaticStruct()->ImportText(ImportPayload, &PropertyDesc, nullptr, PPF_None, &Errors, FPropertyBagPropertyDesc::StaticStruct()->GetName());
+
+		if (Errors.IsEmpty())
+		{
+			if (DataflowAsset)
+			{
+				// make sure the name is unique and the GUID is invalidated to avoid copying it as is 
+				PropertyDesc.Name = Private::GenerateUniqueVariableName(*DataflowAsset, PropertyDesc.Name);
+				PropertyDesc.ID.Invalidate();
+
+				const FText TransactionName = FText::Format(LOCTEXT("PasteDataflowVariable", "Paste Dataflow Variable: {0}"), FText::FromName(PropertyDesc.Name));
+
+				auto PasteVariableInternal =
+					[&PropertyDesc](UDataflow& DataflowAsset) -> Private::EChangeResult
+					{
+						DataflowAsset.Variables.AddProperties(MakeConstArrayView(&PropertyDesc, 1));
+						return Private::EChangeResult::Changed;
+					};
+
+				Private::ChangeDataflowAssetWithTransaction(DataflowAsset, TransactionName, PasteVariableInternal);
+				return PropertyDesc.Name;
+			}
+		}
+		return NAME_None;
+	}
+}
+
+#undef LOCTEXT_NAMESPACE
\ No newline at end of file
diff --git a/Engine/Plugins/Experimental/Dataflow/Source/DataflowEditor/Private/Dataflow/DataflowAssetEditUtils.h b/Engine/Plugins/Experimental/Dataflow/Source/DataflowEditor/Private/Dataflow/DataflowAssetEditUtils.h
new file mode 100644
--- /dev/null
+++ b/Engine/Plugins/Experimental/Dataflow/Source/DataflowEditor/Private/Dataflow/DataflowAssetEditUtils.h
@@ -0,0 +1,65 @@
+// Copyright Epic Games, Inc. All Rights Reserved.
+
+#pragma once
+
+#include "UObject/NameTypes.h"
+
+///////////////////////////////////////////////////////////
+
+class UDataflow;
+struct FEdGraphPinType;
+
+struct FInstancedPropertyBag;
+struct FPropertyBagPropertyDesc;
+
+/**
+* Set to functions to modify a dataflow asset programmatically 
+* all changes are wrapped inside a scoped transaction as well and propertly marked the asset as modified if needed
+*/
+namespace UE::Dataflow::EditAssetUtils
+{
+	static const FName DefaultNewVariableBaseName = "NewVariable";
+
+	/** 
+	* Add a new Dataflow asset variable
+	* If BaseName already exists, use it to generate a unique name by adding an increment to it in the form of BaseName_## 
+	*/
+	static FName AddNewVariable(UDataflow* DataflowAsset, FName BaseName = DefaultNewVariableBaseName);
+
+	/** 
+	* Remove a Dataflow asset variable
+	*/
+	static void DeleteVariable(UDataflow* DataflowAsset, FName VariableName);
+
+	/**
+	* Rename a Dataflow asset variable
+	*/
+	static void RenameVariable(UDataflow* DataflowAsset, FName OldVariableName, FName NewVariableName);
+
+	/**
+	* Duplicate a Dataflow asset variable
+	* the new variable name will be generate uniquely based on the name of the original variable
+	*/
+	static FName DuplicateVariable(UDataflow* DataflowAsset, FName VariableName);
+
+	/**
+	* Set the type of a Dataflow asset variable
+	*/
+	static void SetVariableType(UDataflow* DataflowAsset, FName VariableName, const FEdGraphPinType& PinType);
+
+	/*
+	* Set value of a Dataflow asset variable from an property in a propertybag 
+	* Name and type must match for the operation to be successful
+	*/
+	void SetVariableValue(UDataflow* DataflowAsset, FName VariableName, const FInstancedPropertyBag& SourceBag);
+
+	/**
+	* Copy a variable to clipboard
+	*/
+	static void CopyVariableToClipboard(UDataflow* DataflowAsset, FName VariableName);
+
+	/**
+	* Paste a variable from clipboard and return its name
+	*/
+	static FName PasteVariableFromClipboard(UDataflow* DataflowAsset);
+}
