Dataflow editor : add new variable command ( preparation for adding variable UI )

diff --git a/Engine/Plugins/Experimental/Dataflow/Source/DataflowEditor/Private/Dataflow/DataflowAssetEditUtils.h b/Engine/Plugins/Experimental/Dataflow/Source/DataflowEditor/Private/Dataflow/DataflowAssetEditUtils.h
--- a/Engine/Plugins/Experimental/Dataflow/Source/DataflowEditor/Private/Dataflow/DataflowAssetEditUtils.h
+++ b/Engine/Plugins/Experimental/Dataflow/Source/DataflowEditor/Private/Dataflow/DataflowAssetEditUtils.h
@@ -24,28 +24,28 @@ namespace UE::Dataflow::EditAssetUtils
 	* Add a new Dataflow asset variable
 	* If BaseName already exists, use it to generate a unique name by adding an increment to it in the form of BaseName_## 
 	*/
-	static FName AddNewVariable(UDataflow* DataflowAsset, FName BaseName = DefaultNewVariableBaseName);
+	FName AddNewVariable(UDataflow* DataflowAsset, FName BaseName = DefaultNewVariableBaseName);
 
 	/** 
 	* Remove a Dataflow asset variable
 	*/
-	static void DeleteVariable(UDataflow* DataflowAsset, FName VariableName);
+	void DeleteVariable(UDataflow* DataflowAsset, FName VariableName);
 
 	/**
 	* Rename a Dataflow asset variable
 	*/
-	static void RenameVariable(UDataflow* DataflowAsset, FName OldVariableName, FName NewVariableName);
+	void RenameVariable(UDataflow* DataflowAsset, FName OldVariableName, FName NewVariableName);
 
 	/**
 	* Duplicate a Dataflow asset variable
 	* the new variable name will be generate uniquely based on the name of the original variable
 	*/
-	static FName DuplicateVariable(UDataflow* DataflowAsset, FName VariableName);
+	FName DuplicateVariable(UDataflow* DataflowAsset, FName VariableName);
 
 	/**
 	* Set the type of a Dataflow asset variable
 	*/
-	static void SetVariableType(UDataflow* DataflowAsset, FName VariableName, const FEdGraphPinType& PinType);
+	void SetVariableType(UDataflow* DataflowAsset, FName VariableName, const FEdGraphPinType& PinType);
 
 	/*
 	* Set value of a Dataflow asset variable from an property in a propertybag 
@@ -56,10 +56,10 @@ namespace UE::Dataflow::EditAssetUtils
 	/**
 	* Copy a variable to clipboard
 	*/
-	static void CopyVariableToClipboard(UDataflow* DataflowAsset, FName VariableName);
+	void CopyVariableToClipboard(UDataflow* DataflowAsset, FName VariableName);
 
 	/**
 	* Paste a variable from clipboard and return its name
 	*/
-	static FName PasteVariableFromClipboard(UDataflow* DataflowAsset);
+	FName PasteVariableFromClipboard(UDataflow* DataflowAsset);
 }
diff --git a/Engine/Plugins/Experimental/Dataflow/Source/DataflowEditor/Private/Dataflow/DataflowEditorCommands.cpp b/Engine/Plugins/Experimental/Dataflow/Source/DataflowEditor/Private/Dataflow/DataflowEditorCommands.cpp
--- a/Engine/Plugins/Experimental/Dataflow/Source/DataflowEditor/Private/Dataflow/DataflowEditorCommands.cpp
+++ b/Engine/Plugins/Experimental/Dataflow/Source/DataflowEditor/Private/Dataflow/DataflowEditorCommands.cpp
@@ -74,6 +74,8 @@ void FDataflowEditorCommandsImpl::RegisterCommands()
 	UI_COMMAND(StartSimulationScene, "StartSimulationScene", "Start the simulation scene", EUserInterfaceActionType::ToggleButton, FInputChord());
 	UI_COMMAND(StepSimulationScene, "StepSimulationScene", "Step the simulation scene", EUserInterfaceActionType::ToggleButton, FInputChord());
 
+	UI_COMMAND(AddNewVariable, "Variable", "Adds a new variable to this dataflow graph.", EUserInterfaceActionType::Button, FInputChord());
+
 	for (const TPair<FName, TUniquePtr<UE::Dataflow::IDataflowConstructionViewMode>>& NameAndMode : UE::Dataflow::FRenderingViewModeFactory::GetInstance().GetViewModes())
 	{
 		TSharedPtr< FUICommandInfo > SetViewModeCommand;
diff --git a/Engine/Plugins/Experimental/Dataflow/Source/DataflowEditor/Private/Dataflow/DataflowGraphEditor.cpp b/Engine/Plugins/Experimental/Dataflow/Source/DataflowEditor/Private/Dataflow/DataflowGraphEditor.cpp
--- a/Engine/Plugins/Experimental/Dataflow/Source/DataflowEditor/Private/Dataflow/DataflowGraphEditor.cpp
+++ b/Engine/Plugins/Experimental/Dataflow/Source/DataflowEditor/Private/Dataflow/DataflowGraphEditor.cpp
@@ -3,6 +3,7 @@
 #include "Dataflow/DataflowGraphEditor.h"
 
 #include "BoneDragDropOp.h"
+#include "Dataflow/DataflowAssetEditUtils.h"
 #include "Dataflow/DataflowEditor.h"
 #include "Dataflow/DataflowEditorToolkit.h"
 #include "Dataflow/DataflowEngine.h"
@@ -147,6 +148,10 @@ void SDataflowGraphEditor::Construct(const FArguments& InArgs, UObject* InAssetO
 				FExecuteAction::CreateSP(this, &SDataflowGraphEditor::RenameNode),
 				FCanExecuteAction::CreateSP(this, &SDataflowGraphEditor::CanRenameNode)
 			);
+			GraphEditorCommands->MapAction(
+				FDataflowEditorCommands::Get().AddNewVariable,
+				FExecuteAction::CreateSP(this, &SDataflowGraphEditor::AddNewVariable)
+			);
 		}
 	}
 
@@ -291,6 +296,11 @@ bool SDataflowGraphEditor::CanRenameNode() const
 	return false;
 }
 
+void SDataflowGraphEditor::AddNewVariable() const
+{
+	UE::Dataflow::EditAssetUtils::AddNewVariable(DataflowAsset.Get());
+}
+
 void SDataflowGraphEditor::OnSelectedNodesChanged(const TSet<UObject*>& NewSelection)
 {
 	// Set the currently selected graph editor before running any callback
diff --git a/Engine/Plugins/Experimental/Dataflow/Source/DataflowEditor/Public/Dataflow/DataflowEditorCommands.h b/Engine/Plugins/Experimental/Dataflow/Source/DataflowEditor/Public/Dataflow/DataflowEditorCommands.h
--- a/Engine/Plugins/Experimental/Dataflow/Source/DataflowEditor/Public/Dataflow/DataflowEditorCommands.h
+++ b/Engine/Plugins/Experimental/Dataflow/Source/DataflowEditor/Public/Dataflow/DataflowEditorCommands.h
@@ -65,6 +65,7 @@ public:
 	TSharedPtr<FUICommandInfo> AddOptionPin;
 	TSharedPtr<FUICommandInfo> RemoveOptionPin;
 	TSharedPtr<FUICommandInfo> ZoomToFitGraph;
+	TSharedPtr<FUICommandInfo> AddNewVariable;
 
 	TMap<FName, TSharedPtr<FUICommandInfo>> SetConstructionViewModeCommands;
 
diff --git a/Engine/Plugins/Experimental/Dataflow/Source/DataflowEditor/Public/Dataflow/DataflowGraphEditor.h b/Engine/Plugins/Experimental/Dataflow/Source/DataflowEditor/Public/Dataflow/DataflowGraphEditor.h
--- a/Engine/Plugins/Experimental/Dataflow/Source/DataflowEditor/Public/Dataflow/DataflowGraphEditor.h
+++ b/Engine/Plugins/Experimental/Dataflow/Source/DataflowEditor/Public/Dataflow/DataflowGraphEditor.h
@@ -139,6 +139,9 @@ public:
 	void RenameNode();
 	bool CanRenameNode() const;
 
+	/** Add a new variable for this dataflow graph */
+	void AddNewVariable() const;
+
 	SGraphEditor* GetGraphEditor() { return (SGraphEditor*)this; }
 
 	/** FGCObject interface */
@@ -148,6 +151,8 @@ public:
 	/** FDataflowSNodeInterface */
 	virtual TSharedPtr<UE::Dataflow::FContext> GetDataflowContext() const override;
 
+	const TSharedPtr<FUICommandList> GetCommands() const { return GraphEditorCommands; }
+
 	/** Return the currently selected editor. Only valid for the duration of the OnSelectedNodesChanged callback where the property editor is updated. */
 	static const TWeakPtr<SDataflowGraphEditor>& GetSelectedGraphEditor() { return SelectedGraphEditor; }
 
