Dataflow Editor: add a DataflowEditorOptions config class to have user settings persist across editor sesssions
- Enable persistent FOV settings via the config class (set default to 75)
- Also add an Options menu to the Construction viewport to set FOV
- Also save the Auto Exposure setting for both viewports

diff --git a/Engine/Plugins/Experimental/Dataflow/Source/DataflowEditor/Private/Dataflow/DataflowConstructionViewportClient.cpp b/Engine/Plugins/Experimental/Dataflow/Source/DataflowEditor/Private/Dataflow/DataflowConstructionViewportClient.cpp
--- a/Engine/Plugins/Experimental/Dataflow/Source/DataflowEditor/Private/Dataflow/DataflowConstructionViewportClient.cpp
+++ b/Engine/Plugins/Experimental/Dataflow/Source/DataflowEditor/Private/Dataflow/DataflowConstructionViewportClient.cpp
@@ -4,6 +4,7 @@
 #include "AssetEditorModeManager.h"
 #include "Dataflow/DataflowObject.h"
 #include "Dataflow/DataflowEditorMode.h"
+#include "Dataflow/DataflowEditorOptions.h"
 #include "Dataflow/DataflowEditorToolkit.h"
 #include "Dataflow/DataflowEditorCollectionComponent.h"
 #include "Dataflow/DataflowEngineSceneHitProxies.h"
@@ -36,6 +37,23 @@ FDataflowConstructionViewportClient::FDataflowConstructionViewportClient(FEditor
 
 	PreviewScene = static_cast<FDataflowPreviewSceneBase*>(InPreviewScene);
 	bEnableSceneTicking = bCouldTickScene;
+
+	if (const UDataflowEditorOptions* const Options = UDataflowEditorOptions::StaticClass()->GetDefaultObject<UDataflowEditorOptions>())
+	{
+		FOVAngle = Options->ConstructionViewFOV;
+		ViewFOV = FOVAngle;
+		ExposureSettings.bFixed = Options->bConstructionViewFixedExposure;
+	}
+}
+
+FDataflowConstructionViewportClient::~FDataflowConstructionViewportClient()
+{
+	if (UDataflowEditorOptions* const Options = UDataflowEditorOptions::StaticClass()->GetDefaultObject<UDataflowEditorOptions>())
+	{
+		Options->ConstructionViewFOV = FOVAngle;
+		Options->bConstructionViewFixedExposure = ExposureSettings.bFixed;
+		Options->SaveConfig();
+	}
 }
 
 void FDataflowConstructionViewportClient::SetDataflowEditorToolkit(TWeakPtr<FDataflowEditorToolkit> InDataflowEditorToolkitPtr)
diff --git a/Engine/Plugins/Experimental/Dataflow/Source/DataflowEditor/Private/Dataflow/DataflowConstructionViewportToolbar.cpp b/Engine/Plugins/Experimental/Dataflow/Source/DataflowEditor/Private/Dataflow/DataflowConstructionViewportToolbar.cpp
--- a/Engine/Plugins/Experimental/Dataflow/Source/DataflowEditor/Private/Dataflow/DataflowConstructionViewportToolbar.cpp
+++ b/Engine/Plugins/Experimental/Dataflow/Source/DataflowEditor/Private/Dataflow/DataflowConstructionViewportToolbar.cpp
@@ -12,6 +12,8 @@
 #include "SEditorViewportViewMenu.h"
 #include "ToolMenus.h"
 #include "Widgets/Layout/SBorder.h"
+#include "EditorViewportCommands.h"
+#include "ViewportToolbar/UnrealEdViewportToolbar.h"
 
 #define LOCTEXT_NAMESPACE "SDataflowConstructionViewportSelectionToolBar"
 
@@ -71,6 +73,18 @@ void SDataflowConstructionViewportSelectionToolBar::Construct(const FArguments&
 		]
 	];
 
+	// Options menu
+	MainBoxPtr->AddSlot()
+		.AutoWidth()
+		.Padding(ToolbarSlotPadding)
+		[
+			SNew(SEditorViewportToolbarMenu)
+				.ParentToolBar(SharedThis(this))
+				.Cursor(EMouseCursor::Default)
+				.Image("EditorViewportToolBar.OptionsDropdown")
+				.OnGetMenuContent(this, &SDataflowConstructionViewportSelectionToolBar::GenerateOptionsMenu)
+		];
+
 	// Display (Lit, Unlit, Wireframe, etc.)
 	MainBoxPtr->AddSlot()
 		.AutoWidth()
@@ -122,6 +136,36 @@ TSharedRef<SWidget> SDataflowConstructionViewportSelectionToolBar::GenerateShowM
 	return MenuBuilder.MakeWidget();
 }
 
+// Mostly copied from the inaccessible SCommonEditorViewportToolbarBase::GenerateOptionsMenu()
+TSharedRef<SWidget> SDataflowConstructionViewportSelectionToolBar::GenerateOptionsMenu() const
+{
+	GetInfoProvider().OnFloatingButtonClicked();
+	TSharedRef<SEditorViewport> ViewportRef = GetInfoProvider().GetViewportWidget();
+
+	const bool bIsPerspective = GetViewportClient().GetViewportType() == LVT_Perspective;
+
+	constexpr bool bInShouldCloseWindowAfterMenuSelection = true;
+	FMenuBuilder OptionsMenuBuilder(bInShouldCloseWindowAfterMenuSelection, ViewportRef->GetCommandList());
+	{
+		OptionsMenuBuilder.BeginSection("LevelViewportViewportOptions", LOCTEXT("OptionsMenuHeader", "Viewport Options"));
+		{
+			OptionsMenuBuilder.AddMenuEntry(FEditorViewportCommands::Get().ToggleRealTime);
+			OptionsMenuBuilder.AddMenuEntry(FEditorViewportCommands::Get().ToggleStats);
+			OptionsMenuBuilder.AddMenuEntry(FEditorViewportCommands::Get().ToggleFPS);
+
+			if (bIsPerspective)
+			{
+				OptionsMenuBuilder.AddWidget(GenerateFOVMenu(), LOCTEXT("FOVAngle", "Field of View (H)"));
+			}
+		}
+		OptionsMenuBuilder.EndSection();
+
+		ExtendOptionsMenu(OptionsMenuBuilder);
+	}
+
+	return OptionsMenuBuilder.MakeWidget();
+}
+
 TSharedRef<SWidget> SDataflowConstructionViewportSelectionToolBar::MakeToolBar(const TSharedPtr<FExtender> InExtenders)
 {
 	// The following is modeled after portions of STransformViewportToolBar, which gets 
diff --git a/Engine/Plugins/Experimental/Dataflow/Source/DataflowEditor/Private/Dataflow/DataflowEditorOptions.cpp b/Engine/Plugins/Experimental/Dataflow/Source/DataflowEditor/Private/Dataflow/DataflowEditorOptions.cpp
new file mode 100644
--- /dev/null
+++ b/Engine/Plugins/Experimental/Dataflow/Source/DataflowEditor/Private/Dataflow/DataflowEditorOptions.cpp
@@ -0,0 +1,17 @@
+// Copyright Epic Games, Inc. All Rights Reserved.
+
+#include "Dataflow/DataflowEditorOptions.h"
+
+#include UE_INLINE_GENERATED_CPP_BY_NAME(DataflowEditorOptions)
+
+UDataflowEditorOptions::UDataflowEditorOptions(class FObjectInitializer const& ObjectInitializer)
+{
+	ConstructionViewFOV = 75.0f;
+	SimulationViewFOV = 75.0f;
+	
+	// The construction viewport often becomes overexposed in 2D with auto exposure turned on, so we default to fixed.
+	// In the future we will add the ability to tweak the construction viewport profile and materials (UE-231612)
+	bConstructionViewFixedExposure = true;
+
+	bSimulationViewFixedExposure = false;
+}
diff --git a/Engine/Plugins/Experimental/Dataflow/Source/DataflowEditor/Private/Dataflow/DataflowEditorOptions.h b/Engine/Plugins/Experimental/Dataflow/Source/DataflowEditor/Private/Dataflow/DataflowEditorOptions.h
new file mode 100644
--- /dev/null
+++ b/Engine/Plugins/Experimental/Dataflow/Source/DataflowEditor/Private/Dataflow/DataflowEditorOptions.h
@@ -0,0 +1,30 @@
+// Copyright Epic Games, Inc. All Rights Reserved.
+
+#pragma once
+
+#include "UObject/Object.h"
+#include "DataflowEditorOptions.generated.h"
+
+/** A settings class used to store and retreive user settings related to the Dataflow Editor */
+UCLASS(hidecategories = Object, config = EditorPerProjectUserSettings, MinimalAPI)
+class UDataflowEditorOptions : public UObject
+{
+	GENERATED_UCLASS_BODY()
+
+	/** FOV for the Construction viewport camera */
+	UPROPERTY(config)
+	float ConstructionViewFOV;
+
+	/** FOV for the Simulation viewport camera */
+	UPROPERTY(config)
+	float SimulationViewFOV;
+
+	/** Whether the Construction viewport has Auto or Fixed exposure */
+	UPROPERTY(config)
+	bool bConstructionViewFixedExposure;
+
+	/** Whether the Simulation viewport has Auto or Fixed exposure */
+	UPROPERTY(config)
+	bool bSimulationViewFixedExposure;
+
+};
diff --git a/Engine/Plugins/Experimental/Dataflow/Source/DataflowEditor/Private/Dataflow/DataflowEditorToolkit.cpp b/Engine/Plugins/Experimental/Dataflow/Source/DataflowEditor/Private/Dataflow/DataflowEditorToolkit.cpp
--- a/Engine/Plugins/Experimental/Dataflow/Source/DataflowEditor/Private/Dataflow/DataflowEditorToolkit.cpp
+++ b/Engine/Plugins/Experimental/Dataflow/Source/DataflowEditor/Private/Dataflow/DataflowEditorToolkit.cpp
@@ -435,18 +435,12 @@ void FDataflowEditorToolkit::PostInitAssetEditor()
 		// the component bounds are behind the displaced geometry).
 		Client->EngineShowFlags.SetDisableOcclusionQueries(true);
 
-		// Default FOV of 90 degrees causes a fair bit of lens distortion, especially noticeable with smaller viewports
-		Client->ViewFOV = 45.0;
-
 		// Ortho has too many problems with rendering things, unfortunately, so we should use perspective.
 		Client->SetViewportType(ELevelViewportType::LVT_Perspective);
 
 		// Lit gives us the most options in terms of the materials we can use.
 		Client->SetViewMode(EViewModeIndex::VMI_Lit);
 
-		// If exposure isn't set to fixed, it will flash as we stare into the void
-		Client->ExposureSettings.bFixed = true;
-
 		// We need the viewport client to start out focused, or else it won't get ticked until
 		// we click inside it.
 		if(Client->Viewport)
diff --git a/Engine/Plugins/Experimental/Dataflow/Source/DataflowEditor/Private/Dataflow/DataflowSimulationViewportClient.cpp b/Engine/Plugins/Experimental/Dataflow/Source/DataflowEditor/Private/Dataflow/DataflowSimulationViewportClient.cpp
--- a/Engine/Plugins/Experimental/Dataflow/Source/DataflowEditor/Private/Dataflow/DataflowSimulationViewportClient.cpp
+++ b/Engine/Plugins/Experimental/Dataflow/Source/DataflowEditor/Private/Dataflow/DataflowSimulationViewportClient.cpp
@@ -4,6 +4,7 @@
 #include "AssetEditorModeManager.h"
 #include "Dataflow/DataflowObject.h"
 #include "Dataflow/DataflowEditorMode.h"
+#include "Dataflow/DataflowEditorOptions.h"
 #include "Dataflow/DataflowEditorToolkit.h"
 #include "Dataflow/DataflowEditorCollectionComponent.h"
 #include "Dataflow/DataflowEngineSceneHitProxies.h"
@@ -31,6 +32,23 @@ FDataflowSimulationViewportClient::FDataflowSimulationViewportClient(FEditorMode
 	EngineShowFlags.EnableAdvancedFeatures();
 
 	bEnableSceneTicking = bCouldTickScene;
+
+	if (UDataflowEditorOptions* const Options = UDataflowEditorOptions::StaticClass()->GetDefaultObject<UDataflowEditorOptions>())
+	{
+		FOVAngle = Options->SimulationViewFOV;
+		ViewFOV = FOVAngle;
+		ExposureSettings.bFixed = Options->bSimulationViewFixedExposure;
+	}
+}
+
+FDataflowSimulationViewportClient::~FDataflowSimulationViewportClient()
+{
+	if (UDataflowEditorOptions* const Options = UDataflowEditorOptions::StaticClass()->GetDefaultObject<UDataflowEditorOptions>())
+	{
+		Options->SimulationViewFOV = FOVAngle;
+		Options->bSimulationViewFixedExposure = ExposureSettings.bFixed;
+		Options->SaveConfig();
+	}
 }
 
 void FDataflowSimulationViewportClient::SetDataflowEditorToolkit(TWeakPtr<FDataflowEditorToolkit> InDataflowEditorToolkitPtr)
diff --git a/Engine/Plugins/Experimental/Dataflow/Source/DataflowEditor/Public/Dataflow/DataflowConstructionViewportClient.h b/Engine/Plugins/Experimental/Dataflow/Source/DataflowEditor/Public/Dataflow/DataflowConstructionViewportClient.h
--- a/Engine/Plugins/Experimental/Dataflow/Source/DataflowEditor/Public/Dataflow/DataflowConstructionViewportClient.h
+++ b/Engine/Plugins/Experimental/Dataflow/Source/DataflowEditor/Public/Dataflow/DataflowConstructionViewportClient.h
@@ -32,6 +32,8 @@ public:
 	FDataflowConstructionViewportClient(FEditorModeTools* InModeTools, FPreviewScene* InPreviewScene,  const bool bCouldTickScene,
 								  const TWeakPtr<SEditorViewport> InEditorViewportWidget = nullptr);
 
+	virtual ~FDataflowConstructionViewportClient();
+
 	void SetConstructionViewMode(const UE::Dataflow::IDataflowConstructionViewMode* InViewMode);
 
 	// IInputBehaviorSource
diff --git a/Engine/Plugins/Experimental/Dataflow/Source/DataflowEditor/Public/Dataflow/DataflowConstructionViewportToolbar.h b/Engine/Plugins/Experimental/Dataflow/Source/DataflowEditor/Public/Dataflow/DataflowConstructionViewportToolbar.h
--- a/Engine/Plugins/Experimental/Dataflow/Source/DataflowEditor/Public/Dataflow/DataflowConstructionViewportToolbar.h
+++ b/Engine/Plugins/Experimental/Dataflow/Source/DataflowEditor/Public/Dataflow/DataflowConstructionViewportToolbar.h
@@ -26,6 +26,8 @@ private:
 	// SCommonEditorViewportToolbarBase
 	virtual TSharedRef<SWidget> GenerateShowMenu() const override;
 
+	TSharedRef<SWidget> GenerateOptionsMenu() const;
+
 	TSharedRef<SWidget> MakeToolBar(const TSharedPtr<FExtender> InExtenders);
 
 	// View mode
diff --git a/Engine/Plugins/Experimental/Dataflow/Source/DataflowEditor/Public/Dataflow/DataflowSimulationViewportClient.h b/Engine/Plugins/Experimental/Dataflow/Source/DataflowEditor/Public/Dataflow/DataflowSimulationViewportClient.h
--- a/Engine/Plugins/Experimental/Dataflow/Source/DataflowEditor/Public/Dataflow/DataflowSimulationViewportClient.h
+++ b/Engine/Plugins/Experimental/Dataflow/Source/DataflowEditor/Public/Dataflow/DataflowSimulationViewportClient.h
@@ -25,7 +25,7 @@ public:
 	FDataflowSimulationViewportClient(FEditorModeTools* InModeTools, FPreviewScene* InPreviewScene,  const bool bCouldTickScene,
 								  const TWeakPtr<SEditorViewport> InEditorViewportWidget = nullptr);
 
-	virtual ~FDataflowSimulationViewportClient() = default;
+	virtual ~FDataflowSimulationViewportClient();
 
 	/** Set the data flow toolkit used to create the client*/
 	void SetDataflowEditorToolkit(TWeakPtr<FDataflowEditorToolkit> DataflowToolkit);
