Dataflow : Update the node render option box to be a Pin and fix the issue when it was scale
- widget is now a right side title widget instead of an overlay because somehow overlay widget have issues with mouse coordinates when scaled
- change the toggle widget for a button toggle that look like a pin
- make sure the construction view updates properly when using the toggle

diff --git a/Engine/Plugins/Experimental/Dataflow/Source/DataflowEditor/Private/Dataflow/DataflowGraphEditor.cpp b/Engine/Plugins/Experimental/Dataflow/Source/DataflowEditor/Private/Dataflow/DataflowGraphEditor.cpp
--- a/Engine/Plugins/Experimental/Dataflow/Source/DataflowEditor/Private/Dataflow/DataflowGraphEditor.cpp
+++ b/Engine/Plugins/Experimental/Dataflow/Source/DataflowEditor/Private/Dataflow/DataflowGraphEditor.cpp
@@ -182,6 +182,13 @@ TSharedPtr<UE::Dataflow::FContext> SDataflowGraphEditor::GetDataflowContext() co
 	return TSharedPtr<UE::Dataflow::FContext>();
 }
 
+void SDataflowGraphEditor::OnRenderToggleChanged() const
+{
+	// need to refresh the UI by refreshing the selection
+	FGraphPanelSelectionSet SelectionSet = GetSelectedNodes();
+	GetGraphPanel()->SelectionManager.SetSelectionSet(SelectionSet);
+}
+
 void SDataflowGraphEditor::EvaluateNode()
 {
 	UE_LOG(LogChaosDataflow, VeryVerbose, TEXT("SDataflowGraphEditor::EvaluateNode(): Nodes [%s]"),
diff --git a/Engine/Plugins/Experimental/Dataflow/Source/DataflowEditor/Private/Dataflow/DataflowSNode.cpp b/Engine/Plugins/Experimental/Dataflow/Source/DataflowEditor/Private/Dataflow/DataflowSNode.cpp
--- a/Engine/Plugins/Experimental/Dataflow/Source/DataflowEditor/Private/Dataflow/DataflowSNode.cpp
+++ b/Engine/Plugins/Experimental/Dataflow/Source/DataflowEditor/Private/Dataflow/DataflowSNode.cpp
@@ -52,6 +52,18 @@ void SDataflowOutputPin::Construct(const FArguments& InArgs, UEdGraphPin* InPin)
 // SDataflowEdNode
 //
 
+const FSlateBrush* SDataflowEdNode::GetPinButtonImage() const
+{
+	if (DataflowGraphNode && DataflowGraphNode->ShouldWireframeRenderNode())
+	{
+		return FAppStyle::Get().GetBrush("Icons.Pinned");
+	}
+	else
+	{
+		return FAppStyle::Get().GetBrush("Icons.Unpinned");
+	}
+}
+
 void SDataflowEdNode::Construct(const FArguments& InArgs, UDataflowEdNode* InNode)
 {
 	GraphNode = InNode;
@@ -60,50 +72,6 @@ void SDataflowEdNode::Construct(const FArguments& InArgs, UDataflowEdNode* InNod
 
 	UpdateGraphNode();
 
-	
-	const FSlateBrush* DisabledSwitchBrush = FDataflowEditorStyle::Get().GetBrush(TEXT("Dataflow.Render.Disabled"));
-	const FSlateBrush* EnabledSwitchBrush = FDataflowEditorStyle::Get().GetBrush(TEXT("Dataflow.Render.Enabled"));
-
-	//
-	// Render
-	//
-	CheckBoxStyle = FCheckBoxStyle()
-		.SetCheckBoxType(ESlateCheckBoxType::CheckBox)
-		.SetUncheckedImage(*DisabledSwitchBrush)
-		.SetUncheckedHoveredImage(*DisabledSwitchBrush)
-		.SetUncheckedPressedImage(*DisabledSwitchBrush)
-		.SetCheckedImage(*EnabledSwitchBrush)
-		.SetCheckedHoveredImage(*EnabledSwitchBrush)
-		.SetCheckedPressedImage(*EnabledSwitchBrush)
-		.SetPadding(FMargin(0, 0, 0, 1));
-	
-	RenderCheckBoxWidget = SNew(SCheckBox)
-		.Style(&CheckBoxStyle)
-		.IsChecked_Lambda([this]()-> ECheckBoxState
-			{
-				if (DataflowGraphNode && DataflowGraphNode->ShouldWireframeRenderNode())
-				{
-					return ECheckBoxState::Checked;
-				}
-				return ECheckBoxState::Unchecked;
-			})
-		.OnCheckStateChanged_Lambda([this](const ECheckBoxState NewState)
-			{
-				if (DataflowGraphNode)
-				{
-					const bool bShouldWireframeRenderNode = (NewState == ECheckBoxState::Checked);
-					DataflowGraphNode->SetShouldWireframeRenderNode(bShouldWireframeRenderNode);
-				}
-			})
-		.IsEnabled_Lambda([this]()->bool
-			{
-				if (DataflowGraphNode)
-				{
-					return DataflowGraphNode->CanEnableWireframeRenderNode();
-				}
-				return false;
-			});
-
 	//
 	// Freeze
 	//
@@ -194,19 +162,56 @@ PRAGMA_ENABLE_DEPRECATION_WARNINGS
 
 TSharedRef<SWidget> SDataflowEdNode::CreateTitleRightWidget()
 {
-	// Reserve space for the overlay
-	float RightMargin = 0.f;
-	if (DataflowGraphNode && DataflowGraphNode->GetDataflowNode())
+	RenderCheckBoxWidget = SNew(SCheckBox)
+		.Style(FAppStyle::Get(), "ToggleButtonCheckbox")
+		.IsChecked_Lambda([this]()-> ECheckBoxState
+			{
+				if (DataflowGraphNode && DataflowGraphNode->ShouldWireframeRenderNode())
+				{
+					return ECheckBoxState::Checked;
+				}
+				return ECheckBoxState::Unchecked;
+			})
+		.OnCheckStateChanged_Lambda([this](const ECheckBoxState NewState)
+			{
+				if (DataflowGraphNode)
+				{
+					const bool bShouldWireframeRenderNode = (NewState == ECheckBoxState::Checked);
+					DataflowGraphNode->SetShouldWireframeRenderNode(bShouldWireframeRenderNode);
+					DataflowInterface->OnRenderToggleChanged();
+				}
+			})
+				.IsEnabled_Lambda([this]()->bool
+					{
+						if (DataflowGraphNode)
+						{
+							return DataflowGraphNode->CanEnableWireframeRenderNode();
+						}
+						return false;
+					})
+				.Padding(2.0f)
+				.HAlign(HAlign_Center)
+				[
+					SNew(SImage)
+						.ColorAndOpacity(FSlateColor::UseForeground())
+						.Image(this, &SDataflowEdNode::GetPinButtonImage)
+				];
+
+	const bool bHasRenderingCheckBox = (
+		DataflowInterface->NodesHaveToggleWidget() && 
+		DataflowGraphNode->GetDataflowNode()->GetRenderParameters().Num()
+		);
+
+	if (bHasRenderingCheckBox)
 	{
-		if (DataflowInterface->NodesHaveToggleWidget() && DataflowGraphNode->GetDataflowNode()->GetRenderParameters().Num())
-		{
-			RightMargin += 20.f;
-		}
+		return SNew(SBox)
+			.HAlign(HAlign_Fill)
+			.Padding(0.0f, 2.0f, 0.0f, 2.0f)
+			[
+				RenderCheckBoxWidget.ToSharedRef()
+			];
 	}
-	return SNew(SHorizontalBox)
-		+ SHorizontalBox::Slot()
-		.HAlign(HAlign_Fill)
-		.Padding(FMargin(0.f, 0.f, RightMargin, 0.f));
+	return SNullWidget::NullWidget;
 }
 
 TArray<FOverlayWidgetInfo> SDataflowEdNode::GetOverlayWidgets(bool bSelected, const FVector2D& WidgetSize) const
@@ -215,17 +220,6 @@ TArray<FOverlayWidgetInfo> SDataflowEdNode::GetOverlayWidgets(bool bSelected, co
 
 	if (DataflowGraphNode && DataflowGraphNode->GetDataflowNode())
 	{
-		float ImageOffsetX = 0.f;
-		if (DataflowInterface->NodesHaveToggleWidget() && DataflowGraphNode->GetDataflowNode()->GetRenderParameters().Num())
-		{
-			const FVector2D RenderBoxImageSize = RenderCheckBoxWidget->GetDesiredSize();
-			ImageOffsetX += RenderBoxImageSize.X;
-			FOverlayWidgetInfo RenderBoxInfo;
-			RenderBoxInfo.OverlayOffset = FVector2D(WidgetSize.X - ImageOffsetX - 6.f, 6.f);
-			RenderBoxInfo.Widget = RenderCheckBoxWidget;
-			Widgets.Add(RenderBoxInfo);
-		}
-
 		if (DataflowInterface->NodesHaveFreezeWidget())
 		{
 			static const FVector2D FreezeImageOverhang(10.f);  // The icon must slightly overhang to make space for the lower output pins
diff --git a/Engine/Plugins/Experimental/Dataflow/Source/DataflowEditor/Public/Dataflow/DataflowGraphEditor.h b/Engine/Plugins/Experimental/Dataflow/Source/DataflowEditor/Public/Dataflow/DataflowGraphEditor.h
--- a/Engine/Plugins/Experimental/Dataflow/Source/DataflowEditor/Public/Dataflow/DataflowGraphEditor.h
+++ b/Engine/Plugins/Experimental/Dataflow/Source/DataflowEditor/Public/Dataflow/DataflowGraphEditor.h
@@ -150,6 +150,7 @@ public:
 
 	/** FDataflowSNodeInterface */
 	virtual TSharedPtr<UE::Dataflow::FContext> GetDataflowContext() const override;
+	virtual void OnRenderToggleChanged() const override;
 
 	const TSharedPtr<FUICommandList> GetCommands() const { return GraphEditorCommands; }
 
diff --git a/Engine/Plugins/Experimental/Dataflow/Source/DataflowEditor/Public/Dataflow/DataflowSEditorInterface.h b/Engine/Plugins/Experimental/Dataflow/Source/DataflowEditor/Public/Dataflow/DataflowSEditorInterface.h
--- a/Engine/Plugins/Experimental/Dataflow/Source/DataflowEditor/Public/Dataflow/DataflowSEditorInterface.h
+++ b/Engine/Plugins/Experimental/Dataflow/Source/DataflowEditor/Public/Dataflow/DataflowSEditorInterface.h
@@ -29,6 +29,10 @@ public:
 		return true;
 	}
 
+	virtual void OnRenderToggleChanged() const
+	{
+	}
+
 protected:
 	FDataflowSEditorInterface() = default;
 };
diff --git a/Engine/Plugins/Experimental/Dataflow/Source/DataflowEditor/Public/Dataflow/DataflowSNode.h b/Engine/Plugins/Experimental/Dataflow/Source/DataflowEditor/Public/Dataflow/DataflowSNode.h
--- a/Engine/Plugins/Experimental/Dataflow/Source/DataflowEditor/Public/Dataflow/DataflowSNode.h
+++ b/Engine/Plugins/Experimental/Dataflow/Source/DataflowEditor/Public/Dataflow/DataflowSNode.h
@@ -79,6 +79,8 @@ private:
 	/** Checks whether Add pin button should currently be visible */
 	virtual EVisibility IsAddPinButtonVisible() const;
 
+	const FSlateBrush* GetPinButtonImage() const;
+
 	virtual TSharedPtr<SGraphPin> CreatePinWidget(UEdGraphPin* Pin) const override;
 	virtual TSharedRef<SWidget> CreateTitleRightWidget() override;
 
