Fix Dataflow crash when opening a graph with invalid nodes.

diff --git a/Engine/Plugins/Experimental/Dataflow/Source/DataflowEditor/Private/Dataflow/DataflowSNode.cpp b/Engine/Plugins/Experimental/Dataflow/Source/DataflowEditor/Private/Dataflow/DataflowSNode.cpp
--- a/Engine/Plugins/Experimental/Dataflow/Source/DataflowEditor/Private/Dataflow/DataflowSNode.cpp
+++ b/Engine/Plugins/Experimental/Dataflow/Source/DataflowEditor/Private/Dataflow/DataflowSNode.cpp
@@ -199,6 +199,7 @@ TSharedRef<SWidget> SDataflowEdNode::CreateTitleRightWidget()
 
 	const bool bHasRenderingCheckBox = (
 		DataflowInterface->NodesHaveToggleWidget() && 
+		DataflowGraphNode->GetDataflowNode() &&
 		DataflowGraphNode->GetDataflowNode()->GetRenderParameters().Num()
 		);
 
