ChaosFlesh: fixed async caching in dataflow simulation

diff --git a/Engine/Plugins/Experimental/Dataflow/Source/DataflowEditor/Private/Dataflow/DataflowSimulationControls.cpp b/Engine/Plugins/Experimental/Dataflow/Source/DataflowEditor/Private/Dataflow/DataflowSimulationControls.cpp
--- a/Engine/Plugins/Experimental/Dataflow/Source/DataflowEditor/Private/Dataflow/DataflowSimulationControls.cpp
+++ b/Engine/Plugins/Experimental/Dataflow/Source/DataflowEditor/Private/Dataflow/DataflowSimulationControls.cpp
@@ -163,7 +163,7 @@ namespace UE::Dataflow
 
 			FBoneContainer BoneContainer;
 			BoneContainer.SetUseRAWData(true);
-			BoneContainer.InitializeTo(BoneIndices, UE::Anim::FCurveFilterSettings(), *InSkeletalMesh);
+			BoneContainer.InitializeTo(BoneIndices, UE::Anim::FCurveFilterSettings(), *InSkeletalMesh->GetSkeleton());
 		
 			FCompactPose CompactPose;
 			CompactPose.SetBoneContainer(&BoneContainer);
