Dataflow editor :  prevent potential crash ( repro unclear )

diff --git a/Engine/Plugins/Experimental/Dataflow/Source/DataflowEditor/Private/Dataflow/DataflowEditorMode.cpp b/Engine/Plugins/Experimental/Dataflow/Source/DataflowEditor/Private/Dataflow/DataflowEditorMode.cpp
--- a/Engine/Plugins/Experimental/Dataflow/Source/DataflowEditor/Private/Dataflow/DataflowEditorMode.cpp
+++ b/Engine/Plugins/Experimental/Dataflow/Source/DataflowEditor/Private/Dataflow/DataflowEditorMode.cpp
@@ -252,8 +252,12 @@ void UDataflowEditorMode::AddNode(FName NewNodeType)
 bool UDataflowEditorMode::CanAddNode(FName NewNodeType) const
 {
 	const UEdGraphNode* const CurrentlySelectedNode = GetSingleSelectedNodeWithOutputType(FManagedArrayCollection::StaticType());
-	const bool bToolActive = GetToolManager()->HasActiveTool(EToolSide::Left);
-	return CurrentlySelectedNode && !bToolActive;
+	if (CurrentlySelectedNode)
+	{
+		const bool bToolActive = GetToolManager()->HasActiveTool(EToolSide::Left);
+		return !bToolActive;
+	}
+	return false;
 }
 
 void UDataflowEditorMode::RegisterTools()
