Dataflow asset : add event for when variables are updated

diff --git a/Engine/Plugins/Experimental/Dataflow/Source/DataflowEditor/Private/Dataflow/DataflowAssetEditUtils.cpp b/Engine/Plugins/Experimental/Dataflow/Source/DataflowEditor/Private/Dataflow/DataflowAssetEditUtils.cpp
--- a/Engine/Plugins/Experimental/Dataflow/Source/DataflowEditor/Private/Dataflow/DataflowAssetEditUtils.cpp
+++ b/Engine/Plugins/Experimental/Dataflow/Source/DataflowEditor/Private/Dataflow/DataflowAssetEditUtils.cpp
@@ -81,6 +81,7 @@ namespace UE::Dataflow::EditAssetUtils
 				if (const UPropertyBag* NewBagStruct = UPropertyBag::GetOrCreateFromDescs(NewPropertyDescs))
 				{
 					DataflowAsset.Variables.MigrateToNewBagStruct(NewBagStruct);
+					FDataflowAssetDelegates::OnVariablesChanged.Broadcast(&DataflowAsset, Variable);
 					return true;
 				}
 			}
@@ -121,6 +122,7 @@ namespace UE::Dataflow::EditAssetUtils
 			{
 				UniqueVariableName = Private::GenerateUniqueVariableName(DataflowAsset, BaseName);
 				DataflowAsset.Variables.AddProperty(UniqueVariableName, EPropertyBagPropertyType::Int32);
+				FDataflowAssetDelegates::OnVariablesChanged.Broadcast(&DataflowAsset, UniqueVariableName);
 				return Private::EChangeResult::Changed;
 			};
 
@@ -138,6 +140,7 @@ namespace UE::Dataflow::EditAssetUtils
 			[&VariableName](UDataflow& DataflowAsset) -> Private::EChangeResult
 			{
 				DataflowAsset.Variables.RemovePropertyByName(VariableName);
+				FDataflowAssetDelegates::OnVariablesChanged.Broadcast(&DataflowAsset, VariableName);
 				return Private::EChangeResult::Changed;
 			};
 
@@ -184,6 +187,7 @@ namespace UE::Dataflow::EditAssetUtils
 					NewDesc.Name = NewVariableName;
 					NewDesc.ID.Invalidate();
 					DataflowAsset.Variables.AddProperties(MakeConstArrayView(&NewDesc, 1));
+					FDataflowAssetDelegates::OnVariablesChanged.Broadcast(&DataflowAsset, NewVariableName);
 					return Private::EChangeResult::Changed;
 				}
 				return Private::EChangeResult::None;
@@ -232,6 +236,7 @@ namespace UE::Dataflow::EditAssetUtils
 						EPropertyBagResult Result = DataflowAsset.Variables.SetValue(VariableName, SourceDesc->CachedProperty, SourcePropertyAddr);
 						if (Result == EPropertyBagResult::Success)
 						{
+							FDataflowAssetDelegates::OnVariablesChanged.Broadcast(&DataflowAsset, VariableName);
 							return Private::EChangeResult::Changed;
 						}
 					}
@@ -295,6 +300,7 @@ namespace UE::Dataflow::EditAssetUtils
 					[&PropertyDesc](UDataflow& DataflowAsset) -> Private::EChangeResult
 					{
 						DataflowAsset.Variables.AddProperties(MakeConstArrayView(&PropertyDesc, 1));
+						FDataflowAssetDelegates::OnVariablesChanged.Broadcast(&DataflowAsset, PropertyDesc.Name);
 						return Private::EChangeResult::Changed;
 					};
 
diff --git a/Engine/Source/Runtime/Experimental/Dataflow/Engine/Private/Dataflow/DataflowObject.cpp b/Engine/Source/Runtime/Experimental/Dataflow/Engine/Private/Dataflow/DataflowObject.cpp
--- a/Engine/Source/Runtime/Experimental/Dataflow/Engine/Private/Dataflow/DataflowObject.cpp
+++ b/Engine/Source/Runtime/Experimental/Dataflow/Engine/Private/Dataflow/DataflowObject.cpp
@@ -12,6 +12,8 @@
 
 #define LOCTEXT_NAMESPACE "UDataflow"
 
+FDataflowAssetDelegates::FOnVariablesChanged FDataflowAssetDelegates::OnVariablesChanged;
+
 namespace UE::Dataflow::CVars
 {
 	/** Enable the simulation dataflow (for now WIP) */
diff --git a/Engine/Source/Runtime/Experimental/Dataflow/Engine/Public/Dataflow/DataflowObject.h b/Engine/Source/Runtime/Experimental/Dataflow/Engine/Public/Dataflow/DataflowObject.h
--- a/Engine/Source/Runtime/Experimental/Dataflow/Engine/Public/Dataflow/DataflowObject.h
+++ b/Engine/Source/Runtime/Experimental/Dataflow/Engine/Public/Dataflow/DataflowObject.h
@@ -16,6 +16,7 @@
 #include "UObject/UObjectGlobals.h"
 
 #include "StructUtils/PropertyBag.h"
+#include "Delegates/DelegateCombinations.h"
 
 #include "DataflowObject.generated.h"
 
@@ -26,6 +27,12 @@ class UDataflowEdNode;
 class UMaterial;
 namespace UE::Dataflow { class FGraph; }
 
+struct FDataflowAssetDelegates
+{
+	/** Callback for object property modifications, called by UObject::PostEditChangeProperty with a single property event */
+	DECLARE_MULTICAST_DELEGATE_TwoParams(FOnVariablesChanged, const UDataflow*, FName);
+	static DATAFLOWENGINE_API FOnVariablesChanged OnVariablesChanged;
+};
 
 /**
 *	FDataflowAssetEdit
@@ -117,7 +124,7 @@ public:
 	EDataflowType Type = EDataflowType::Construction;
 
 	// keep hidden for now until we have a proper UI for it 
-	UPROPERTY()
+	UPROPERTY(EditAnywhere, Category=Variables)
 	FInstancedPropertyBag Variables;
 
 public:
