Dataflow : fix crash when copy/pasting a comment that has a label that is the same as an existing node name

diff --git a/Engine/Plugins/Experimental/Dataflow/Source/DataflowEditor/Private/Dataflow/DataflowSCommentNode.cpp b/Engine/Plugins/Experimental/Dataflow/Source/DataflowEditor/Private/Dataflow/DataflowSCommentNode.cpp
--- a/Engine/Plugins/Experimental/Dataflow/Source/DataflowEditor/Private/Dataflow/DataflowSCommentNode.cpp
+++ b/Engine/Plugins/Experimental/Dataflow/Source/DataflowEditor/Private/Dataflow/DataflowSCommentNode.cpp
@@ -84,9 +84,7 @@ TSharedPtr<FAssetSchemaAction_Dataflow_DuplicateCommentNode_DataflowEdNode> FAss
 //
 UEdGraphNode* FAssetSchemaAction_Dataflow_DuplicateCommentNode_DataflowEdNode::PerformAction(class UEdGraph* ParentGraph, UEdGraphPin* FromPin, const FVector2D Location, bool bSelectNewNode)
 {
-	FString NewNodeName = CommentNodeToDuplicate->NodeComment;
-
-	UEdGraphNode_Comment* CommentTemplate = NewObject<UEdGraphNode_Comment>(ParentGraph, FName(*NewNodeName));
+	UEdGraphNode_Comment* CommentTemplate = NewObject<UEdGraphNode_Comment>(ParentGraph, FName("Comment"));
 
 	CommentTemplate->bCommentBubbleVisible_InDetailsPanel = false;
 	CommentTemplate->bCommentBubbleVisible = false;
@@ -115,7 +113,7 @@ UEdGraphNode* FAssetSchemaAction_Dataflow_DuplicateCommentNode_DataflowEdNode::P
 	CommentTemplate->CommentColor = CommentNodeToDuplicate->CommentColor;
 	CommentTemplate->FontSize = CommentNodeToDuplicate->FontSize;
 
-	CommentTemplate->NodeComment = NewNodeName;
+	CommentTemplate->NodeComment = CommentNodeToDuplicate->NodeComment;;
 
 	ParentGraph->NotifyGraphChanged();
 
@@ -138,9 +136,7 @@ TSharedPtr<FAssetSchemaAction_Dataflow_PasteCommentNode_DataflowEdNode> FAssetSc
 //
 UEdGraphNode* FAssetSchemaAction_Dataflow_PasteCommentNode_DataflowEdNode::PerformAction(class UEdGraph* ParentGraph, UEdGraphPin* FromPin, const FVector2D Location, bool bSelectNewNode)
 {
-	FString NewNodeName = NodeName.ToString();
-
-	UEdGraphNode_Comment* CommentTemplate = NewObject<UEdGraphNode_Comment>(ParentGraph, FName(*NewNodeName));
+	UEdGraphNode_Comment* CommentTemplate = NewObject<UEdGraphNode_Comment>(ParentGraph, FName("Comment"));
 
 	CommentTemplate->bCommentBubbleVisible_InDetailsPanel = false;
 	CommentTemplate->bCommentBubbleVisible = false;
@@ -167,7 +163,7 @@ UEdGraphNode* FAssetSchemaAction_Dataflow_PasteCommentNode_DataflowEdNode::Perfo
 	CommentTemplate->CommentColor = Color;
 	CommentTemplate->FontSize = FontSize;
 
-	CommentTemplate->NodeComment = NewNodeName;
+	CommentTemplate->NodeComment = NodeName.ToString();
 
 	ParentGraph->NotifyGraphChanged();
 
