Dataflow Editor: add a cvar to toggle automatic graph evaluation in the editor. Modify the SDataflowGraphEditor to show a warning message in the editor panel when graph evaluation is paused

diff --git a/Engine/Plugins/Experimental/Dataflow/Source/DataflowEditor/Private/Dataflow/DataflowEditorCommands.cpp b/Engine/Plugins/Experimental/Dataflow/Source/DataflowEditor/Private/Dataflow/DataflowEditorCommands.cpp
--- a/Engine/Plugins/Experimental/Dataflow/Source/DataflowEditor/Private/Dataflow/DataflowEditorCommands.cpp
+++ b/Engine/Plugins/Experimental/Dataflow/Source/DataflowEditor/Private/Dataflow/DataflowEditorCommands.cpp
@@ -182,7 +182,7 @@ const FDataflowNode* FDataflowEditorCommands::EvaluateNode(UE::Dataflow::FContex
 
 		if (const FDataflowTerminalNode* const TerminalNode = Node->AsType<const FDataflowTerminalNode>())
 		{
-			if (TerminalNode->bActive && Asset)
+			if (TerminalNode->IsActive() && Asset)
 			{
 				UE_LOG(LogChaosDataflow, Verbose, TEXT("FDataflowTerminalNode::SetAssetValue(): TerminalNode [%s], Asset [%s]"), *TerminalNode->GetName().ToString(), *Asset->GetName());
 				TerminalNode->SetAssetValue(Asset, Context);
@@ -337,12 +337,13 @@ void FDataflowEditorCommands::OnPropertyValueChanged(UDataflow* OutDataflow, TSh
 					if (TSharedPtr<FDataflowNode> DataflowNode = Node->GetDataflowNode())
 					{
 						// Active state update
-						if (InPropertyChangedEvent.GetPropertyName() == GET_MEMBER_NAME_CHECKED(FDataflowNode, bActive))
+						if (InPropertyChangedEvent.GetPropertyName() == FDataflowNode::GetActivePropertyName())
 						{
 							// Reflect the active state on the drawing of the node
-							if (DataflowNode->bActive != Node->IsNodeEnabled())
+							constexpr bool bCheckIsActiveFlagOnly = true;
+							if (DataflowNode->IsActive(bCheckIsActiveFlagOnly) != Node->IsNodeEnabled())
 							{
-								Node->SetEnabledState(DataflowNode->bActive ? ENodeEnabledState::Enabled : ENodeEnabledState::Disabled);
+								Node->SetEnabledState(DataflowNode->IsActive(bCheckIsActiveFlagOnly) ? ENodeEnabledState::Enabled : ENodeEnabledState::Disabled);
 							}
 						}
 						// Invalidate the node and reset the editor timestamp
diff --git a/Engine/Plugins/Experimental/Dataflow/Source/DataflowEditor/Private/Dataflow/DataflowEditorToolkit.cpp b/Engine/Plugins/Experimental/Dataflow/Source/DataflowEditor/Private/Dataflow/DataflowEditorToolkit.cpp
--- a/Engine/Plugins/Experimental/Dataflow/Source/DataflowEditor/Private/Dataflow/DataflowEditorToolkit.cpp
+++ b/Engine/Plugins/Experimental/Dataflow/Source/DataflowEditor/Private/Dataflow/DataflowEditorToolkit.cpp
@@ -1135,16 +1135,22 @@ void FDataflowEditorToolkit::EvaluateNode(const FDataflowNode* Node, const FData
 {
 	UE_LOG(LogChaosDataflow, VeryVerbose, TEXT("FDataflowEditorToolkit::EvaluateNode(): Node [%s], Output [%s]"), Node ? *Node->GetName().ToString() : TEXT("nullptr"), Output ? *Output->GetName().ToString() : TEXT("nullptr"));
 
-	const bool bIsInPIEOrSimulate = GEditor->PlayWorld || GEditor->bIsSimulatingInEditor;
-	if (!bIsInPIEOrSimulate)  // TODO: make this test optional for some toolkit
+	const IConsoleVariable* const ConsoleVar = IConsoleManager::Get().FindConsoleVariable(TEXT("p.Dataflow.EnableGraphEval"));
+	const bool bDataflowEnableGraphEval = ConsoleVar ? ConsoleVar->GetBool() : true;
+
+	if (bDataflowEnableGraphEval)
 	{
-		if (const TObjectPtr<UDataflowBaseContent>& EditorContent = GetEditorContent())
+		const bool bIsInPIEOrSimulate = GEditor->PlayWorld || GEditor->bIsSimulatingInEditor;
+		if (!bIsInPIEOrSimulate)  // TODO: make this test optional for some toolkit
 		{
-			if (EditorContent->GetDataflowAsset())
+			if (const TObjectPtr<UDataflowBaseContent>& EditorContent = GetEditorContent())
 			{
-				// If Node is null, the terminal node with the given name will be used instead
-				FDataflowEditorCommands::EvaluateNode(*EditorContent->GetDataflowContext().Get(), InOutTimestamp, EditorContent->GetDataflowAsset(),
-					Node, Output, EditorContent->GetDataflowTerminal(), EditorContent->GetTerminalAsset());
+				if (EditorContent->GetDataflowAsset())
+				{
+					// If Node is null, the terminal node with the given name will be used instead
+					FDataflowEditorCommands::EvaluateNode(*EditorContent->GetDataflowContext().Get(), InOutTimestamp, EditorContent->GetDataflowAsset(),
+						Node, Output, EditorContent->GetDataflowTerminal(), EditorContent->GetTerminalAsset());
+				}
 			}
 		}
 	}
diff --git a/Engine/Plugins/Experimental/Dataflow/Source/DataflowEditor/Private/Dataflow/DataflowGraphEditor.cpp b/Engine/Plugins/Experimental/Dataflow/Source/DataflowEditor/Private/Dataflow/DataflowGraphEditor.cpp
--- a/Engine/Plugins/Experimental/Dataflow/Source/DataflowEditor/Private/Dataflow/DataflowGraphEditor.cpp
+++ b/Engine/Plugins/Experimental/Dataflow/Source/DataflowEditor/Private/Dataflow/DataflowGraphEditor.cpp
@@ -168,6 +168,35 @@ void SDataflowGraphEditor::Construct(const FArguments& InArgs, UObject* InAssetO
 	SGraphEditor::Construct(Arguments);
 
 	SetNodeFactory( MakeShared<FDataflowGraphNodeFactory>(this) );
+
+	InitGraphEditorMessageBar();
+
+	// Take the existing graph panel widget and add it to a new SOverlay so we can place the message bar over top of it
+	if (ChildSlot.Num() > 0)
+	{
+		TSharedRef<SWidget> ChildWidget = ChildSlot.GetChildAt(0);
+
+		this->ChildSlot
+		[
+			SNew(SOverlay)
+				+ SOverlay::Slot()
+				[
+					ChildWidget
+				]
+				+ SOverlay::Slot()
+				[
+					MessageBar.ToSharedRef()
+				]
+		];
+	}
+}
+
+SDataflowGraphEditor::~SDataflowGraphEditor()
+{
+	if (IConsoleVariable* const ConsoleVar = IConsoleManager::Get().FindConsoleVariable(TEXT("p.Dataflow.EnableGraphEval")))
+	{
+		ConsoleVar->OnChangedDelegate().Remove(CVarChangedDelegateHandle);
+	}
 }
 
 TSharedPtr<UE::Dataflow::FContext> SDataflowGraphEditor::GetDataflowContext() const
@@ -210,7 +239,7 @@ void SDataflowGraphEditor::EvaluateNode()
 			{
 				if (const TSharedPtr<FDataflowNode> DataflowNode = DataflowGraph->FindBaseNode(EdNode->GetDataflowNodeGuid()))
 				{
-					if (DataflowNode->bActive)
+					if (DataflowNode->IsActive())
 					{
 						DataflowNode->Invalidate();  // Force evaluation
 
@@ -750,4 +779,69 @@ void SDataflowGraphEditor::PasteSelectedNodes()
 	}
 }
 
+void SDataflowGraphEditor::InitGraphEditorMessageBar()
+{
+	SAssignNew(MessageBar, SHorizontalBox)
+		+ SHorizontalBox::Slot()
+		.HAlign(HAlign_Center)
+		.VAlign(VAlign_Bottom)
+		.FillWidth(1.0f)
+		.Padding(FMargin(0.0f, 0.0f, 0.f, 15.f))
+		[
+			SNew(SBorder)
+				.BorderImage(FCoreStyle::Get().GetBrush(TEXT("BlackBrush")))
+				.BorderBackgroundColor(FLinearColor::Red)
+				.Padding(FMargin(20.f, 5.0f, 20.f, 5.0f))
+				[
+					SNew(SHorizontalBox)
+						+ SHorizontalBox::Slot()
+						.VAlign(VAlign_Center)
+						.Padding(FMargin(0.f, 0.f, 0.f, 0.f))
+						.FillWidth(1.0f)
+						[
+							SNew(STextBlock)
+								.Text(this, &SDataflowGraphEditor::GetGraphEditorOverlayText)
+								.Font(FCoreStyle::GetDefaultFontStyle("Bold", 14))
+						]
+				]
+		];
+
+	auto UpdateGraphForEvalEnabled = [this](bool bGraphEvalEnabled)
+	{
+		SetColorAndOpacity(bGraphEvalEnabled ? FLinearColor::White : FLinearColor(1.0f, 1.0f, 1.0f, 0.85f));
+
+		if (EdGraphObj)
+		{
+			for (UEdGraphNode* const Node : EdGraphObj->Nodes)
+			{
+				Node->SetForceDisplayAsDisabled(!bGraphEvalEnabled);
+			}
+		}
+
+		if (MessageBar)
+		{
+			MessageBar->SetVisibility(bGraphEvalEnabled ? EVisibility::Collapsed : EVisibility::SelfHitTestInvisible);
+		}
+		MessageBarText = bGraphEvalEnabled ? LOCTEXT("DataflowGraphEditorOverlayTextRunning", "") : LOCTEXT("DataflowGraphEditorOverlayTextPaused", "GRAPH EVALUATION PAUSED");
+	};
+
+	IConsoleVariable* const ConsoleVar = IConsoleManager::Get().FindConsoleVariable(TEXT("p.Dataflow.EnableGraphEval"));
+	const bool bDataflowEnableGraphEval = ConsoleVar ? ConsoleVar->GetBool() : true;
+	UpdateGraphForEvalEnabled(bDataflowEnableGraphEval);
+
+	if (ConsoleVar)
+	{
+		CVarChangedDelegateHandle = ConsoleVar->OnChangedDelegate().AddLambda([this, UpdateGraphForEvalEnabled](IConsoleVariable* Var)
+		{
+			const bool bEvalEnabled = Var->GetBool();
+			UpdateGraphForEvalEnabled(bEvalEnabled);
+		});
+	}
+}
+
+FText SDataflowGraphEditor::GetGraphEditorOverlayText() const
+{
+	return MessageBarText;
+}
+
 #undef LOCTEXT_NAMESPACE
diff --git a/Engine/Plugins/Experimental/Dataflow/Source/DataflowEditor/Public/Dataflow/DataflowGraphEditor.h b/Engine/Plugins/Experimental/Dataflow/Source/DataflowEditor/Public/Dataflow/DataflowGraphEditor.h
--- a/Engine/Plugins/Experimental/Dataflow/Source/DataflowEditor/Public/Dataflow/DataflowGraphEditor.h
+++ b/Engine/Plugins/Experimental/Dataflow/Source/DataflowEditor/Public/Dataflow/DataflowGraphEditor.h
@@ -61,6 +61,8 @@ public:
 	DECLARE_MULTICAST_DELEGATE_OneParam(FOnNodeDeletedMulticast, const FGraphPanelSelectionSet&)
 	FOnNodeDeletedMulticast OnNodeDeletedMulticast;
 
+	virtual ~SDataflowGraphEditor();
+
 	// SWidget overrides
 	virtual FReply OnKeyDown(const FGeometry& MyGeometry, const FKeyEvent& InKeyEvent) override;
 	virtual FReply OnKeyUp(const FGeometry& MyGeometry, const FKeyEvent& InKeyEvent) override;
@@ -170,6 +172,12 @@ private:
 
 	bool GetPinVisibility(SGraphEditor::EPinVisibility InVisibility) const;
 
+	/** Create a widget to display an overlaid message in the graph editor panel */
+	void InitGraphEditorMessageBar();
+
+	/** Text for overlay message in graph editor panel */
+	FText GetGraphEditorOverlayText() const;
+
 	FDataflowEditorCommands::FOnDragDropEventCallback OnDragDropEventCallback;
 	FDataflowEditorCommands::FGraphEvaluationCallback EvaluateGraphCallback;
 
@@ -199,5 +207,9 @@ private:
 	bool RightControlKeyDown = false;
 	bool LeftAltKeyDown = false;
 	bool RightAltKeyDown = false;
+
+	FDelegateHandle CVarChangedDelegateHandle;
+	TSharedPtr<SWidget> MessageBar;
+	FText MessageBarText;
 };
 
diff --git a/Engine/Source/Runtime/Experimental/Dataflow/Core/Private/Dataflow/DataflowConnection.cpp b/Engine/Source/Runtime/Experimental/Dataflow/Core/Private/Dataflow/DataflowConnection.cpp
--- a/Engine/Source/Runtime/Experimental/Dataflow/Core/Private/Dataflow/DataflowConnection.cpp
+++ b/Engine/Source/Runtime/Experimental/Dataflow/Core/Private/Dataflow/DataflowConnection.cpp
@@ -54,7 +54,7 @@ void FDataflowConnection::InitFromType()
 
 bool FDataflowConnection::IsOwningNodeEnabled() const
 {
-	return (OwningNode && OwningNode->bActive);
+	return (OwningNode && OwningNode->IsActive());
 }
 
 FGuid FDataflowConnection::GetOwningNodeGuid() const
diff --git a/Engine/Source/Runtime/Experimental/Dataflow/Core/Private/Dataflow/DataflowNode.cpp b/Engine/Source/Runtime/Experimental/Dataflow/Core/Private/Dataflow/DataflowNode.cpp
--- a/Engine/Source/Runtime/Experimental/Dataflow/Core/Private/Dataflow/DataflowNode.cpp
+++ b/Engine/Source/Runtime/Experimental/Dataflow/Core/Private/Dataflow/DataflowNode.cpp
@@ -22,6 +22,10 @@ const FLinearColor FDataflowNode::DefaultNodeBodyTintColor = FLinearColor(0.f, 0
 
 const FName FDataflowAnyType::TypeName = TEXT("FDataflowAnyType");
 
+bool bDataflowEnableGraphEval = true;
+FAutoConsoleVariableRef CVARDataflowEnableGraphEval(TEXT("p.Dataflow.EnableGraphEval"), bDataflowEnableGraphEval,
+	TEXT("Enable automatic graph evaluation in the Dataflow Editor. [def:true]"));
+
 namespace UE::Dataflow::Private
 {
 static uint32 GetArrayElementOffsetFromReference(const FArrayProperty* const ArrayProperty, const UE::Dataflow::FConnectionReference& Reference)
@@ -619,6 +623,13 @@ void FDataflowNode::Unfreeze(UE::Dataflow::FContext& Context)
 	Invalidate();
 }
 
+bool FDataflowNode::IsActive(bool bCheckFlagOnly) const
+{
+	PRAGMA_DISABLE_DEPRECATION_WARNINGS
+	return bActive && (bCheckFlagOnly || bDataflowEnableGraphEval);
+	PRAGMA_ENABLE_DEPRECATION_WARNINGS
+}
+
 void FDataflowNode::Invalidate(const UE::Dataflow::FTimestamp& InModifiedTimestamp)
 {
 	if (bPauseInvalidations)
diff --git a/Engine/Source/Runtime/Experimental/Dataflow/Core/Private/Dataflow/DataflowNodeParameters.cpp b/Engine/Source/Runtime/Experimental/Dataflow/Core/Private/Dataflow/DataflowNodeParameters.cpp
--- a/Engine/Source/Runtime/Experimental/Dataflow/Core/Private/Dataflow/DataflowNodeParameters.cpp
+++ b/Engine/Source/Runtime/Experimental/Dataflow/Core/Private/Dataflow/DataflowNodeParameters.cpp
@@ -97,7 +97,7 @@ namespace UE::Dataflow
 			// Note: If the node is deactivated and has an output (like above), then the output might still need to be forwarded.
 			//       Therefore the Evaluate method has to be called for whichever value of bActive.
 			//       However if the node is deactivated and has no outputs (like below), now is the time to check its bActive state.
-			else if (Node->bActive && !Node->IsFrozen())
+			else if (Node->IsActive() && !Node->IsFrozen())
 			{
 				// TODO: When no outputs are specified, this call to Evaluate should really be removed.
 				//       The purpose of the node evaluation function is to evaluate outputs.
diff --git a/Engine/Source/Runtime/Experimental/Dataflow/Core/Public/Dataflow/DataflowNode.h b/Engine/Source/Runtime/Experimental/Dataflow/Core/Public/Dataflow/DataflowNode.h
--- a/Engine/Source/Runtime/Experimental/Dataflow/Core/Public/Dataflow/DataflowNode.h
+++ b/Engine/Source/Runtime/Experimental/Dataflow/Core/Public/Dataflow/DataflowNode.h
@@ -54,6 +54,7 @@ PRAGMA_ENABLE_DEPRECATION_WARNINGS
 	UE_DEPRECATED(5.5, "Will be made private in 5.7, use Get/Set methods instead.")
 	TMap< int, FDataflowOutput* > Outputs;
 
+	UE_DEPRECATED(5.6, "Will be made private in 5.8, use IsActive function instead.")
 	UPROPERTY(EditAnywhere, Category = "Dataflow")
 	bool bActive = true;
 
@@ -342,6 +343,22 @@ PRAGMA_ENABLE_DEPRECATION_WARNINGS
 		return FrozenProperties;
 	}
 
+	/**
+	 * Gets the property name for bActive. This exists so the property can be made private later.
+	 */
+	static const FName GetActivePropertyName()
+	{
+		PRAGMA_DISABLE_DEPRECATION_WARNINGS
+		return GET_MEMBER_NAME_CHECKED(FDataflowNode, bActive);
+		PRAGMA_ENABLE_DEPRECATION_WARNINGS
+	}
+
+	/** 
+	 * If bCheckFlagOnly is false, this function returns true if both the bActive property and the CVAR p.Dataflow.EnableGraphEval are true
+	 * If bCheckFlagOnly is true, this function returns true if the bActive property is true
+	 */
+	DATAFLOWCORE_API bool IsActive(bool bCheckFlagOnly = false) const;
+
 	/**
 	*   GetValue(...)
 	*
diff --git a/Engine/Source/Runtime/Experimental/Dataflow/Engine/Private/Dataflow/DataflowBlueprintLibrary.cpp b/Engine/Source/Runtime/Experimental/Dataflow/Engine/Private/Dataflow/DataflowBlueprintLibrary.cpp
--- a/Engine/Source/Runtime/Experimental/Dataflow/Engine/Private/Dataflow/DataflowBlueprintLibrary.cpp
+++ b/Engine/Source/Runtime/Experimental/Dataflow/Engine/Private/Dataflow/DataflowBlueprintLibrary.cpp
@@ -20,7 +20,7 @@ void UDataflowBlueprintLibrary::EvaluateTerminalNodeByName(UDataflow* Dataflow,
 				//       Therefore the Evaluate method has to be called for whichever value of bActive.
 				//       This however isn't the case of SetAssetValue() for which the active state needs to be checked before the call.
 				TerminalNode->Evaluate(Context);
-				if (TerminalNode->bActive && ResultAsset)
+				if (TerminalNode->IsActive() && ResultAsset)
 				{
 					UE_LOG(LogChaosDataflow, Verbose, TEXT("FDataflowTerminalNode::SetAssetValue(): TerminalNode [%s], Asset [%s]"), *TerminalNodeName.ToString(), *ResultAsset->GetName());
 					TerminalNode->SetAssetValue(ResultAsset, Context);
