Fix localization warning from empty LOCTEXT

diff --git a/Engine/Plugins/Experimental/Dataflow/Source/DataflowEditor/Private/Dataflow/DataflowGraphEditor.cpp b/Engine/Plugins/Experimental/Dataflow/Source/DataflowEditor/Private/Dataflow/DataflowGraphEditor.cpp
--- a/Engine/Plugins/Experimental/Dataflow/Source/DataflowEditor/Private/Dataflow/DataflowGraphEditor.cpp
+++ b/Engine/Plugins/Experimental/Dataflow/Source/DataflowEditor/Private/Dataflow/DataflowGraphEditor.cpp
@@ -822,7 +822,7 @@ void SDataflowGraphEditor::InitGraphEditorMessageBar()
 		{
 			MessageBar->SetVisibility(bGraphEvalEnabled ? EVisibility::Collapsed : EVisibility::SelfHitTestInvisible);
 		}
-		MessageBarText = bGraphEvalEnabled ? LOCTEXT("DataflowGraphEditorOverlayTextRunning", "") : LOCTEXT("DataflowGraphEditorOverlayTextPaused", "GRAPH EVALUATION PAUSED");
+		MessageBarText = bGraphEvalEnabled ? FText() : LOCTEXT("DataflowGraphEditorOverlayTextPaused", "GRAPH EVALUATION PAUSED");
 	};
 
 	IConsoleVariable* const ConsoleVar = IConsoleManager::Get().FindConsoleVariable(TEXT("p.Dataflow.EnableGraphEval"));
