Dataflow Editor: enable middle mouse button panning (really truck/pedestal) in the construction viewport in 2D mode. Toggling between right/middle/both can be done in the project settings window

diff --git a/Engine/Plugins/Experimental/Dataflow/Source/DataflowEditor/DataflowEditor.Build.cs b/Engine/Plugins/Experimental/Dataflow/Source/DataflowEditor/DataflowEditor.Build.cs
--- a/Engine/Plugins/Experimental/Dataflow/Source/DataflowEditor/DataflowEditor.Build.cs
+++ b/Engine/Plugins/Experimental/Dataflow/Source/DataflowEditor/DataflowEditor.Build.cs
@@ -31,6 +31,7 @@ namespace UnrealBuildTool.Rules
 					"DataflowEnginePlugin",
 					"DataflowNodes",
 					"DataflowSimulation",
+					"DeveloperSettings",
 					"DynamicMesh",
 					"Engine",
 					"EditorFramework",
diff --git a/Engine/Plugins/Experimental/Dataflow/Source/DataflowEditor/Private/Dataflow/DataflowConstructionViewportClient.cpp b/Engine/Plugins/Experimental/Dataflow/Source/DataflowEditor/Private/Dataflow/DataflowConstructionViewportClient.cpp
--- a/Engine/Plugins/Experimental/Dataflow/Source/DataflowEditor/Private/Dataflow/DataflowConstructionViewportClient.cpp
+++ b/Engine/Plugins/Experimental/Dataflow/Source/DataflowEditor/Private/Dataflow/DataflowConstructionViewportClient.cpp
@@ -21,6 +21,7 @@
 #include "SGraphPanel.h"
 #include "SNodePanel.h"
 #include "BaseBehaviors/ClickDragBehavior.h"
+#include "Behaviors/2DViewportBehaviorTargets.h"
 
 #define LOCTEXT_NAMESPACE "DataflowConstructionViewportClient"
 
@@ -38,6 +39,44 @@ FDataflowConstructionViewportClient::FDataflowConstructionViewportClient(FEditor
 	PreviewScene = static_cast<FDataflowPreviewSceneBase*>(InPreviewScene);
 	bEnableSceneTicking = bCouldTickScene;
 
+	bool bUseRightMouseButton = true;
+	bool bUseMiddleMouseButton = true;
+	if (const UDataflowEditorOptions* const Options = GetDefault<UDataflowEditorOptions>())
+	{
+		bUseRightMouseButton = (Options->ConstructionViewportMousePanButton == EDataflowConstructionViewportMousePanButton::Right || Options->ConstructionViewportMousePanButton == EDataflowConstructionViewportMousePanButton::RightOrMiddle);
+		bUseMiddleMouseButton = (Options->ConstructionViewportMousePanButton == EDataflowConstructionViewportMousePanButton::Middle || Options->ConstructionViewportMousePanButton == EDataflowConstructionViewportMousePanButton::RightOrMiddle);
+	}
+
+	if (!bUseRightMouseButton)
+	{
+		// Intercept right mouse drag with a do-nothing ClickDragInputBehavior
+		const TObjectPtr<ULocalClickDragInputBehavior> RightMouseClickDragBehavior = NewObject<ULocalClickDragInputBehavior>();
+		RightMouseClickDragBehavior->Initialize();
+
+		// We'll have the priority of our viewport behaviors be lower (i.e. higher numerically) than both the gizmo default and the tool default
+		constexpr int ViewportBehaviorPriority = FMath::Max(FInputCapturePriority::DEFAULT_GIZMO_PRIORITY, FInputCapturePriority::DEFAULT_TOOL_PRIORITY) + 2;
+		RightMouseClickDragBehavior->SetDefaultPriority(ViewportBehaviorPriority);
+		RightMouseClickDragBehavior->SetUseRightMouseButton();
+		RightMouseClickDragBehavior->CanBeginClickDragFunc = [](const FInputDeviceRay& InputDeviceRay)
+		{
+			return FInputRayHit(TNumericLimits<float>::Max()); // bHit is true. Depth is max to lose the standard tiebreaker.
+		};
+		BehaviorsFor2DMode.Add(RightMouseClickDragBehavior);
+	}
+
+	if (bUseMiddleMouseButton)
+	{
+		OrthoScrollBehaviorTarget = MakeUnique<FEditor2DScrollBehaviorTarget>(this);
+		UClickDragInputBehavior* const MiddleMouseClickDragInputBehavior = NewObject<UClickDragInputBehavior>();
+		MiddleMouseClickDragInputBehavior->Initialize(OrthoScrollBehaviorTarget.Get());
+
+		// We'll have the priority of our viewport behaviors be lower (i.e. higher numerically) than both the gizmo default and the tool default
+		constexpr int ViewportBehaviorPriority = FMath::Max(FInputCapturePriority::DEFAULT_GIZMO_PRIORITY, FInputCapturePriority::DEFAULT_TOOL_PRIORITY) + 3;
+		MiddleMouseClickDragInputBehavior->SetDefaultPriority(ViewportBehaviorPriority);
+		MiddleMouseClickDragInputBehavior->SetUseMiddleMouseButton();
+		BehaviorsFor2DMode.Add(MiddleMouseClickDragInputBehavior);
+	}
+
 	if (const UDataflowEditorOptions* const Options = UDataflowEditorOptions::StaticClass()->GetDefaultObject<UDataflowEditorOptions>())
 	{
 		FOVAngle = Options->ConstructionViewFOV;
@@ -236,6 +275,29 @@ void FDataflowConstructionViewportClient::SetConstructionViewMode(const UE::Data
 
 	ConstructionViewMode = InViewMode;
 
+	//
+	// Update input behaviors
+	//
+
+	BehaviorSet->RemoveAll();
+
+	for (UInputBehavior* const Behavior : BaseBehaviors)
+	{
+		BehaviorSet->Add(Behavior);
+	}
+
+	if (!InViewMode->IsPerspective())
+	{
+		for (UInputBehavior* const Behavior : BehaviorsFor2DMode)
+		{
+			BehaviorSet->Add(Behavior);
+		}
+	}
+
+	ModeTools->GetInteractiveToolsContext()->InputRouter->DeregisterSource(this);
+	ModeTools->GetInteractiveToolsContext()->InputRouter->RegisterSource(this);
+
+
 	SetViewportType((ELevelViewportType)ConstructionViewMode->GetViewportType());
 
 	if (const FViewportCameraTransform* const FoundPreviousTransform = SavedInactiveViewTransforms.Find(InViewMode->GetName()))
@@ -281,4 +343,10 @@ void FDataflowConstructionViewportClient::DrawCanvas(FViewport& InViewport, FSce
 	}
 }
 
+void FDataflowConstructionViewportClient::AddReferencedObjects(FReferenceCollector& Collector)
+{
+	FDataflowEditorViewportClientBase::AddReferencedObjects(Collector);
+	Collector.AddReferencedObjects(BehaviorsFor2DMode);
+}
+
 #undef LOCTEXT_NAMESPACE 
diff --git a/Engine/Plugins/Experimental/Dataflow/Source/DataflowEditor/Private/Dataflow/DataflowEditorOptions.cpp b/Engine/Plugins/Experimental/Dataflow/Source/DataflowEditor/Private/Dataflow/DataflowEditorOptions.cpp
--- a/Engine/Plugins/Experimental/Dataflow/Source/DataflowEditor/Private/Dataflow/DataflowEditorOptions.cpp
+++ b/Engine/Plugins/Experimental/Dataflow/Source/DataflowEditor/Private/Dataflow/DataflowEditorOptions.cpp
@@ -15,3 +15,17 @@ UDataflowEditorOptions::UDataflowEditorOptions(class FObjectInitializer const& O
 
 	bSimulationViewFixedExposure = false;
 }
+
+FName UDataflowEditorOptions::GetCategoryName() const
+{
+	return TEXT("Plugins");
+}
+
+#if WITH_EDITOR
+
+FText UDataflowEditorOptions::GetSectionText() const
+{
+	return NSLOCTEXT("DataflowEditorPlugin", "DataflowEditorSettingsSection", "Dataflow Editor");
+}
+
+#endif
diff --git a/Engine/Plugins/Experimental/Dataflow/Source/DataflowEditor/Private/Dataflow/DataflowEditorOptions.h b/Engine/Plugins/Experimental/Dataflow/Source/DataflowEditor/Private/Dataflow/DataflowEditorOptions.h
--- a/Engine/Plugins/Experimental/Dataflow/Source/DataflowEditor/Private/Dataflow/DataflowEditorOptions.h
+++ b/Engine/Plugins/Experimental/Dataflow/Source/DataflowEditor/Private/Dataflow/DataflowEditorOptions.h
@@ -3,11 +3,26 @@
 #pragma once
 
 #include "UObject/Object.h"
+#include "Engine/DeveloperSettings.h"
 #include "DataflowEditorOptions.generated.h"
 
+UENUM(BlueprintType)
+enum class EDataflowConstructionViewportMousePanButton : uint8
+{
+	/** Right Mouse Button */
+	Right = 0,
+
+	/** Middle Mouse Button */
+	Middle = 1,
+
+	/** Either Right or Middle Mouse Button */
+	RightOrMiddle = 2,
+};
+
+
 /** A settings class used to store and retreive user settings related to the Dataflow Editor */
-UCLASS(hidecategories = Object, config = EditorPerProjectUserSettings, MinimalAPI)
-class UDataflowEditorOptions : public UObject
+UCLASS(HideCategories = Object, Config = EditorPerProjectUserSettings, Meta = (DisplayName = "Dataflow Editor"), MinimalAPI)
+class UDataflowEditorOptions : public UDeveloperSettings
 {
 	GENERATED_UCLASS_BODY()
 
@@ -27,4 +42,15 @@ class UDataflowEditorOptions : public UObject
 	UPROPERTY(config)
 	bool bSimulationViewFixedExposure;
 
+	/** Which mouse button controls camera panning in the Construction Viewport in 2D view mode */
+	UPROPERTY(Config, EditAnywhere, Category = UI)
+	EDataflowConstructionViewportMousePanButton ConstructionViewportMousePanButton;
+
+	// UDeveloperSettings overrides
+	virtual FName GetCategoryName() const override;
+
+#if WITH_EDITOR
+	virtual FText GetSectionText() const override;
+#endif
+
 };
diff --git a/Engine/Plugins/Experimental/Dataflow/Source/DataflowEditor/Private/Dataflow/DataflowEditorViewportClientBase.cpp b/Engine/Plugins/Experimental/Dataflow/Source/DataflowEditor/Private/Dataflow/DataflowEditorViewportClientBase.cpp
--- a/Engine/Plugins/Experimental/Dataflow/Source/DataflowEditor/Private/Dataflow/DataflowEditorViewportClientBase.cpp
+++ b/Engine/Plugins/Experimental/Dataflow/Source/DataflowEditor/Private/Dataflow/DataflowEditorViewportClientBase.cpp
@@ -40,6 +40,7 @@ FDataflowEditorViewportClientBase::FDataflowEditorViewportClientBase(FEditorMode
 		OnViewportClicked(HitProxy);
 	};
 
+	BaseBehaviors.Add(ClickDragBehavior);
 	BehaviorSet->Add(ClickDragBehavior);
 
 	ModeTools->GetInteractiveToolsContext()->InputRouter->RegisterSource(this);
@@ -62,6 +63,7 @@ const UInputBehaviorSet* FDataflowEditorViewportClientBase::GetInputBehaviors()
 void FDataflowEditorViewportClientBase::AddReferencedObjects(FReferenceCollector& Collector)
 {
 	Super::AddReferencedObjects(Collector);
+	Collector.AddReferencedObjects(BaseBehaviors);
 	Collector.AddReferencedObject(BehaviorSet);
 }
 
diff --git a/Engine/Plugins/Experimental/Dataflow/Source/DataflowEditor/Public/Dataflow/DataflowConstructionViewportClient.h b/Engine/Plugins/Experimental/Dataflow/Source/DataflowEditor/Public/Dataflow/DataflowConstructionViewportClient.h
--- a/Engine/Plugins/Experimental/Dataflow/Source/DataflowEditor/Public/Dataflow/DataflowConstructionViewportClient.h
+++ b/Engine/Plugins/Experimental/Dataflow/Source/DataflowEditor/Public/Dataflow/DataflowConstructionViewportClient.h
@@ -19,6 +19,7 @@ class FDataflowPreviewSceneBase;
 class UInputBehaviorSet;
 class USelection;
 class FDataflowConstructionViewportClient;
+class IClickDragBehaviorTarget;
 namespace UE::Dataflow
 {
 	class IDataflowConstructionViewMode;
@@ -55,6 +56,11 @@ public:
 	DECLARE_MULTICAST_DELEGATE_OneParam(FOnSelectionChangedMulticast, const TArray<UPrimitiveComponent*>&)
 	FOnSelectionChangedMulticast OnSelectionChangedMulticast;
 
+protected:
+
+	// FGCObject
+	virtual void AddReferencedObjects(FReferenceCollector& Collector) override;
+
 private:
 
 	// FDataflowEditorViewportClientBase
@@ -67,6 +73,9 @@ private:
 	virtual void Draw(const FSceneView* View, FPrimitiveDrawInterface* PDI) override;
 	virtual void DrawCanvas(FViewport& InViewport, FSceneView& View, FCanvas& Canvas) override;
 
+	TUniquePtr<IClickDragBehaviorTarget> OrthoScrollBehaviorTarget;
+	TArray<TObjectPtr<UInputBehavior>> BehaviorsFor2DMode;
+
 	/** Toolkit used to create the viewport client */
 	TWeakPtr<FDataflowEditorToolkit> DataflowEditorToolkitPtr = nullptr;
 
diff --git a/Engine/Plugins/Experimental/Dataflow/Source/DataflowEditor/Public/Dataflow/DataflowEditorViewportClientBase.h b/Engine/Plugins/Experimental/Dataflow/Source/DataflowEditor/Public/Dataflow/DataflowEditorViewportClientBase.h
--- a/Engine/Plugins/Experimental/Dataflow/Source/DataflowEditor/Public/Dataflow/DataflowEditorViewportClientBase.h
+++ b/Engine/Plugins/Experimental/Dataflow/Source/DataflowEditor/Public/Dataflow/DataflowEditorViewportClientBase.h
@@ -23,24 +23,27 @@ protected:
 
 	virtual void OnViewportClicked(HHitProxy* HitProxy) = 0;
 
+	// FGCObject
+	virtual void AddReferencedObjects(FReferenceCollector& Collector) override;
+
 	/** Dataflow preview scene from the toolkit */
 	FDataflowPreviewSceneBase* DataflowPreviewScene = nullptr;
 
+	/** Behaviors defined by this base class */
+	TArray<TObjectPtr<UInputBehavior>> BaseBehaviors;
+
+	/** All behaviors available to the current viewport (subclasses can add to this set) */
+	TObjectPtr<UInputBehaviorSet> BehaviorSet;
+
 private:
 
 	// IInputBehaviorSource
 	virtual const UInputBehaviorSet* GetInputBehaviors() const override;
 
-	// FGCObject
-	virtual void AddReferencedObjects(FReferenceCollector& Collector) override;
-
 	/** Called from constructor/destructor to manage callbacks when settings change */
 	void RegisterDelegates();
 	void DeregisterDelegates();
 
-	/** Behaviors available to the current viewport */
-	TObjectPtr<UInputBehaviorSet> BehaviorSet;
-
 	/** Handle for callback when preview scene profile settings change */
 	FDelegateHandle OnAssetViewerSettingsChangedDelegateHandle;
 
