Dataflow : UI for variables and future subgraphs

diff --git a/Engine/Plugins/Experimental/Dataflow/Source/DataflowEditor/DataflowEditor.Build.cs b/Engine/Plugins/Experimental/Dataflow/Source/DataflowEditor/DataflowEditor.Build.cs
--- a/Engine/Plugins/Experimental/Dataflow/Source/DataflowEditor/DataflowEditor.Build.cs
+++ b/Engine/Plugins/Experimental/Dataflow/Source/DataflowEditor/DataflowEditor.Build.cs
@@ -22,6 +22,7 @@ namespace UnrealBuildTool.Rules
 					"AssetTools",
 					"AssetRegistry",
 					"BaseCharacterFXEditor",
+					"BlueprintGraph",
 					"Chaos",
 					"Core",
 					"CoreUObject",
diff --git a/Engine/Plugins/Experimental/Dataflow/Source/DataflowEditor/Private/Dataflow/DataflowAssetEditUtils.cpp b/Engine/Plugins/Experimental/Dataflow/Source/DataflowEditor/Private/Dataflow/DataflowAssetEditUtils.cpp
--- a/Engine/Plugins/Experimental/Dataflow/Source/DataflowEditor/Private/Dataflow/DataflowAssetEditUtils.cpp
+++ b/Engine/Plugins/Experimental/Dataflow/Source/DataflowEditor/Private/Dataflow/DataflowAssetEditUtils.cpp
@@ -221,7 +221,7 @@ namespace UE::Dataflow::EditAssetUtils
 	}
 
 	//////////////////////////////////////////////////////////////////////////////////////////////////////////
-	void SetVariableValue(UDataflow* DataflowAsset, FName VariableName, const FInstancedPropertyBag & SourceBag)
+	void SetVariableValue(UDataflow* DataflowAsset, FName VariableName, const FInstancedPropertyBag& SourceBag)
 	{
 		const FText TransactionName{ LOCTEXT("ChangeDataflowVariableValue", "Change Dataflow Variable Value") };
 
@@ -232,8 +232,7 @@ namespace UE::Dataflow::EditAssetUtils
 				{
 					if (SourceDesc->CachedProperty)
 					{
-						const void* SourcePropertyAddr = SourceDesc->CachedProperty->ContainerPtrToValuePtr<void>(SourceBag.GetValue().GetMemory());
-						EPropertyBagResult Result = DataflowAsset.Variables.SetValue(VariableName, SourceDesc->CachedProperty, SourcePropertyAddr);
+						EPropertyBagResult Result = DataflowAsset.Variables.SetValue(VariableName, SourceDesc->CachedProperty, SourceBag.GetValue().GetMemory());
 						if (Result == EPropertyBagResult::Success)
 						{
 							FDataflowAssetDelegates::OnVariablesChanged.Broadcast(&DataflowAsset, VariableName);
diff --git a/Engine/Plugins/Experimental/Dataflow/Source/DataflowEditor/Private/Dataflow/DataflowEditorModule.cpp b/Engine/Plugins/Experimental/Dataflow/Source/DataflowEditor/Private/Dataflow/DataflowEditorModule.cpp
--- a/Engine/Plugins/Experimental/Dataflow/Source/DataflowEditor/Private/Dataflow/DataflowEditorModule.cpp
+++ b/Engine/Plugins/Experimental/Dataflow/Source/DataflowEditor/Private/Dataflow/DataflowEditorModule.cpp
@@ -79,7 +79,7 @@ void FDataflowEditorModule::StartupModule()
 	{
 		PropertyModule->RegisterCustomPropertyTypeLayout(ScalarVertexPropertyGroupName, FOnGetPropertyTypeCustomizationInstance::CreateStatic(&UE::Dataflow::FScalarVertexPropertyGroupCustomization::MakeInstance));
 		PropertyModule->RegisterCustomPropertyTypeLayout(DataflowFunctionPropertyName, FOnGetPropertyTypeCustomizationInstance::CreateStatic(&UE::Dataflow::FFunctionPropertyCustomization::MakeInstance));
-		PropertyModule->RegisterCustomPropertyTypeLayout("DataflowInstance", FOnGetPropertyTypeCustomizationInstance::CreateStatic(&FDataflowInstanceDetails::MakeInstance));
+		PropertyModule->RegisterCustomPropertyTypeLayout("DataflowVariableOverrides", FOnGetPropertyTypeCustomizationInstance::CreateStatic(&FDataflowVariableOverridesDetails::MakeInstance));
 	}
 
 	UE::Dataflow::RenderingCallbacks();
diff --git a/Engine/Plugins/Experimental/Dataflow/Source/DataflowEditor/Private/Dataflow/DataflowEditorToolkit.cpp b/Engine/Plugins/Experimental/Dataflow/Source/DataflowEditor/Private/Dataflow/DataflowEditorToolkit.cpp
--- a/Engine/Plugins/Experimental/Dataflow/Source/DataflowEditor/Private/Dataflow/DataflowEditorToolkit.cpp
+++ b/Engine/Plugins/Experimental/Dataflow/Source/DataflowEditor/Private/Dataflow/DataflowEditorToolkit.cpp
@@ -35,6 +35,7 @@
 #include "Dataflow/DataflowSchema.h"
 #include "Dataflow/DataflowSkeletonView.h"
 #include "Dataflow/DataflowSimulationViewportClient.h"
+#include "Dataflow/DataflowMembersWidget.h"
 #include "DynamicMeshBuilder.h"
 #include "EditorModeManager.h"
 #include "EditorStyleSet.h"
@@ -81,6 +82,7 @@ const FName FDataflowEditorToolkit::CollectionSpreadSheetTabId_3(TEXT("DataflowE
 const FName FDataflowEditorToolkit::CollectionSpreadSheetTabId_4(TEXT("DataflowEditor_CollectionSpreadSheet_4"));
 const FName FDataflowEditorToolkit::SimulationViewportTabId(TEXT("DataflowEditor_SimulationViewport"));
 const FName FDataflowEditorToolkit::SimulationVisualizationTabId(TEXT("DataflowEditor_SimulationVisualizationTab"));
+const FName FDataflowEditorToolkit::MembersWidgetTabId(TEXT("DataflowEditor_MembersWidgetTab"));
 
 FDataflowEditorToolkit::FDataflowEditorToolkit(UAssetEditor* InOwningAssetEditor)
 	: FBaseCharacterFXEditorToolkit(InOwningAssetEditor, FName("DataflowEditor"))
@@ -135,12 +137,21 @@ FDataflowEditorToolkit::FDataflowEditorToolkit(UAssetEditor* InOwningAssetEditor
 					->Split
 					(
 						FTabManager::NewStack()
-						->SetSizeCoefficient(0.2f)
+						->SetSizeCoefficient(0.1f)
 						->AddTab(CollectionSpreadSheetTabId_1, ETabState::OpenedTab)
 						->SetExtensionId("CollectionSpreadSheetArea")
 						->SetHideTabWell(false)
 					)
 					->Split
+					(
+						FTabManager::NewStack()
+						->SetSizeCoefficient(0.1f)	// Relative height of (Dataflow Graph Editor, Outliner) vs (Tools Panel, Construction Viewport, Preview Viewport)
+						->AddTab(GraphCanvasTabId, ETabState::OpenedTab)
+						->SetExtensionId("MembersWidgetArea")
+						->SetHideTabWell(false)
+						->SetForegroundTab(MembersWidgetTabId)
+					)
+					->Split
 					(
 						FTabManager::NewStack()
 						->SetSizeCoefficient(0.8f)	// Relative height of (Dataflow Graph Editor, Outliner) vs (Tools Panel, Construction Viewport, Preview Viewport)
@@ -1258,6 +1269,16 @@ TSharedPtr<IDetailsView> FDataflowEditorToolkit::CreateAssetDetailsEditorWidget(
 
 }
 
+TSharedPtr<SDataflowMembersWidget> FDataflowEditorToolkit::CreateDataflowMembersWidget()
+{
+	using namespace UE::Dataflow;
+
+	TSharedRef<SDataflowMembersWidget> NewMembersWidget = SNew(SDataflowMembersWidget, GraphEditor);
+
+	return NewMembersWidget;
+
+}
+
 TSharedRef<SDockTab> FDataflowEditorToolkit::SpawnTab_AssetDetails(const FSpawnTabArgs& Args)
 {
 	check(Args.GetTabId() == DetailsTabID);
@@ -1519,6 +1540,17 @@ TSharedRef<SDockTab> FDataflowEditorToolkit::SpawnTab_SimulationVisualization(co
 }
 
 
+TSharedRef<SDockTab> FDataflowEditorToolkit::SpawnTab_MembersWidget(const FSpawnTabArgs& Args)
+{
+	TSharedRef<SDockTab> DataflowMembersTab = SNew(SDockTab)
+		.Label(LOCTEXT("MembersWidgetTitle", "Dataflow Members"));
+
+	MembersWidget = CreateDataflowMembersWidget();
+	DataflowMembersTab->SetContent(MembersWidget.ToSharedRef());
+
+	return DataflowMembersTab;
+}
+
 void FDataflowEditorToolkit::RegisterTabSpawners(const TSharedRef<FTabManager>& InTabManager)
 {
 	FAssetEditorToolkit::RegisterTabSpawners(InTabManager);
@@ -1609,7 +1641,9 @@ void FDataflowEditorToolkit::RegisterTabSpawners(const TSharedRef<FTabManager>&
 		.SetDisplayName(LOCTEXT("SimulationVisualizationTabDisplayName", "Simulation Visualization"))
 		.SetGroup(AssetEditorTabsCategory.ToSharedRef());
 
-	
+	InTabManager->RegisterTabSpawner(MembersWidgetTabId, FOnSpawnTab::CreateSP(this, &FDataflowEditorToolkit::SpawnTab_MembersWidget))
+		.SetDisplayName(LOCTEXT("MembersWidgetTabDisplayName", "Dataflow Members"))
+		.SetGroup(AssetEditorTabsCategory.ToSharedRef());
 }
 
 void FDataflowEditorToolkit::UnregisterTabSpawners(const TSharedRef<class FTabManager>& InTabManager)
diff --git a/Engine/Plugins/Experimental/Dataflow/Source/DataflowEditor/Private/Dataflow/DataflowGraphSchemaAction.cpp b/Engine/Plugins/Experimental/Dataflow/Source/DataflowEditor/Private/Dataflow/DataflowGraphSchemaAction.cpp
new file mode 100644
--- /dev/null
+++ b/Engine/Plugins/Experimental/Dataflow/Source/DataflowEditor/Private/Dataflow/DataflowGraphSchemaAction.cpp
@@ -0,0 +1,169 @@
+// Copyright Epic Games, Inc. All Rights Reserved.
+
+#include "Dataflow/DataflowGraphSchemaAction.h"
+
+#include "Dataflow/DataflowAssetEditUtils.h"
+#include "Dataflow/DataflowInstance.h"
+#include "Dataflow/DataflowObject.h"
+#include "Dataflow/DataflowSNode.h"
+#include "Dataflow/DataflowVariableNodes.h"
+#include "EdGraph/EdGraphSchema.h"
+#include "PropertyBagDetails.h"
+#include "StructUtils/PropertyBag.h"
+
+
+#define LOCTEXT_NAMESPACE "DataflowGraphSchemaAction"
+
+/////////////////////////////////////////////////////////////////////////////////////////////////////////////////////////
+
+FEdGraphSchemaAction_DataflowVariable::FEdGraphSchemaAction_DataflowVariable()
+	: FEdGraphSchemaAction()
+{}
+
+FEdGraphSchemaAction_DataflowVariable::FEdGraphSchemaAction_DataflowVariable(IDataflowInstanceInterface* InDataflowInstanceInterface, const FPropertyBagPropertyDesc& PropertyDesc)
+	: FEdGraphSchemaAction(
+		FText::FromString(CategoryFromFullName(PropertyDesc.Name)),
+		LOCTEXT("Description", "Description"),
+		LOCTEXT("Tooltip", "Tooltip"),
+		/*InGrouping*/0,
+		/*Keywords*/FText::GetEmpty(),
+		(int32)UE::Dataflow::ESchemaActionSectionID::VARIABLES)
+	, FullVariableName(PropertyDesc.Name)
+	, VariableName(NameFromFullName(PropertyDesc.Name))
+	, VariableCategory(CategoryFromFullName(PropertyDesc.Name))
+	, VariableType(UE::StructUtils::GetPropertyDescAsPin(PropertyDesc))
+	, DataflowInstanceInterface(InDataflowInstanceInterface)
+{
+}
+
+FString FEdGraphSchemaAction_DataflowVariable::CategoryFromFullName(FName FullName)
+{
+	const FString StrFullName(FullName.ToString());
+	int32 DotIndex = 0;
+	if (StrFullName.FindChar('|', DotIndex))
+	{
+		return StrFullName.Left(DotIndex);
+	}
+	return {};
+}
+
+FString FEdGraphSchemaAction_DataflowVariable::NameFromFullName(FName FullName)
+{
+	const FString StrFullName(FullName.ToString());
+	int32 DotIndex = 0;
+	if (StrFullName.FindChar('|', DotIndex))
+	{
+		return StrFullName.Right(StrFullName.Len() - DotIndex + 1);
+	}
+	return StrFullName;
+}
+
+UDataflow* FEdGraphSchemaAction_DataflowVariable::GetDataflowAsset() const
+{
+	if (DataflowInstanceInterface)
+	{
+		return DataflowInstanceInterface->GetDataflowInstance().GetDataflowAsset();
+	}
+	return nullptr;
+}
+
+bool FEdGraphSchemaAction_DataflowVariable::CanRenameVariable(FText NewName) const
+{
+	if (const UDataflow* DataflowAsset = GetDataflowAsset())
+	{
+		const FString NewVariableName(NewName.ToString());
+		if (NewVariableName.Contains("|"))
+		{
+			return false;
+		}
+		const bool bVariableAlreadyExists = (nullptr != DataflowAsset->Variables.FindPropertyDescByName(FName(NewVariableName)));
+		return !bVariableAlreadyExists;
+	}
+	return false;
+}
+
+void FEdGraphSchemaAction_DataflowVariable::RenameVariable(FText NewName)
+{
+	if (UDataflow* DataflowAsset = GetDataflowAsset())
+	{
+		FName NewVariableName;
+		if (VariableCategory.IsEmpty())
+		{
+			NewVariableName = FName(NewName.ToString());
+		}
+		else
+		{
+			NewVariableName = FName(FString::Format(TEXT("{0}|{1}"), { VariableCategory, NewName.ToString() }));
+		}
+		UE::Dataflow::EditAssetUtils::RenameVariable(DataflowAsset, FullVariableName, NewVariableName);
+	}
+}
+
+void FEdGraphSchemaAction_DataflowVariable::CopyVariableToClipboard()
+{
+	UE::Dataflow::EditAssetUtils::CopyVariableToClipboard(GetDataflowAsset(), GetFullVariableName());
+}
+
+void FEdGraphSchemaAction_DataflowVariable::PasteVariableFromClipboard()
+{
+	UE::Dataflow::EditAssetUtils::PasteVariableFromClipboard(GetDataflowAsset());
+}
+
+void FEdGraphSchemaAction_DataflowVariable::SetVariableType(const FEdGraphPinType& PinType)
+{
+	if (UDataflow* DataflowAsset = GetDataflowAsset())
+	{
+		UE::Dataflow::EditAssetUtils::SetVariableType(DataflowAsset, FullVariableName, PinType);
+	}
+}
+
+void FEdGraphSchemaAction_DataflowVariable::DeleteVariable()
+{
+	UE::Dataflow::EditAssetUtils::DeleteVariable(GetDataflowAsset(), GetFullVariableName());
+}
+
+void FEdGraphSchemaAction_DataflowVariable::DuplicateVariable()
+{
+	UE::Dataflow::EditAssetUtils::DuplicateVariable(GetDataflowAsset(), GetFullVariableName());
+}
+
+/** Execute this action, given the graph and schema, and possibly a pin that we were dragged from. Returns a node that was created by this action (if any). */
+UEdGraphNode* FEdGraphSchemaAction_DataflowVariable::PerformAction(class UEdGraph* ParentGraph, UEdGraphPin* FromPin, const FVector2D Location, bool bSelectNewNode)
+{
+	// TODO : move this code to the FEditableDataflow eventually 
+	UEdGraphNode* NewEdNode = nullptr;
+	if (TSharedPtr<FAssetSchemaAction_Dataflow_CreateNode_DataflowEdNode>  VariableNodeAction = FAssetSchemaAction_Dataflow_CreateNode_DataflowEdNode::CreateAction(ParentGraph, "FGetDataflowVariableNode", FullVariableName))
+	{
+		UDataflowEdNode* DataflowEdNode = Cast<UDataflowEdNode>(VariableNodeAction->PerformAction(ParentGraph, nullptr, Location, false));
+		NewEdNode = DataflowEdNode;
+
+		if (DataflowEdNode)
+		{
+			if (TSharedPtr<FGetDataflowVariableNode> VariableNode = StaticCastSharedPtr<FGetDataflowVariableNode>(DataflowEdNode->GetDataflowNode()))
+			{
+				if (UDataflow* DataflowAsset = GetDataflowAsset())
+				{
+					VariableNode->SetVariable(DataflowAsset, FullVariableName);
+					DataflowEdNode->UpdatePinsFromDataflowNode();
+					ParentGraph->NotifyNodeChanged(DataflowEdNode);
+				}
+			}
+		}
+	}
+	return NewEdNode;
+}
+/////////////////////////////////////////////////////////////////////////////////////////////////////////////////////////
+FGraphSchemaActionDragDropAction_DataflowVariable::FGraphSchemaActionDragDropAction_DataflowVariable()
+	: FGraphSchemaActionDragDropAction()
+{}
+
+TSharedRef<FGraphSchemaActionDragDropAction_DataflowVariable> FGraphSchemaActionDragDropAction_DataflowVariable::New(TSharedPtr<FEdGraphSchemaAction_DataflowVariable> InAction)
+{
+	TSharedRef<FGraphSchemaActionDragDropAction_DataflowVariable> Operation = MakeShareable(new FGraphSchemaActionDragDropAction_DataflowVariable);
+	Operation->VariableAction = InAction;
+	Operation->SourceAction = InAction;
+	Operation->Construct();
+	return Operation;
+}
+
+#undef LOCTEXT_NAMESPACE
\ No newline at end of file
diff --git a/Engine/Plugins/Experimental/Dataflow/Source/DataflowEditor/Private/Dataflow/DataflowInstanceDetails.cpp b/Engine/Plugins/Experimental/Dataflow/Source/DataflowEditor/Private/Dataflow/DataflowInstanceDetails.cpp
--- a/Engine/Plugins/Experimental/Dataflow/Source/DataflowEditor/Private/Dataflow/DataflowInstanceDetails.cpp
+++ b/Engine/Plugins/Experimental/Dataflow/Source/DataflowEditor/Private/Dataflow/DataflowInstanceDetails.cpp
@@ -11,40 +11,44 @@
 #include "PropertyCustomizationHelpers.h"
 #include "ScopedTransaction.h"
 
-#define LOCTEXT_NAMESPACE "DataflowEditor"
+#include "DetailLayoutBuilder.h"
+
+
+#define LOCTEXT_NAMESPACE "DataflowInstanceDetails"
+
 
 namespace UE::Dataflow::Private
 {
 	// variables details ( using FInstancedPropertyBag )
-	class FVariablesDetails : public FPropertyBagInstanceDataDetails
+	class FVariablesOverridesDetails : public FPropertyBagInstanceDataDetails
 	{
 	public:
-		FVariablesDetails(const TSharedPtr<IPropertyHandle> InDataflowInstanceStructProperty, const TSharedPtr<IPropertyHandle> InVariableStructProperty, const TSharedPtr<IPropertyUtilities>& InPropUtils)
+		FVariablesOverridesDetails(const TSharedPtr<IPropertyHandle> InDataflowInstanceStructProperty, const TSharedPtr<IPropertyHandle> InVariableStructProperty, const TSharedPtr<IPropertyUtilities>& InPropUtils)
 			: FPropertyBagInstanceDataDetails(InVariableStructProperty, InPropUtils, /*bInFixedLayout*/ true)
 			, DataflowInstanceStructProperty(InDataflowInstanceStructProperty)
 		{
 		}
 
 	protected:
-		struct FVariablesOverrideProvider : public IPropertyBagOverrideProvider
+		struct FVariablesOverridesProvider : public IPropertyBagOverrideProvider
 		{
-			FVariablesOverrideProvider(FDataflowInstance& InDataflowInstance)
-				: DataflowInstance(InDataflowInstance)
+			FVariablesOverridesProvider(FDataflowVariableOverrides& InDataflowVariableOverride)
+				: DataflowVariableOverride(InDataflowVariableOverride)
 			{
 			}
 
 			virtual bool IsPropertyOverridden(const FGuid PropertyID) const override
 			{
-				return DataflowInstance.IsVariableOverridden(PropertyID);
+				return DataflowVariableOverride.IsVariableOverridden(PropertyID);
 			}
 
 			virtual void SetPropertyOverride(const FGuid PropertyID, const bool bIsOverridden) const override
 			{
-				DataflowInstance.SetVariableOverridden(PropertyID, bIsOverridden);
+				DataflowVariableOverride.SetVariableOverridden(PropertyID, bIsOverridden);
 			}
 
 		private:
-			FDataflowInstance& DataflowInstance;
+			FDataflowVariableOverrides& DataflowVariableOverride;
 		};
 
 		virtual bool HasPropertyOverrides() const override
@@ -70,14 +74,14 @@ namespace UE::Dataflow::Private
 			check(DataflowInstanceStructProperty);
 			DataflowInstanceStructProperty->EnumerateRawData([Func](void* RawData, const int32 /*DataIndex*/, const int32 /*NumDatas*/)
 				{
-					if (FDataflowInstance* DataflowInstance = static_cast<FDataflowInstance*>(RawData))
+					if (FDataflowVariableOverrides* DataflowVariableOverrides = static_cast<FDataflowVariableOverrides*>(RawData))
 					{
-						if (const UDataflow* DataflowAsset = DataflowInstance->GetDataflowAsset())
+						
+						if (const FInstancedPropertyBag* DefaultVariables = DataflowVariableOverrides->GetDefaultVariablesFromAsset())
 						{
-							const FInstancedPropertyBag& DefaultParameters = DataflowAsset->Variables;
-							FInstancedPropertyBag& Parameters = DataflowInstance->GetMutableVariables();
-							FVariablesOverrideProvider OverrideProvider(*DataflowInstance);
-							if (!Func(DefaultParameters, Parameters, OverrideProvider))
+							FInstancedPropertyBag& OVerridenVariables = DataflowVariableOverrides->GetOverridenVariables();
+							FVariablesOverridesProvider OverridesProvider(*DataflowVariableOverrides);
+							if (!Func(*DefaultVariables, OVerridenVariables, OverridesProvider))
 							{
 								return false;
 							}
@@ -99,12 +103,12 @@ namespace UE::Dataflow::Private
 /// 
 ///=============================================================================================================
 /// 
-TSharedRef<IPropertyTypeCustomization> FDataflowInstanceDetails::MakeInstance()
+TSharedRef<IPropertyTypeCustomization> FDataflowVariableOverridesDetails::MakeInstance()
 {
-	return MakeShareable(new FDataflowInstanceDetails);
+	return MakeShareable(new FDataflowVariableOverridesDetails);
 }
 
-void FDataflowInstanceDetails::CustomizeHeader(TSharedRef<IPropertyHandle> InStructPropertyHandle, FDetailWidgetRow& InHeaderRow, IPropertyTypeCustomizationUtils& InStructCustomizationUtils)
+void FDataflowVariableOverridesDetails::CustomizeHeader(TSharedRef<IPropertyHandle> InStructPropertyHandle, FDetailWidgetRow& InHeaderRow, IPropertyTypeCustomizationUtils& InStructCustomizationUtils)
 {
 	// default header row
 	InHeaderRow
@@ -114,9 +118,9 @@ void FDataflowInstanceDetails::CustomizeHeader(TSharedRef<IPropertyHandle> InStr
 		];
 }
 
-void FDataflowInstanceDetails::CustomizeChildren(const TSharedRef<IPropertyHandle> InStructPropertyHandle, IDetailChildrenBuilder& InChildrenBuilder, IPropertyTypeCustomizationUtils& InCustomizationUtils)
+void FDataflowVariableOverridesDetails::CustomizeChildren(const TSharedRef<IPropertyHandle> InStructPropertyHandle, IDetailChildrenBuilder& InChildrenBuilder, IPropertyTypeCustomizationUtils& InCustomizationUtils)
 {
-	const TSharedPtr<IPropertyHandle> VariablesPropertyHandle = InStructPropertyHandle->GetChildHandle(FDataflowInstance::GetVariablePropertyName());
+	const TSharedPtr<IPropertyHandle> VariablesPropertyHandle = InStructPropertyHandle->GetChildHandle(FDataflowVariableOverrides::GetVariablePropertyName());
 	check(VariablesPropertyHandle);
 
 	uint32 NumChildren;
@@ -128,7 +132,7 @@ void FDataflowInstanceDetails::CustomizeChildren(const TSharedRef<IPropertyHandl
 			if (ChildProperty && ChildProperty->GetProperty() == VariablesPropertyHandle->GetProperty())
 			{
 				// customize the variables property
-				const TSharedRef<UE::Dataflow::Private::FVariablesDetails> VariablesDetails = MakeShareable(new  UE::Dataflow::Private::FVariablesDetails(InStructPropertyHandle, VariablesPropertyHandle, InCustomizationUtils.GetPropertyUtilities()));
+				const TSharedRef<UE::Dataflow::Private::FVariablesOverridesDetails> VariablesDetails = MakeShareable(new  UE::Dataflow::Private::FVariablesOverridesDetails(InStructPropertyHandle, VariablesPropertyHandle, InCustomizationUtils.GetPropertyUtilities()));
 				InChildrenBuilder.AddCustomBuilder(VariablesDetails);
 			}
 			else
@@ -141,4 +145,20 @@ void FDataflowInstanceDetails::CustomizeChildren(const TSharedRef<IPropertyHandl
 
 }
 
+////////////////////////////////////////////////////////////////////////////////////////////////////////////////////////
+
+FDataflowInstanceDetailCustomization::FDataflowInstanceDetailCustomization(bool bInOnlyShowVariableOverrides)
+	: bOnlyShowVariableOverrides(bInOnlyShowVariableOverrides)
+{}
+
+void FDataflowInstanceDetailCustomization::CustomizeDetails(IDetailLayoutBuilder& DetailBuilder)
+{
+	if (bOnlyShowVariableOverrides)
+	{
+		DetailBuilder.HideProperty(FDataflowInstance::GetDataflowAssetPropertyName());
+		DetailBuilder.HideProperty(FDataflowInstance::GetDataflowTerminalPropertyName());
+	}
+}
+
+
 #undef LOCTEXT_NAMESPACE
diff --git a/Engine/Plugins/Experimental/Dataflow/Source/DataflowEditor/Private/Dataflow/DataflowMembersWidget.cpp b/Engine/Plugins/Experimental/Dataflow/Source/DataflowEditor/Private/Dataflow/DataflowMembersWidget.cpp
new file mode 100644
--- /dev/null
+++ b/Engine/Plugins/Experimental/Dataflow/Source/DataflowEditor/Private/Dataflow/DataflowMembersWidget.cpp
@@ -0,0 +1,755 @@
+// Copyright Epic Games, Inc. All Rights Reserved.
+
+#include "Dataflow/DataflowMembersWidget.h"
+
+#include "Dataflow/DataflowAssetEditUtils.h"
+#include "Dataflow/DataflowEditor.h"
+#include "Dataflow/DataflowGraphEditor.h"
+#include "Dataflow/DataflowGraphSchemaAction.h"
+#include "Dataflow/DataflowInstance.h"
+#include "Dataflow/DataflowInstanceDetails.h"
+#include "Dataflow/DataflowObject.h"
+#include "Dataflow/DataflowVariablePaletteItem.h"
+#include "Dataflow/DataflowVariableNodes.h"
+#include "Editor/PropertyEditor/Private/PropertyNode.h"
+#include "Framework/Commands/UICommandList.h"
+#include "GraphActionNode.h"
+#include "IDetailsView.h"
+#include "IStructureDetailsView.h"
+#include "DetailsViewArgs.h"
+#include "PropertyEditorModule.h"
+#include "SGraphActionMenu.h"
+#include "Widgets/Input/SComboButton.h"
+#include "Widgets/Input/SSearchBox.h"
+#include "Widgets/SBoxPanel.h"
+
+#define LOCTEXT_NAMESPACE "SDataflowMembersWidget"
+
+/////////////////////////////////////////////////////////////////////////////////////////////////////////////////////////
+
+namespace UE::Dataflow::Private
+{
+	template<typename TActionType>
+	TActionType* CastActionTo(FEdGraphSchemaAction& InAction)
+	{
+		TActionType* TypedAction = nullptr;
+		if (InAction.GetTypeId() == TActionType::StaticGetTypeId())
+		{
+			TypedAction = static_cast<TActionType*>(&InAction);
+		}
+		return TypedAction;
+	}
+
+	template<typename TAction>
+	struct FMemberSection: public SDataflowMembersWidget::ISection
+	{
+		virtual ~FMemberSection() {}
+
+		virtual const FText& GetTitle() const { return Title; }
+		virtual const SDataflowMembersWidget::FButton* GetAddButton() const { return &AddButton; }
+
+		virtual bool CanRequestRename() const override { return true; }
+
+		// todo : when needed we should forward the methods to an action
+		virtual bool CanCopy() const override { return true; }
+		virtual bool CanPaste() const override { return true; }
+		virtual bool CanDuplicate() const override { return true; }
+		virtual bool CanDelete() const override { return true; }
+
+		virtual void OnCopy(FEdGraphSchemaAction& InAction) const
+		{
+			if (TAction* VariableAction = AsTypedAction(InAction))
+			{
+				VariableAction->CopyVariableToClipboard();
+			}
+		}
+		virtual void OnPaste(FEdGraphSchemaAction& InAction) const override
+		{
+			if (TAction* VariableAction = AsTypedAction(InAction))
+			{
+				VariableAction->PasteVariableFromClipboard();
+			}
+		}
+
+		virtual void OnDuplicate(FEdGraphSchemaAction& InAction) const override
+		{
+			if (TAction* VariableAction = AsTypedAction(InAction))
+			{
+				VariableAction->DuplicateVariable();
+			}
+		}
+
+		virtual void OnDelete(FEdGraphSchemaAction& InAction) const override
+		{
+			if (TAction* VariableAction = AsTypedAction(InAction))
+			{
+				VariableAction->DeleteVariable();
+			}
+		}
+
+		static TAction* AsTypedAction(FEdGraphSchemaAction& InAction)
+		{
+			return CastActionTo<TAction>(InAction);
+		}
+
+		FText Title;
+		SDataflowMembersWidget::FButton AddButton;
+	};
+
+	TSharedPtr<SDataflowMembersWidget::ISection> MakeVariablesSection()
+	{
+		using FVariablesSection = FMemberSection<FEdGraphSchemaAction_DataflowVariable>;
+		
+		TSharedPtr<FVariablesSection> VariablesSection = MakeShared<FVariablesSection>();
+
+		VariablesSection->Title = LOCTEXT("Variables", "Variables");
+		VariablesSection->AddButton = 
+		{
+			.Tooltip = LOCTEXT("AddNewVariable", "Add New Variable"),
+			.MetadataTag = TEXT("AddNewVariable"),
+			.Command = FDataflowEditorCommands::Get().AddNewVariable,
+		};
+		return VariablesSection;
+	}
+}
+
+/////////////////////////////////////////////////////////////////////////////////////////////////////////////////////////
+
+void SDataflowMembersWidget::Construct(const FArguments& InArgs, TSharedPtr<SDataflowGraphEditor> InGraphEditor)
+{
+	GraphEditor = InGraphEditor;
+	CacheEditedAsset();
+
+	InitializeCommands();
+	InitializeSections();
+
+	// search box that applies to the entire widget
+	SAssignNew(FilterBox, SSearchBox)
+		.OnTextChanged(this, &SDataflowMembersWidget::OnFilterTextChanged);
+
+	// create the main action list piece of this widget
+	SAssignNew(GraphActionMenu, SGraphActionMenu, false)
+		.OnGetFilterText(this, &SDataflowMembersWidget::GetFilterText)
+		.OnCreateWidgetForAction(this, &SDataflowMembersWidget::OnCreateWidgetForAction)
+		.OnCollectAllActions(this, &SDataflowMembersWidget::CollectAllActions)
+		.OnCollectStaticSections(this, &SDataflowMembersWidget::CollectStaticSections)
+		.OnActionDragged(this, &SDataflowMembersWidget::OnActionDragged)
+		.OnActionDoubleClicked(this, &SDataflowMembersWidget::OnActionDoubleClicked)
+		.OnContextMenuOpening(this, &SDataflowMembersWidget::OnContextMenuOpening)
+		.OnCanRenameSelectedAction(this, &SDataflowMembersWidget::CanRequestRenameOnActionNode)
+		.OnGetSectionTitle(this, &SDataflowMembersWidget::OnGetSectionTitle)
+		.OnGetSectionWidget(this, &SDataflowMembersWidget::OnGetSectionWidget)
+		.OnActionMatchesName(this, &SDataflowMembersWidget::HandleActionMatchesName)
+		.DefaultRowExpanderBaseIndentLevel(1)
+		.AlphaSortItems(false)
+		.UseSectionStyling(true);
+
+	CreateVariableOverrideDetailView();
+
+	FMenuBuilder ViewOptions(true, nullptr);
+
+	// now piece together all the content for this widget
+	ChildSlot
+	[
+		SNew(SVerticalBox)
+
+		// top part 
+		// ( seach bar ) + ( view option menu )
+		+ SVerticalBox::Slot()
+		.AutoHeight()
+		[
+			SNew(SBorder)
+			.Padding(4.0f)
+			.BorderImage(FAppStyle::GetBrush("ToolPanel.GroupBorder"))
+			[
+				SNew(SVerticalBox)
+
+				+ SVerticalBox::Slot()
+				.AutoHeight()
+				[ 
+					SNew(SHorizontalBox)
+
+					// search bar 
+					+ SHorizontalBox::Slot()
+					.FillWidth(1.0f)
+					.VAlign(VAlign_Center)
+					[
+						FilterBox.ToSharedRef()
+					]
+
+					// view option menu 
+					+ SHorizontalBox::Slot()
+					.AutoWidth()
+					.Padding(2, 0, 0, 0)
+					[
+						SNew(SComboButton)
+						.ContentPadding(0.0f)
+						.ComboButtonStyle(&FAppStyle::Get().GetWidgetStyle<FComboButtonStyle>("SimpleComboButton"))
+						.HasDownArrow(false)
+						.AddMetaData<FTagMetaData>(FTagMetaData(TEXT("ViewOptions")))
+						.ButtonContent()
+						[
+							SNew(SImage)
+							.ColorAndOpacity(FSlateColor::UseForeground())
+							.Image(FAppStyle::Get().GetBrush("Icons.Settings"))
+						]
+						.MenuContent()
+						[
+							ViewOptions.MakeWidget()
+						]
+					]
+				]
+			]
+		]
+
+		// actions organized by section ( variables only for now )
+		+ SVerticalBox::Slot()
+		.AutoHeight()
+		[
+			GraphActionMenu.ToSharedRef()
+		]
+
+		+ SVerticalBox::Slot()
+		.FillHeight(1.0f)
+		[
+			OverridesDetailsView->GetWidget().ToSharedRef()
+		]
+	];
+
+	FCoreUObjectDelegates::OnObjectPropertyChanged.AddRaw(this, &SDataflowMembersWidget::OnObjectPropertyChanged);
+}
+
+void SDataflowMembersWidget::CreateVariableOverrideDetailView()
+{
+	FDetailsViewArgs DetailsViewArgs;
+	{
+		DetailsViewArgs.bAllowSearch = false;
+		DetailsViewArgs.NameAreaSettings = FDetailsViewArgs::HideNameArea;
+		DetailsViewArgs.bHideSelectionTip = true;
+		DetailsViewArgs.DefaultsOnlyVisibility = EEditDefaultsOnlyNodeVisibility::Automatic;
+		DetailsViewArgs.bShowOptions = false;
+		DetailsViewArgs.bAllowMultipleTopLevelObjects = true;
+		DetailsViewArgs.bShowKeyablePropertiesOption = false;
+		DetailsViewArgs.bShowModifiedPropertiesOption = false;
+		DetailsViewArgs.bAllowFavoriteSystem = false;
+		DetailsViewArgs.bShowAnimatedPropertiesOption = false;
+	}
+
+	FStructureDetailsViewArgs StructureViewArgs;
+	{
+		StructureViewArgs.bShowObjects = true;
+		StructureViewArgs.bShowAssets = true;
+		StructureViewArgs.bShowClasses = true;
+		StructureViewArgs.bShowInterfaces = true;
+	}
+
+	FPropertyEditorModule& PropertyEditorModule = FModuleManager::Get().GetModuleChecked<FPropertyEditorModule>("PropertyEditor");
+	OverridesDetailsView = PropertyEditorModule.CreateStructureDetailView(DetailsViewArgs, StructureViewArgs, nullptr);
+
+	auto MakeDataflowInstanceDetailCustomizationLambda = []()
+		{
+			constexpr bool bOnlyShowVariableOverrides = true;
+			return MakeShared<FDataflowInstanceDetailCustomization>(bOnlyShowVariableOverrides);
+		};
+	FOnGetDetailCustomizationInstance DataflowInstanceDetailsCustomizationInstance = FOnGetDetailCustomizationInstance::CreateLambda(MakeDataflowInstanceDetailCustomizationLambda);
+	OverridesDetailsView->GetDetailsView()->RegisterInstancedCustomPropertyLayout(FDataflowInstance::StaticStruct(), DataflowInstanceDetailsCustomizationInstance);
+
+	OverridesDetailsView->GetOnFinishedChangingPropertiesDelegate().AddRaw(this, &SDataflowMembersWidget::OverridesDetailsViewFinishedChangingProperties);
+
+	RefreshVariableOverrideDetailView();
+}
+
+void SDataflowMembersWidget::RefreshVariableOverrideDetailView()
+{
+	// now assign the corresponding objects and structures
+	if (IDataflowInstanceInterface* Interface = GetDataflowInstanceInterface())
+	{
+		OverridesDetailsView->GetDetailsView()->SetObject(EditedAsset.Get());
+		TSharedPtr<FStructOnScope> StructOnScope = Interface->GetDataflowInstance().MakeStructOnScope();
+		OverridesDetailsView->SetStructureData(StructOnScope);
+	}
+}
+
+SDataflowMembersWidget::~SDataflowMembersWidget()
+{
+	FCoreUObjectDelegates::OnObjectPropertyChanged.RemoveAll(this);
+}
+
+void SDataflowMembersWidget::InitializeCommands()
+{
+	// initialize command list ( merge from Graph editor )
+	CommandList = MakeShareable(new FUICommandList);
+	if (GraphEditor)
+	{
+		CommandList->Append(GraphEditor->GetCommands().ToSharedRef());
+	}
+	CommandList->MapAction(FGenericCommands::Get().Rename,
+		FExecuteAction::CreateSP(this, &SDataflowMembersWidget::OnRequestRename),
+		FCanExecuteAction::CreateSP(this, &SDataflowMembersWidget::CanRequestRename));
+
+	CommandList->MapAction(FGenericCommands::Get().Copy,
+		FExecuteAction::CreateSP(this, &SDataflowMembersWidget::OnCopy),
+		FCanExecuteAction::CreateSP(this, &SDataflowMembersWidget::CanCopy));
+
+	CommandList->MapAction(FGenericCommands::Get().Cut,
+		FExecuteAction::CreateSP(this, &SDataflowMembersWidget::OnCut),
+		FCanExecuteAction::CreateSP(this, &SDataflowMembersWidget::CanCut));
+
+	CommandList->MapAction(FGenericCommands::Get().Paste,
+		FExecuteAction::CreateSP(this, &SDataflowMembersWidget::OnPaste),
+		FCanExecuteAction(), FIsActionChecked(),
+		FIsActionButtonVisible::CreateSP(this, &SDataflowMembersWidget::CanPaste));
+
+	CommandList->MapAction(FGenericCommands::Get().Duplicate,
+		FExecuteAction::CreateSP(this, &SDataflowMembersWidget::OnDuplicate),
+		FCanExecuteAction(), FIsActionChecked(),
+		FIsActionButtonVisible::CreateSP(this, &SDataflowMembersWidget::CanDuplicate));
+
+	CommandList->MapAction(FGenericCommands::Get().Delete,
+		FExecuteAction::CreateSP(this, &SDataflowMembersWidget::OnDelete),
+		FCanExecuteAction(), FIsActionChecked(),
+		FIsActionButtonVisible::CreateSP(this, &SDataflowMembersWidget::CanDelete));
+}
+
+void SDataflowMembersWidget::InitializeSections()
+{
+	using namespace UE::Dataflow;
+	SectionMap.Reset();
+
+	SectionMap.Emplace((int32)ESchemaActionSectionID::VARIABLES, Private::MakeVariablesSection());
+}
+
+const TSharedPtr<SDataflowMembersWidget::ISection> SDataflowMembersWidget::GetSectionById(int32 SectionId) const
+{
+	if (const TSharedPtr<ISection>* Section = SectionMap.Find(SectionId))
+	{
+		check(*Section); // should never be null - see InitializeSections()
+		return (*Section);
+	}
+	return nullptr;
+}
+
+FReply SDataflowMembersWidget::OnKeyDown(const FGeometry& MyGeometry, const FKeyEvent& InKeyEvent)
+{
+	if (CommandList.IsValid() && CommandList->ProcessCommandBindings(InKeyEvent))
+	{
+		return FReply::Handled();
+	}
+	return FReply::Unhandled();
+}
+
+void SDataflowMembersWidget::CacheEditedAsset()
+{
+	if (!GraphEditor)
+	{
+		return;
+	}
+
+	TSharedPtr<UE::Dataflow::FContext> Context = GraphEditor->GetDataflowContext();
+	if (!Context)
+	{
+		return;
+	}
+
+	const UE::Dataflow::FEngineContext* EngineContext = Context->AsType<UE::Dataflow::FEngineContext>();
+	if (!EngineContext)
+	{
+		return;
+	}
+
+	EditedAsset = EngineContext->Owner;
+}
+
+IDataflowInstanceInterface* SDataflowMembersWidget::GetDataflowInstanceInterface() const
+{
+	return Cast<IDataflowInstanceInterface>(EditedAsset);
+}
+
+UDataflow* SDataflowMembersWidget::GetDataflowAsset() const
+{
+	if (IDataflowInstanceInterface* DataflowInstanceInterface = GetDataflowInstanceInterface())
+	{
+		return DataflowInstanceInterface->GetDataflowInstance().GetDataflowAsset();
+	}
+	return nullptr;
+}
+
+void SDataflowMembersWidget::OverridesDetailsViewFinishedChangingProperties(const FPropertyChangedEvent& InPropertyChangedEvent)
+{
+	const FName VariableName = InPropertyChangedEvent.GetPropertyName();
+	if (UDataflow* DataflowAsset = GetDataflowAsset())
+	{
+		for (const TSharedPtr<FDataflowNode>& Node : DataflowAsset->Dataflow->GetNodes())
+		{
+			if (FGetDataflowVariableNode* VariableNode =  Node->AsType<FGetDataflowVariableNode>())
+			{
+				if (VariableNode->GetVariableName() == VariableName)
+				{
+					VariableNode->Invalidate();
+				}
+			}
+		}
+	}
+}
+
+void SDataflowMembersWidget::OnObjectPropertyChanged(UObject* InObject, FPropertyChangedEvent& InPropertyChangedEvent)
+{
+	static const FName DataflowVariablesPropertyName = GET_MEMBER_NAME_CHECKED(UDataflow, Variables);
+
+	if (InObject)
+	{
+		if (InObject == GetDataflowAsset())
+		{
+			if (InPropertyChangedEvent.GetMemberPropertyName() == DataflowVariablesPropertyName
+				|| InPropertyChangedEvent.GetPropertyName() == DataflowVariablesPropertyName)
+			{
+				Refresh();
+			}
+			RefreshVariableOverrideDetailView();
+		}
+		else if (InObject == EditedAsset)
+		{
+			RefreshVariableOverrideDetailView();
+		}
+	}
+}
+
+void SDataflowMembersWidget::Refresh()
+{
+	if (GraphActionMenu)
+	{
+		GraphActionMenu->RefreshAllActions(/*bPreserveExpansion=*/ true);
+	}
+}
+
+void SDataflowMembersWidget::OnFilterTextChanged(const FText& InFilterText)
+{
+	GraphActionMenu->GenerateFilteredItems(false);
+}
+
+FText SDataflowMembersWidget::GetFilterText() const
+{
+	return FilterBox->GetText();
+}
+
+TSharedRef<SWidget> SDataflowMembersWidget::OnCreateWidgetForAction(FCreateWidgetForActionData* const InCreateData)
+{
+	using namespace UE::Dataflow::Private;
+	return SNew(SDataflowVariablePaletteItem, InCreateData, GraphEditor);
+}
+
+void SDataflowMembersWidget::CollectAllActions(FGraphActionListBuilderBase& OutAllActions)
+{
+	// todo : this should be using the Section map and running a callback from each of the entry 
+	if (IDataflowInstanceInterface* DataflowInstanceInterface = GetDataflowInstanceInterface())
+	{
+		if (const UDataflow* DataflowAsset = DataflowInstanceInterface->GetDataflowInstance().GetDataflowAsset())
+		{
+			if (const UPropertyBag* PropertyBag = DataflowAsset->Variables.GetPropertyBagStruct())
+			{
+				for (const FPropertyBagPropertyDesc& ConfigDesc : PropertyBag->GetPropertyDescs())
+				{
+					TSharedPtr<FEdGraphSchemaAction_DataflowVariable> NewVarAction =
+						MakeShareable(new FEdGraphSchemaAction_DataflowVariable(DataflowInstanceInterface, ConfigDesc));
+					OutAllActions.AddAction(NewVarAction);
+				}
+			}
+		}
+	}
+}
+
+void SDataflowMembersWidget::CollectStaticSections(TArray<int32>& StaticSectionIDs)
+{
+	SectionMap.GetKeys(StaticSectionIDs);
+}
+
+FReply SDataflowMembersWidget::OnActionDragged(const TArray<TSharedPtr<FEdGraphSchemaAction>>& InActions, const FPointerEvent& MouseEvent)
+{
+	TSharedPtr<FEdGraphSchemaAction> InAction(InActions.Num() > 0 ? InActions[0] : nullptr);
+	if (!InAction.IsValid())
+	{
+		return FReply::Unhandled();
+	}
+
+	if (InAction->GetTypeId() == FEdGraphSchemaAction_DataflowVariable::StaticGetTypeId())
+	{
+		TSharedPtr<FEdGraphSchemaAction_DataflowVariable> VariableAction = StaticCastSharedPtr<FEdGraphSchemaAction_DataflowVariable>(InAction);
+		if (UDataflow* DataflowAsset = GetDataflowAsset())
+		{
+			TSharedRef<FGraphSchemaActionDragDropAction_DataflowVariable> DragOperation = FGraphSchemaActionDragDropAction_DataflowVariable::New(VariableAction);
+			return FReply::Handled().BeginDragDrop(DragOperation);
+		}
+	}
+	return FReply::Unhandled();
+}
+
+void SDataflowMembersWidget::OnActionDoubleClicked(const TArray< TSharedPtr<FEdGraphSchemaAction> >& InActions)
+{
+	// TODO : select the first variable node in the graph
+}
+
+TSharedPtr<SWidget> SDataflowMembersWidget::OnContextMenuOpening()
+{
+	const bool bShouldCloseWindowAfterMenuSelection = true;
+	FMenuBuilder MenuBuilder(bShouldCloseWindowAfterMenuSelection, CommandList);
+
+	if (IsAnyActionSelected())
+	{
+		MenuBuilder.AddMenuEntry(FGenericCommands::Get().Rename);
+		MenuBuilder.AddMenuSeparator();
+		MenuBuilder.AddMenuEntry(FGenericCommands::Get().Cut);
+		MenuBuilder.AddMenuEntry(FGenericCommands::Get().Copy);
+		MenuBuilder.AddMenuEntry(FGenericCommands::Get().Paste);
+		MenuBuilder.AddMenuEntry(FGenericCommands::Get().Duplicate);
+		MenuBuilder.AddMenuSeparator();
+		MenuBuilder.AddMenuEntry(FGenericCommands::Get().Delete);
+	}
+	else
+	{
+		MenuBuilder.AddMenuEntry(FGenericCommands::Get().Paste);
+	}
+
+	return MenuBuilder.MakeWidget();
+}
+
+FText SDataflowMembersWidget::OnGetSectionTitle(int32 InSectionID)
+{
+	if (const TSharedPtr<ISection> Section = GetSectionById(InSectionID))
+	{
+		return Section->GetTitle();
+	}
+	return {};
+}
+
+TSharedRef<SWidget> SDataflowMembersWidget::OnGetSectionWidget(TSharedRef<SWidget> RowWidget, int32 InSectionID)
+{
+	TWeakPtr<SWidget> WeakRowWidget = RowWidget;
+	return CreateAddToSectionButton(InSectionID, WeakRowWidget);
+}
+
+TSharedRef<SWidget> SDataflowMembersWidget::CreateAddToSectionButton(int32 InSectionID, TWeakPtr<SWidget> WeakRowWidget)
+{
+	if (const TSharedPtr<ISection> Section = GetSectionById(InSectionID))
+	{
+		if (const FButton* AddButton = Section->GetAddButton())
+		{
+			return
+				SNew(SButton)
+				.ButtonStyle(FAppStyle::Get(), "SimpleButton")
+				.OnClicked(this, &SDataflowMembersWidget::OnAddButtonClickedOnSection, InSectionID)
+				.IsEnabled(this, &SDataflowMembersWidget::CanAddNewElementToSection, InSectionID)
+				.ContentPadding(FMargin(1, 0))
+				.AddMetaData<FTagMetaData>(FTagMetaData(AddButton->MetadataTag))
+				.ToolTipText(AddButton->Tooltip)
+				[
+					SNew(SImage)
+						.Image(FAppStyle::Get().GetBrush("Icons.PlusCircle"))
+						.ColorAndOpacity(FSlateColor::UseForeground())
+				];
+		}
+	}
+	return SNullWidget::NullWidget;
+}
+
+FReply SDataflowMembersWidget::OnAddButtonClickedOnSection(int32 InSectionID)
+{
+	if (const TSharedPtr<ISection> Section = GetSectionById(InSectionID))
+	{
+		if (const FButton* AddButton = Section->GetAddButton())
+		{
+			CommandList->ExecuteAction(AddButton->Command.ToSharedRef());
+		}
+	}
+	return FReply::Handled();
+}
+
+bool SDataflowMembersWidget::CanAddNewElementToSection(int32 InSectionID) const
+{
+	// for now always allowed
+	return true;
+}
+
+bool SDataflowMembersWidget::HandleActionMatchesName(FEdGraphSchemaAction* InAction, const FName& InName) const
+{
+	// todo
+	return false;
+}
+
+void SDataflowMembersWidget::SelectItemByName(const FName& ItemName, ESelectInfo::Type SelectInfo, int32 SectionId/* = INDEX_NONE*/, bool bIsCategory/* = false*/)
+{
+	// Check if the graph action menu is being told to clear
+	if (ItemName == NAME_None)
+	{
+		GraphActionMenu->SelectItemByName(NAME_None);
+	}
+	else
+	{
+		// Attempt to select the item in the main graph action menu
+		const bool bSucceededAtSelecting = GraphActionMenu->SelectItemByName(ItemName, SelectInfo, SectionId, bIsCategory);
+		if (!bSucceededAtSelecting)
+		{ 
+			// We failed to select the item, maybe because it was filtered out?
+			// Reset the item filter and try again (we don't do this first because someone went to the effort of typing
+			// a filter and probably wants to keep it unless it is getting in the way, as it just has)
+			FilterBox->SetText(FText::GetEmpty());
+			GraphActionMenu->SelectItemByName(ItemName, SelectInfo, SectionId, bIsCategory);
+		}
+	}
+}
+
+bool SDataflowMembersWidget::IsAnyActionSelected() const
+{
+	TArray<TSharedPtr<FEdGraphSchemaAction> > SelectedActions;
+	GraphActionMenu->GetSelectedActions(SelectedActions);
+	return (SelectedActions.Num() > 0);
+}
+
+TSharedPtr<FEdGraphSchemaAction> SDataflowMembersWidget::GetFirstSelectedAction() const
+{
+	TArray<TSharedPtr<FEdGraphSchemaAction>> SelectedActions;
+	GraphActionMenu->GetSelectedActions(SelectedActions);
+	return (SelectedActions.Num()) ? SelectedActions[0] : nullptr;
+}
+
+void SDataflowMembersWidget::OnRequestRename()
+{
+	// simple forward to rename request of the action menu
+	GraphActionMenu->OnRequestRenameOnActionNode();
+}
+
+bool SDataflowMembersWidget::CanRequestRename() const
+{
+	return IsAnyActionSelected() 
+		&& GraphActionMenu->CanRequestRenameOnActionNode();
+}
+
+bool SDataflowMembersWidget::CanRequestRenameOnActionNode(TWeakPtr<FGraphActionNode> InSelectedNode) const
+{
+	if (TSharedPtr<FEdGraphSchemaAction> Action = InSelectedNode.Pin()->Action)
+	{
+		if (const TSharedPtr<ISection> Section = GetSectionById(Action->GetSectionID()))
+		{
+			return Section->CanRequestRename();
+		}
+	}
+	return false;
+}
+
+void SDataflowMembersWidget::OnCopy()
+{
+	if (const TSharedPtr<FEdGraphSchemaAction> Action = GetFirstSelectedAction())
+	{
+		if (const TSharedPtr<ISection> Section = GetSectionById(Action->GetSectionID()))
+		{
+			Section->OnCopy(*Action);
+		}
+	}
+}
+
+bool SDataflowMembersWidget::CanCopy() const
+{
+	if (const TSharedPtr<FEdGraphSchemaAction> Action = GetFirstSelectedAction())
+	{
+		if (const TSharedPtr<ISection> Section = GetSectionById(Action->GetSectionID()))
+		{
+			return Section->CanCopy();
+		}
+	}
+	return false;
+}
+
+void SDataflowMembersWidget::OnCut()
+{
+	if (const TSharedPtr<FEdGraphSchemaAction> Action = GetFirstSelectedAction())
+	{
+		if (const TSharedPtr<ISection> Section = GetSectionById(Action->GetSectionID()))
+		{
+			Section->OnCopy(*Action);
+			Section->OnDelete(*Action);
+		}
+	}
+}
+
+bool SDataflowMembersWidget::CanCut() const
+{
+	if (const TSharedPtr<FEdGraphSchemaAction> Action = GetFirstSelectedAction())
+	{
+		if (const TSharedPtr<ISection> Section = GetSectionById(Action->GetSectionID()))
+		{
+			return Section->CanCopy() && Section->CanDelete();
+		}
+	}
+	return false;
+}
+
+void SDataflowMembersWidget::OnPaste()
+{
+	if (const TSharedPtr<FEdGraphSchemaAction> Action = GetFirstSelectedAction())
+	{
+		if (const TSharedPtr<ISection> Section = GetSectionById(Action->GetSectionID()))
+		{
+			Section->OnPaste(*Action);
+		}
+	}
+}
+
+bool SDataflowMembersWidget::CanPaste() const
+{
+	if (const TSharedPtr<FEdGraphSchemaAction> Action = GetFirstSelectedAction())
+	{
+		if (const TSharedPtr<ISection> Section = GetSectionById(Action->GetSectionID()))
+		{
+			return Section->CanPaste();
+		}
+	}
+	// paste is allowed anywhere on the empty space of the widget 
+	return true;
+}
+
+void SDataflowMembersWidget::OnDuplicate()
+{
+	if (const TSharedPtr<FEdGraphSchemaAction> Action = GetFirstSelectedAction())
+	{
+		if (const TSharedPtr<ISection> Section = GetSectionById(Action->GetSectionID()))
+		{
+			Section->OnDuplicate(*Action);
+		}
+	}
+}
+
+bool SDataflowMembersWidget::CanDuplicate() const
+{
+	if (const TSharedPtr<FEdGraphSchemaAction> Action = GetFirstSelectedAction())
+	{
+		if (const TSharedPtr<ISection> Section = GetSectionById(Action->GetSectionID()))
+		{
+			return Section->CanDuplicate();
+		}
+	}
+	return false;
+}
+
+void SDataflowMembersWidget::OnDelete()
+{
+	if (const TSharedPtr<FEdGraphSchemaAction> Action = GetFirstSelectedAction())
+	{
+		if (const TSharedPtr<ISection> Section = GetSectionById(Action->GetSectionID()))
+		{
+			Section->OnDelete(*Action);
+		}
+	}
+}
+
+bool SDataflowMembersWidget::CanDelete() const
+{
+	if (const TSharedPtr<FEdGraphSchemaAction> Action = GetFirstSelectedAction())
+	{
+		if (const TSharedPtr<ISection> Section = GetSectionById(Action->GetSectionID()))
+		{
+			return Section->CanDelete();
+		}
+	}
+	return false;
+}
+
+#undef LOCTEXT_NAMESPACE
\ No newline at end of file
diff --git a/Engine/Plugins/Experimental/Dataflow/Source/DataflowEditor/Private/Dataflow/DataflowVariablePaletteItem.cpp b/Engine/Plugins/Experimental/Dataflow/Source/DataflowEditor/Private/Dataflow/DataflowVariablePaletteItem.cpp
new file mode 100644
--- /dev/null
+++ b/Engine/Plugins/Experimental/Dataflow/Source/DataflowEditor/Private/Dataflow/DataflowVariablePaletteItem.cpp
@@ -0,0 +1,175 @@
+// Copyright Epic Games, Inc. All Rights Reserved.
+
+#include "Dataflow/DataflowVariablePaletteItem.h"
+
+#include "Dataflow/DataflowGraphSchemaAction.h"
+#include "DetailLayoutBuilder.h"
+#include "PropertyBagDetails.h"
+#include "SPinTypeSelector.h"
+#include "Widgets/Text/SInlineEditableTextBlock.h"
+
+void SDataflowVariablePaletteItem::Construct(const FArguments& InArgs, FCreateWidgetForActionData* const InCreateData, TSharedPtr<SDataflowGraphEditor> InEditor)
+{
+	if (InCreateData)
+	{
+		if (InCreateData->Action || InCreateData->Action->GetTypeId() == FEdGraphSchemaAction_DataflowVariable::StaticGetTypeId())
+		{
+			VariableAction = StaticCastSharedPtr<FEdGraphSchemaAction_DataflowVariable>(InCreateData->Action);
+		}
+	}
+	SGraphPaletteItem::Construct(SGraphPaletteItem::FArguments(), InCreateData);
+}
+
+static bool CanHaveMemberVariableOfType(const FEdGraphPinType& PinType)
+{
+	if (PinType.PinCategory == UEdGraphSchema_K2::PC_Exec
+		|| PinType.PinCategory == UEdGraphSchema_K2::PC_Wildcard
+		|| PinType.PinCategory == UEdGraphSchema_K2::PC_MCDelegate
+		|| PinType.PinCategory == UEdGraphSchema_K2::PC_Delegate
+		|| PinType.PinCategory == UEdGraphSchema_K2::PC_Interface)
+	{
+		return false;
+	}
+
+	return true;
+}
+
+namespace UE::Dataflow::Private
+{
+	static bool IsSupportedDataflowVariableType(const FEdGraphPinType& PinType)
+	{
+		const bool bContainerIsAllowed = (false
+			|| PinType.ContainerType == EPinContainerType::None 
+			|| PinType.ContainerType == EPinContainerType::Array
+			);
+
+		const bool bCategoryIsAllowed = (false
+			|| PinType.PinCategory == UEdGraphSchema_K2::PC_Boolean
+			|| PinType.PinCategory == UEdGraphSchema_K2::PC_Class
+			|| PinType.PinCategory == UEdGraphSchema_K2::PC_Int
+			|| PinType.PinCategory == UEdGraphSchema_K2::PC_Int64
+			|| PinType.PinCategory == UEdGraphSchema_K2::PC_Float
+			|| PinType.PinCategory == UEdGraphSchema_K2::PC_Double
+			|| PinType.PinCategory == UEdGraphSchema_K2::PC_Real
+			|| PinType.PinCategory == UEdGraphSchema_K2::PC_Name
+			|| PinType.PinCategory == UEdGraphSchema_K2::PC_Object
+			|| PinType.PinCategory == UEdGraphSchema_K2::PC_String
+			|| PinType.PinCategory == UEdGraphSchema_K2::PC_Struct
+			|| PinType.PinCategory == UEdGraphSchema_K2::AllObjectTypes
+			//|| PinType.PinCategory == UEdGraphSchema_K2::PC_Byte	// some enums are using this category  (see below )
+			//|| PinType.PinCategory == UEdGraphSchema_K2::PC_Enum			// not yet as we do not support enum properly in dataflow graphs
+			//|| PinType.PinCategory == UEdGraphSchema_K2::PC_SoftObject	// Graph do not really use this type
+			//|| PinType.PinCategory == UEdGraphSchema_K2::PC_SoftClass		// Graph do not really use this type
+			);
+
+		return (bContainerIsAllowed && bCategoryIsAllowed);
+	}
+
+	static void GetFilteredVariableTypeTree(TArray<TSharedPtr<UEdGraphSchema_K2::FPinTypeTreeInfo>>& TypeTree, ETypeTreeFilter TypeTreeFilter)
+	{
+		// Collect all the supported types from the default property bag schema 
+		check(GetDefault<UEdGraphSchema_K2>());
+		TArray<TSharedPtr<UEdGraphSchema_K2::FPinTypeTreeInfo>> TempTypeTree;
+		GetDefault<UPropertyBagSchema>()->GetVariableTypeTree(TempTypeTree, TypeTreeFilter);
+
+		// Filter the results
+		for (TSharedPtr<UEdGraphSchema_K2::FPinTypeTreeInfo>& PinType : TempTypeTree)
+		{
+			if (PinType.IsValid())
+			{
+				if (IsSupportedDataflowVariableType(PinType->GetPinType(false)))
+				{
+					// remove children types that are not supported 
+					for (int32 ChildIndex = 0; ChildIndex < PinType->Children.Num(); )
+					{
+						TSharedPtr<UEdGraphSchema_K2::FPinTypeTreeInfo> Child = PinType->Children[ChildIndex];
+						if (Child.IsValid())
+						{
+							const FEdGraphPinType& ChildPinType = Child->GetPinType(/*bForceLoadSubCategoryObject*/false);
+
+							if (!CanHaveMemberVariableOfType(ChildPinType))
+							{
+								if (!IsSupportedDataflowVariableType(ChildPinType))
+								{
+									PinType->Children.RemoveAt(ChildIndex);
+									continue;
+								}
+							}
+						}
+						++ChildIndex;
+					}
+					TypeTree.Add(PinType);
+				}
+			}
+		}
+	}
+}
+TSharedRef<SWidget> SDataflowVariablePaletteItem::CreateTextSlotWidget(FCreateWidgetForActionData* const InCreateData, TAttribute<bool> bIsReadOnlyIn)
+{
+	if (VariableAction == nullptr)
+	{
+		return SNullWidget::NullWidget;
+	}
+
+	const FText VariableName = FText::FromString(VariableAction->GetVariableName());
+
+	TSharedPtr<SInlineEditableTextBlock> EditableTextElement;
+
+	// create a type selection widget
+	TSharedPtr<SWidget> Widget =
+		SNew(SHorizontalBox)
+
+		// search bar 
+		+ SHorizontalBox::Slot()
+		.AutoWidth()
+		.HAlign(HAlign_Left)
+		[
+			SAssignNew(EditableTextElement, SInlineEditableTextBlock)
+				.Text(VariableName)
+				.OnTextCommitted(this, &SDataflowVariablePaletteItem::OnNameTextCommitted)
+				.OnVerifyTextChanged(this, &SDataflowVariablePaletteItem::OnNameTextVerifyChanged)
+		]
+
+		+ SHorizontalBox::Slot()
+		.FillWidth(1.0f)
+		.HAlign(HAlign_Right)
+		[
+			SNew(SBox)
+				.Padding(FMargin(12, 0, 12, 0))
+				[
+					SNew(SPinTypeSelector, FGetPinTypeTree::CreateStatic(UE::Dataflow::Private::GetFilteredVariableTypeTree))
+						.TargetPinType(VariableAction->GetVariableType())
+						.OnPinTypeChanged(this, &SDataflowVariablePaletteItem::OnPinTypeChanged)
+						.Schema(GetDefault<UPropertyBagSchema>())
+						.bAllowArrays(/*bAllowArrays*/true)
+						.TypeTreeFilter(ETypeTreeFilter::None)
+						.Font(IDetailLayoutBuilder::GetDetailFont())
+				]
+		];
+
+	InlineRenameWidget = EditableTextElement.ToSharedRef();
+	InCreateData->OnRenameRequest->BindSP(InlineRenameWidget.Get(), &SInlineEditableTextBlock::EnterEditingMode);
+
+	return Widget.ToSharedRef();
+}
+
+bool SDataflowVariablePaletteItem::OnNameTextVerifyChanged(const FText& InNewText, FText& OutErrorMessage)
+{
+	return VariableAction ? VariableAction->CanRenameVariable(InNewText) : false;
+}
+
+void SDataflowVariablePaletteItem::OnNameTextCommitted(const FText& InNewText, ETextCommit::Type InTextCommit)
+{
+	if (VariableAction)
+	{
+		VariableAction->RenameVariable(InNewText);
+	}
+}
+
+void SDataflowVariablePaletteItem::OnPinTypeChanged(const FEdGraphPinType& PinType)
+{
+	if (VariableAction)
+	{
+		VariableAction->SetVariableType(PinType);
+	}
+}
diff --git a/Engine/Plugins/Experimental/Dataflow/Source/DataflowEditor/Private/Dataflow/DataflowVariablePaletteItem.h b/Engine/Plugins/Experimental/Dataflow/Source/DataflowEditor/Private/Dataflow/DataflowVariablePaletteItem.h
new file mode 100644
--- /dev/null
+++ b/Engine/Plugins/Experimental/Dataflow/Source/DataflowEditor/Private/Dataflow/DataflowVariablePaletteItem.h
@@ -0,0 +1,36 @@
+// Copyright Epic Games, Inc. All Rights Reserved.
+
+#pragma once
+
+#include "SGraphPalette.h"
+
+class SDataflowGraphEditor;
+struct FEdGraphSchemaAction_DataflowVariable;
+struct FEdGraphPinType;
+
+/**
+* Widget for displaying a single variable item in SDataflowMembersWidget
+*/
+class SDataflowVariablePaletteItem : public SGraphPaletteItem
+{
+public:
+	SLATE_BEGIN_ARGS(SDataflowVariablePaletteItem) {}
+	SLATE_END_ARGS()
+
+	void Construct(const FArguments& InArgs, FCreateWidgetForActionData* const InCreateData, TSharedPtr<SDataflowGraphEditor> InEditor);
+
+private:
+	// SGraphPaletteItem Interface
+	virtual TSharedRef<SWidget> CreateTextSlotWidget(FCreateWidgetForActionData* const InCreateData, TAttribute<bool> bIsReadOnly) override;
+	virtual bool OnNameTextVerifyChanged(const FText& InNewText, FText& OutErrorMessage) override;
+	virtual void OnNameTextCommitted(const FText& NewText, ETextCommit::Type InTextCommit) override;
+	// End of SGraphPaletteItem Interface
+
+	void OnPinTypeChanged(const FEdGraphPinType& PinType);
+
+private:
+	TSharedPtr<FEdGraphSchemaAction_DataflowVariable> VariableAction;
+
+	// held for OnRequestRename calls
+	TSharedPtr<SInlineEditableTextBlock> InlineRenameWidget;
+};
diff --git a/Engine/Plugins/Experimental/Dataflow/Source/DataflowEditor/Public/Dataflow/DataflowEditorToolkit.h b/Engine/Plugins/Experimental/Dataflow/Source/DataflowEditor/Public/Dataflow/DataflowEditorToolkit.h
--- a/Engine/Plugins/Experimental/Dataflow/Source/DataflowEditor/Public/Dataflow/DataflowEditorToolkit.h
+++ b/Engine/Plugins/Experimental/Dataflow/Source/DataflowEditor/Public/Dataflow/DataflowEditorToolkit.h
@@ -29,6 +29,8 @@ class UDataflowBaseContent;
 class FDataflowSimulationScene;
 class FDataflowSkeletonView;
 class UDataflowEditor;
+class SDataflowMembersWidget;
+
 namespace UE::Dataflow
 {
 	class FDataflowNodeDetailExtensionHandler;
@@ -112,6 +114,7 @@ private:
 	TSharedRef<SDockTab> SpawnTab_SimulationViewport(const FSpawnTabArgs& Args);
 	TSharedRef<SDockTab> SpawnTab_PreviewScene(const FSpawnTabArgs& Args);
 	TSharedRef<SDockTab> SpawnTab_SimulationVisualization(const FSpawnTabArgs& Args);
+	TSharedRef<SDockTab> SpawnTab_MembersWidget(const FSpawnTabArgs& Args);
 	
 	// FTickableEditorObject interface
 	virtual void Tick(float DeltaTime) override;
@@ -152,12 +155,14 @@ private:
 	static const FName SimulationViewportTabId;
 	static const FName PreviewSceneTabId;
 	static const FName SimulationVisualizationTabId;
+	static const FName MembersWidgetTabId;
 
 	// List of all the widgets shared ptr that will be built in the editor
 	TSharedPtr<SDataflowConstructionViewport> DataflowConstructionViewport;
 	TSharedPtr<SDataflowSimulationViewport> DataflowSimulationViewport;
 	TSharedPtr<SDataflowGraphEditor> GraphEditor;
 	TSharedPtr<IStructureDetailsView> NodeDetailsEditor;
+	TSharedPtr<SDataflowMembersWidget> MembersWidget;
 	TSharedPtr<UE::Dataflow::FDataflowNodeDetailExtensionHandler> NodeDetailsExtensionHandler;
 	TSharedPtr<FDataflowSkeletonView> SkeletonEditorView;
 	TSharedPtr<IDetailsView> AssetDetailsEditor;
@@ -180,6 +185,7 @@ private:
     TSharedPtr<IDetailsView> CreateAssetDetailsEditorWidget(const TArray<UObject*>& ObjectsToEdit);
 	TSharedPtr<SWidget> CreateSimulationVisualizationWidget();
     TSharedPtr<IStructureDetailsView> CreateNodeDetailsEditorWidget(UObject* ObjectToEdit);
+	TSharedPtr<SDataflowMembersWidget> CreateDataflowMembersWidget();
 
 	/** Create the simulation viewport client */
 	void CreateSimulationViewportClient();
diff --git a/Engine/Plugins/Experimental/Dataflow/Source/DataflowEditor/Public/Dataflow/DataflowGraphSchemaAction.h b/Engine/Plugins/Experimental/Dataflow/Source/DataflowEditor/Public/Dataflow/DataflowGraphSchemaAction.h
new file mode 100644
--- /dev/null
+++ b/Engine/Plugins/Experimental/Dataflow/Source/DataflowEditor/Public/Dataflow/DataflowGraphSchemaAction.h
@@ -0,0 +1,126 @@
+// Copyright Epic Games, Inc. All Rights Reserved.
+
+#pragma once
+
+#include "GraphEditorDragDropAction.h"
+#include "DataflowGraphSchemaAction.generated.h"
+
+class IDataflowInstanceInterface;
+class UDataflow;
+struct FPropertyBagPropertyDesc;
+
+namespace UE::Dataflow
+{
+	/* Enums to use when grouping action in widgets 
+	* see SDataflowMembersWidget
+	*/
+	enum class ESchemaActionSectionID: int32
+	{
+		NONE = 0,
+		VARIABLES,
+	};
+}
+
+////////////////////////////////////////////////////////////////////////////////////////////////////////////////////////////////
+
+/**
+* Action that refers to a dataflow variable
+* It is used to display the variables in the SDataflowMembersWidget
+*/
+USTRUCT()
+struct FEdGraphSchemaAction_DataflowVariable : public FEdGraphSchemaAction
+{
+	GENERATED_USTRUCT_BODY()
+
+public:
+	FEdGraphSchemaAction_DataflowVariable();
+
+	FEdGraphSchemaAction_DataflowVariable(IDataflowInstanceInterface* InDataflowInstanceInterface, const FPropertyBagPropertyDesc& PropertyDesc);
+
+	// Simple type info
+	static FName StaticGetTypeId() { static FName Type("FEdGraphSchemaAction_DataflowVariable"); return Type; }
+	virtual FName GetTypeId() const override { return StaticGetTypeId(); }
+
+	const FString& GetVariableName() const
+	{
+		return VariableName;
+	}
+
+	FName GetFullVariableName() const
+	{
+		return FullVariableName;
+	}
+
+	const FEdGraphPinType& GetVariableType() const
+	{
+		return VariableType;
+	}
+
+	// FEdGraphSchemaAction interface
+	virtual bool IsAVariable() const { return true; }
+	virtual UEdGraphNode* PerformAction(class UEdGraph* ParentGraph, UEdGraphPin* FromPin, const FVector2D Location, bool bSelectNewNode = true) override;
+	// End of FEdGraphSchemaAction interface
+
+	bool CanRenameVariable(FText NewName) const;
+	void RenameVariable(FText NewName);
+	void SetVariableType(const FEdGraphPinType& PinType);
+
+	void CopyVariableToClipboard();
+	void PasteVariableFromClipboard();
+	void DeleteVariable();
+	void DuplicateVariable();
+
+	UDataflow* GetDataflowAsset() const;
+
+private:
+	static FString CategoryFromFullName(FName FullName);
+	static FString NameFromFullName(FName FullName);
+
+private:
+	/**
+	* Name of the variable
+	* the name can contain a category is using a separating |
+	* example : Category.VariableName
+	*/
+	FName FullVariableName;
+
+	/**
+	* Variable short name ( extracted from FullName )
+	*/
+	FString VariableName;
+
+	/**
+	* Variable category name ( extracted from FullName )
+	*/
+	FString VariableCategory;
+
+	/**
+	* type of the variable (as a Pin type)
+	*/
+	FEdGraphPinType VariableType;
+
+	/**
+	* interface from where to read the variable information
+	* this allow access to the dataflow asset as well as the override from the asset using this dataflow asset
+	*/
+	IDataflowInstanceInterface* DataflowInstanceInterface = nullptr;;
+};
+
+////////////////////////////////////////////////////////////////////////////////////////////////////////////////////////////////
+
+/**
+* Dataflow variable drag and drop action
+* see FEdGraphSchemaAction_DataflowVariable
+*/
+struct FGraphSchemaActionDragDropAction_DataflowVariable : public FGraphSchemaActionDragDropAction
+{
+public:
+	DRAG_DROP_OPERATOR_TYPE(FGraphSchemaActionDragDropAction_DataflowVariable, FGraphSchemaActionDragDropAction)
+
+	static TSharedRef<FGraphSchemaActionDragDropAction_DataflowVariable> New(TSharedPtr<FEdGraphSchemaAction_DataflowVariable> InAction);
+
+private:
+	FGraphSchemaActionDragDropAction_DataflowVariable();
+
+	TSharedPtr<FEdGraphSchemaAction_DataflowVariable> VariableAction;
+};
diff --git a/Engine/Plugins/Experimental/Dataflow/Source/DataflowEditor/Public/Dataflow/DataflowInstanceDetails.h b/Engine/Plugins/Experimental/Dataflow/Source/DataflowEditor/Public/Dataflow/DataflowInstanceDetails.h
--- a/Engine/Plugins/Experimental/Dataflow/Source/DataflowEditor/Public/Dataflow/DataflowInstanceDetails.h
+++ b/Engine/Plugins/Experimental/Dataflow/Source/DataflowEditor/Public/Dataflow/DataflowInstanceDetails.h
@@ -3,15 +3,17 @@
 #pragma once
 
 #include "IPropertyTypeCustomization.h"
+#include "IDetailCustomization.h"
+#include "IDetailPropertyChildrenCustomizationHandler.h"
 
 class IPropertyHandle;
 class USmartObjectDefinition;
 class IPropertyUtilities;
 
 /**
- * Type customization for FDataflowInstance.
+ * Type customization for FDataflowVariableOverrides.
  */
-class FDataflowInstanceDetails : public IPropertyTypeCustomization
+class FDataflowVariableOverridesDetails : public IPropertyTypeCustomization
 {
 public:
 	/** Makes a new instance of this detail layout class for a specific detail view requesting it */
@@ -20,4 +22,17 @@ public:
 	/** IPropertyTypeCustomization interface */
 	virtual void CustomizeHeader(TSharedRef<IPropertyHandle> InStructPropertyHandle, FDetailWidgetRow& InHeaderRow, IPropertyTypeCustomizationUtils& InCustomizationUtils) override;
 	virtual void CustomizeChildren(TSharedRef<IPropertyHandle> InStructPropertyHandle, IDetailChildrenBuilder& InChildrenBuilder, IPropertyTypeCustomizationUtils& InCustomizationUtils) override;
-};
\ No newline at end of file
+};
+
+class FDataflowInstanceDetailCustomization : public IDetailCustomization
+{
+public:
+	/** Makes a new instance of this detail layout class for a specific detail view requesting it */
+	FDataflowInstanceDetailCustomization(bool bOnlyShowVariableOverrides);
+
+	/** IDetailCustomization interface */
+	virtual void CustomizeDetails(IDetailLayoutBuilder& DetailBuilder) override;
+
+private:
+	bool bOnlyShowVariableOverrides = false;
+};
diff --git a/Engine/Plugins/Experimental/Dataflow/Source/DataflowEditor/Public/Dataflow/DataflowMembersWidget.h b/Engine/Plugins/Experimental/Dataflow/Source/DataflowEditor/Public/Dataflow/DataflowMembersWidget.h
new file mode 100644
--- /dev/null
+++ b/Engine/Plugins/Experimental/Dataflow/Source/DataflowEditor/Public/Dataflow/DataflowMembersWidget.h
@@ -0,0 +1,137 @@
+// Copyright Epic Games, Inc. All Rights Reserved.
+
+#pragma once
+
+#include "Widgets/SCompoundWidget.h"
+
+class UDataflow;
+class UDataflowEditorVariables;
+class FUICommandList;
+class FUICommandInfo;
+class SSearchBox;
+class SGraphActionMenu;
+class SDataflowGraphEditor;
+class IDataflowInstanceInterface;
+class IStructureDetailsView;
+
+struct FCreateWidgetForActionData;
+struct FGraphActionListBuilderBase;
+struct FEdGraphSchemaAction;
+struct FGraphActionNode;
+struct FPropertyBagPropertyDesc;
+struct FDataflowVariableClassFilter;
+
+/**
+*
+* Widget to interact with Variables, subgraphs in Dataflow editors
+*
+*/
+class SDataflowMembersWidget: public SCompoundWidget
+{
+public:
+	SLATE_BEGIN_ARGS(SDataflowMembersWidget) {}
+	SLATE_END_ARGS()
+
+	void Construct(const FArguments& InArgs, TSharedPtr<SDataflowGraphEditor> InGraphEditor);
+	~SDataflowMembersWidget();
+
+public:
+	struct FButton
+	{
+		FText Tooltip;
+		FName MetadataTag;
+		TSharedPtr<FUICommandInfo> Command;
+	};
+
+	struct ISection
+	{
+		virtual const FText& GetTitle() const = 0;
+		virtual const FButton* GetAddButton() const = 0;
+
+		virtual bool CanRequestRename() const = 0;
+
+		virtual bool CanCopy() const = 0;
+		virtual bool CanPaste() const = 0;
+		virtual bool CanDuplicate() const = 0;
+		virtual bool CanDelete() const = 0;
+
+		virtual void OnCopy(FEdGraphSchemaAction&) const = 0;
+		virtual void OnPaste(FEdGraphSchemaAction&) const = 0;
+		virtual void OnDuplicate(FEdGraphSchemaAction&) const = 0;
+		virtual void OnDelete(FEdGraphSchemaAction&) const = 0;
+	};
+
+private:
+	/* SWidget interface */
+	virtual FReply OnKeyDown(const FGeometry& MyGeometry, const FKeyEvent& InKeyEvent);
+
+	void OnObjectPropertyChanged(UObject* InObject, FPropertyChangedEvent& InPropertyChangedEvent);
+	void OnFilterTextChanged(const FText& InFilterText);
+	FText GetFilterText() const;
+	TSharedRef<SWidget> OnCreateWidgetForAction(FCreateWidgetForActionData* const InCreateData);
+	void CollectAllActions(FGraphActionListBuilderBase& OutAllActions);
+	void CollectStaticSections(TArray<int32>& StaticSectionIDs);
+	FReply OnActionDragged(const TArray<TSharedPtr<FEdGraphSchemaAction>>& InActions, const FPointerEvent& MouseEvent);
+	void OnActionDoubleClicked(const TArray< TSharedPtr<FEdGraphSchemaAction> >& InActions);
+	TSharedPtr<SWidget> OnContextMenuOpening();
+	bool CanRequestRenameOnActionNode(TWeakPtr<FGraphActionNode> InSelectedNode) const;
+	FText OnGetSectionTitle(int32 InSectionID);
+	TSharedRef<SWidget> OnGetSectionWidget(TSharedRef<SWidget> RowWidget, int32 InSectionID);
+	bool HandleActionMatchesName(FEdGraphSchemaAction* InAction, const FName& InName) const;
+
+	TSharedRef<SWidget> CreateAddToSectionButton(int32 InSectionID, TWeakPtr<SWidget> WeakRowWidget);
+	FReply OnAddButtonClickedOnSection(int32 InSectionID);
+	bool CanAddNewElementToSection(int32 InSectionID) const;
+
+	void SelectItemByName(const FName& ItemName, ESelectInfo::Type SelectInfo = ESelectInfo::Direct, int32 SectionId = INDEX_NONE, bool bIsCategory = false);
+	bool IsAnyActionSelected() const;
+	TSharedPtr<FEdGraphSchemaAction> GetFirstSelectedAction() const;
+
+	const TSharedPtr<ISection> GetSectionById(int32 SectionId) const;
+
+	void CacheEditedAsset();
+	UDataflow* GetDataflowAsset() const;
+	IDataflowInstanceInterface* GetDataflowInstanceInterface() const;
+
+	// command functions
+	void OnRequestRename();
+	bool CanRequestRename() const;
+	void OnCopy();
+	bool CanCopy() const;
+	void OnCut();
+	bool CanCut() const;
+	void OnPaste();
+	bool CanPaste() const;
+	void OnDuplicate();
+	bool CanDuplicate() const;
+	void OnDelete();
+	bool CanDelete() const;
+
+	// UI refresh 
+	void Refresh();
+
+	void InitializeCommands();
+	void InitializeSections();
+	void CreateVariableOverrideDetailView();
+	void RefreshVariableOverrideDetailView();
+	void OverridesDetailsViewFinishedChangingProperties(const FPropertyChangedEvent& InPropertyChangedEvent);
+
+private:
+	TSharedPtr<SDataflowGraphEditor> GraphEditor;
+
+	/** asset being edited by the graph - maybe a geometry collection with a dataflow asset bound to it */
+	TWeakObjectPtr<UObject> EditedAsset;
+
+	/** List of UI Commands for this scope */
+	TSharedPtr<FUICommandList> CommandList;
+
+	/** The filter box that handles filtering for graph action menu. */
+	TSharedPtr<SSearchBox> FilterBox;
+
+	/** Graph Action Menu for displaying all our variables and sub-graphs */
+	TSharedPtr<SGraphActionMenu> GraphActionMenu;
+
+	TMap<int32, TSharedPtr<ISection>> SectionMap;
+
+	TSharedPtr<IStructureDetailsView> OverridesDetailsView;
+};
diff --git a/Engine/Source/Runtime/Experimental/Dataflow/Core/Private/Dataflow/DataflowConnection.cpp b/Engine/Source/Runtime/Experimental/Dataflow/Core/Private/Dataflow/DataflowConnection.cpp
--- a/Engine/Source/Runtime/Experimental/Dataflow/Core/Private/Dataflow/DataflowConnection.cpp
+++ b/Engine/Source/Runtime/Experimental/Dataflow/Core/Private/Dataflow/DataflowConnection.cpp
@@ -128,6 +128,10 @@ bool FDataflowConnection::IsSafeToTryChangingType() const
 	{
 		return false;
 	}
+	if (bLockType)
+	{
+		return false;
+	}
 	if (IsConnected())
 	{
 		return false;
diff --git a/Engine/Source/Runtime/Experimental/Dataflow/Core/Public/Dataflow/DataflowConnection.h b/Engine/Source/Runtime/Experimental/Dataflow/Core/Public/Dataflow/DataflowConnection.h
--- a/Engine/Source/Runtime/Experimental/Dataflow/Core/Public/Dataflow/DataflowConnection.h
+++ b/Engine/Source/Runtime/Experimental/Dataflow/Core/Public/Dataflow/DataflowConnection.h
@@ -146,6 +146,9 @@ protected:
 	bool bCanHidePin:1 = false;
 	bool bPinIsHidden:1 = false;
 
+	/** Lock the type from being changed */
+	bool bLockType:1 = false;
+
 	// Name of the group dependecy the connection is part of 
 	// This applies to anytype connections
 	// All other connection with the same group will inherit type from each other
@@ -213,6 +216,9 @@ public:
 		return UE::Dataflow::FConnectionKey(GetOffset(), GetContainerIndex(), GetContainerElementOffset());
 	}
 
+	void LockType() { bLockType = true; }
+	void UnlockType() { bLockType = false; }
+
 	FName GetType() const { return Type; }
 	FName GetOriginalType() const { return OriginalType; }
 
diff --git a/Engine/Source/Runtime/Experimental/Dataflow/Engine/Private/Dataflow/DataflowEngineModule.cpp b/Engine/Source/Runtime/Experimental/Dataflow/Engine/Private/Dataflow/DataflowEngineModule.cpp
--- a/Engine/Source/Runtime/Experimental/Dataflow/Engine/Private/Dataflow/DataflowEngineModule.cpp
+++ b/Engine/Source/Runtime/Experimental/Dataflow/Engine/Private/Dataflow/DataflowEngineModule.cpp
@@ -6,6 +6,7 @@
 #include "Modules/ModuleManager.h"
 
 #include "Dataflow/DataflowTextureAssetNodes.h"
+#include "Dataflow/DataflowVariableNodes.h"
 
 /**
  * The public interface to this module
@@ -19,6 +20,7 @@ public:
 	virtual void StartupModule() override
 	{
 		UE::Dataflow::RegisterTextureAssetNodes();
+		UE::Dataflow::RegisterVariableNodes();
 	}
 
 	virtual void ShutdownModule() override { }
diff --git a/Engine/Source/Runtime/Experimental/Dataflow/Engine/Private/Dataflow/DataflowInstance.cpp b/Engine/Source/Runtime/Experimental/Dataflow/Engine/Private/Dataflow/DataflowInstance.cpp
--- a/Engine/Source/Runtime/Experimental/Dataflow/Engine/Private/Dataflow/DataflowInstance.cpp
+++ b/Engine/Source/Runtime/Experimental/Dataflow/Engine/Private/Dataflow/DataflowInstance.cpp
@@ -108,64 +108,214 @@ namespace UE::Dataflow::InstanceUtils
 	}
 }
 
+
 //---------------------------------------------------------------------------
-// FDataflowInstance
+// FDataflowVariableOverrides
 //---------------------------------------------------------------------------
-
-FDataflowInstance::FDataflowInstance(UObject* InOwner)
+FDataflowVariableOverrides::FDataflowVariableOverrides(FDataflowInstance* InOwner)
 	: Owner(InOwner)
 {
-#if WITH_EDITOR
-	// let listen to property change on the owner so we can listen to our own changes ( because we are a struct )
+}
+
+const FInstancedPropertyBag* FDataflowVariableOverrides::GetDefaultVariablesFromAsset() const 
+{
 	if (Owner)
 	{
-		FCoreUObjectDelegates::OnObjectPropertyChanged.AddWeakLambda(Owner.Get(), [this](UObject* Obj, struct FPropertyChangedEvent& Event)
+		if (const UDataflow* DataflowAsset = Owner->GetDataflowAsset())
+		{
+			return &DataflowAsset->Variables;
+		}
+	}
+	return nullptr;
+}
+
+FInstancedPropertyBag& FDataflowVariableOverrides::GetOverridenVariables()
+{
+	return Variables;
+}
+
+void FDataflowVariableOverrides::RemoveAllVariables()
+{
+	Variables.Reset();
+	OverriddenVariableGuids.Reset();
+}
+
+void FDataflowVariableOverrides::SyncVariables()
+{
+	if (const FInstancedPropertyBag* DefaultVariables = GetDefaultVariablesFromAsset())
+	{
+		// In editor builds, sync with overrides.
+		Variables.MigrateToNewBagInstanceWithOverrides(*DefaultVariables, OverriddenVariableGuids);
+
+		// clean up differences
+		RemoveOverridenVariablesNotInDataflowAsset();
+	}
+	else
+	{
+		RemoveAllVariables();
+	}
+}
+
+void FDataflowVariableOverrides::RemoveOverridenVariablesNotInDataflowAsset()
+{
+	// Remove overriden variables that do not exists in the dataflow asset
+	if (!OverriddenVariableGuids.IsEmpty())
+	{
+		if (const UPropertyBag* Bag = Variables.GetPropertyBagStruct())
+		{
+			for (TArray<FGuid>::TIterator It = OverriddenVariableGuids.CreateIterator(); It; ++It)
 			{
-				if (IDataflowInstanceInterface* Interface = Cast<IDataflowInstanceInterface>(Owner))
+				if (!Bag->FindPropertyDescByID(*It))
 				{
-					if (Obj == Interface->GetDataflowInstance().GetDataflowAsset())
-					{
-						OnOwnerPostEditChangeProperty(Event);
-					}
+					It.RemoveCurrentSwap();
 				}
-			});
+			}
+		}
+		else
+		{
+			OverriddenVariableGuids.Reset();
+		}
 	}
-#endif
+}
+
+bool FDataflowVariableOverrides::HasVariable(FName VariableName) const
+{
+	return (Variables.FindPropertyDescByName(VariableName) != nullptr);
+}
+
+const FInstancedPropertyBag& FDataflowVariableOverrides::GetVariables() const
+{
+	return Variables;
+}
+
+const FInstancedPropertyBag& FDataflowVariableOverrides::GetVariables()
+{
+	return Variables;
+}
+
+bool FDataflowVariableOverrides::IsVariableOverridden(const FGuid PropertyID) const
+{
+	return OverriddenVariableGuids.Contains(PropertyID);
+}
+
+bool FDataflowVariableOverrides::IsVariableOverridden(FName VariableName) const
+{
+	if (const FPropertyBagPropertyDesc* Desc = Variables.FindPropertyDescByName(VariableName))
+	{
+		return IsVariableOverridden(Desc->ID);
+	}
+	return false;
+}
+
+void FDataflowVariableOverrides::SetVariableOverridden(const FGuid PropertyID, const bool bIsOverridden)
+{
+	if (bIsOverridden)
+	{
+		OverriddenVariableGuids.AddUnique(PropertyID);
+	}
+	else
+	{
+		OverriddenVariableGuids.Remove(PropertyID);
+	}
+	SyncVariables();
+}
+
+
+FName FDataflowVariableOverrides::GetVariablePropertyName()
+{
+	return GET_MEMBER_NAME_CHECKED(FDataflowVariableOverrides, Variables);
 }
 
 #if WITH_EDITOR
-void FDataflowInstance::OnOwnerPostEditChangeProperty(struct FPropertyChangedEvent& PropertyChangedEvent)
+void FDataflowVariableOverrides::OnOwnerPostEditChangeProperty(struct FPropertyChangedEvent& PropertyChangedEvent)
 {
 	const FName PropertyName = PropertyChangedEvent.GetPropertyName();
 	const FName MemberPropertyName = PropertyChangedEvent.GetMemberPropertyName();
-	
-	if (PropertyName == GET_MEMBER_NAME_STRING_CHECKED(UDataflow, Variables))
+
+	const FName VariablePropertyName = GetVariablePropertyName();
+
+	if (PropertyName == VariablePropertyName || MemberPropertyName == VariablePropertyName)
 	{
 		// Add / remove / move properties / change type
 		// resync variable from the dataflow asset
 		SyncVariables();
 	}
-	else if (Variables.FindPropertyDescByName(MemberPropertyName))
+	else if (HasVariable(MemberPropertyName))
 	{
 		// Specific variable change , nothing to sync 
 		// nothing to do now ( potential may need to re-evaluate the dataflow )
 		SyncVariables();
 	}
 }
+
+void FDataflowVariableOverrides::OnDataflowVariablesChanged(const UDataflow* DataflowAsset, FName VariableName)
+{
+	SyncVariables();
+}
 #endif
 
-FName FDataflowInstance::GetVariablePropertyName()
+//---------------------------------------------------------------------------
+// FDataflowInstance
+//---------------------------------------------------------------------------
+
+FDataflowInstance::FDataflowInstance(UObject* InOwner)
+	: VariableOverrides(this)
+	, Owner(InOwner)
 {
-	return GET_MEMBER_NAME_CHECKED(FDataflowInstance, Variables);
+#if WITH_EDITOR
+	// let listen to property change on the owner so we can listen to our own changes ( because we are a struct )
+	if (Owner)
+	{
+		FDataflowAssetDelegates::OnVariablesChanged.AddWeakLambda(Owner.Get(), 
+			[this](const UDataflow* InDataflowAsset, FName InVariableName)
+			{
+				if (InDataflowAsset == DataflowAsset)
+				{
+					VariableOverrides.OnDataflowVariablesChanged(InDataflowAsset, InVariableName);
+				}
+			});
+	}
+#endif
 }
 
+FDataflowInstance::~FDataflowInstance()
+{
+	if (Owner)
+	{
+		FDataflowAssetDelegates::OnVariablesChanged.RemoveAll(Owner.Get());
+	}
+}
+
+#if WITH_EDITOR
+void FDataflowInstance::OnOwnerPostEditChangeProperty(struct FPropertyChangedEvent& PropertyChangedEvent)
+{
+	VariableOverrides.OnOwnerPostEditChangeProperty(PropertyChangedEvent);
+	// should check for terminal node / dataflow asset changes ?
+}
+
+TSharedPtr<FStructOnScope> FDataflowInstance::MakeStructOnScope() const
+{
+	return MakeShared<FStructOnScope>(FDataflowInstance::StaticStruct(), (uint8*)(this));
+}
+#endif
+
+FName FDataflowInstance::GetDataflowTerminalPropertyName()
+{
+	return GET_MEMBER_NAME_CHECKED(FDataflowInstance, DataflowTerminal);
+}
+
+/** Get the dataflow asset member property name */
+FName FDataflowInstance::GetDataflowAssetPropertyName()
+{
+	return GET_MEMBER_NAME_CHECKED(FDataflowInstance, DataflowAsset);
+}
 
 void FDataflowInstance::SetDataflowAsset(UDataflow* InDataflowAsset)
 {
 	if (DataflowAsset != InDataflowAsset)
 	{
 		DataflowAsset = InDataflowAsset;
-		SyncVariables();
+		VariableOverrides.SyncVariables();
 	}
 }
 
@@ -187,71 +337,25 @@ FName FDataflowInstance::GetDataflowTerminal() const
 
 const FInstancedPropertyBag& FDataflowInstance::GetVariables() const
 {
-	return Variables;
+	return VariableOverrides.GetVariables();
 }
 
 const FInstancedPropertyBag& FDataflowInstance::GetVariables()
 {
-	return Variables;
+	return VariableOverrides.GetVariables();
 }
 
-FInstancedPropertyBag& FDataflowInstance::GetMutableVariables()
+FDataflowVariableOverrides& FDataflowInstance::GetVariableOverrides()
 {
-	return Variables;
+	return VariableOverrides;
 }
 
-bool FDataflowInstance::IsVariableOverridden(const FGuid PropertyID) const
+const FDataflowVariableOverrides& FDataflowInstance::GetVariableOverrides() const
 {
-	return OverriddenVariables.Contains(PropertyID);
-}
-
-void FDataflowInstance::SetVariableOverridden(const FGuid PropertyID, const bool bIsOverridden)
-{
-	if (bIsOverridden)
-	{
-		OverriddenVariables.AddUnique(PropertyID);
-	}
-	else
-	{
-		OverriddenVariables.Remove(PropertyID);
-	}
-	SyncVariables();
+	return VariableOverrides;
 }
 
 void FDataflowInstance::SyncVariables()
 {
-	if (DataflowAsset == nullptr)
-	{
-		Variables.Reset();
-	}
-	else
-	{
-		// In editor builds, sync with overrides.
-		Variables.MigrateToNewBagInstanceWithOverrides(DataflowAsset->Variables, OverriddenVariables);
-
-		// clean up differences
-		RemoveOverridenVariablesNotInDataflowAsset();
-	}
-}
-
-void FDataflowInstance::RemoveOverridenVariablesNotInDataflowAsset()
-{
-	// Remove overriden variables that do not exists in the dataflow asset
-	if (!OverriddenVariables.IsEmpty())
-	{
-		if (const UPropertyBag* Bag = Variables.GetPropertyBagStruct())
-		{
-			for (TArray<FGuid>::TIterator It = OverriddenVariables.CreateIterator(); It; ++It)
-			{
-				if (!Bag->FindPropertyDescByID(*It))
-				{
-					It.RemoveCurrentSwap();
-				}
-			}
-		}
-		else
-		{
-			OverriddenVariables.Reset();
-		}
-	}
+	VariableOverrides.SyncVariables();
 }
diff --git a/Engine/Source/Runtime/Experimental/Dataflow/Engine/Private/Dataflow/DataflowVariableNodes.cpp b/Engine/Source/Runtime/Experimental/Dataflow/Engine/Private/Dataflow/DataflowVariableNodes.cpp
new file mode 100644
--- /dev/null
+++ b/Engine/Source/Runtime/Experimental/Dataflow/Engine/Private/Dataflow/DataflowVariableNodes.cpp
@@ -0,0 +1,342 @@
+// Copyright Epic Games, Inc. All Rights Reserved.
+
+#include "Dataflow/DataflowVariableNodes.h"
+#include "Dataflow/DataflowNodeFactory.h"
+#include "Dataflow/DataflowObject.h"
+#include "Dataflow/DataflowInstance.h"
+#include "Dataflow/DataflowNodeParameters.h"
+#include "Dataflow/DataflowObjectInterface.h"
+
+namespace UE::Dataflow
+{
+	void RegisterVariableNodes()
+	{
+		DATAFLOW_NODE_REGISTER_CREATION_FACTORY(FGetDataflowVariableNode);
+	}
+}
+
+
+FGetDataflowVariableNode::FGetDataflowVariableNode(const UE::Dataflow::FNodeParameters& InParam, FGuid InGuid)
+	: FDataflowNode(InParam, InGuid)
+{
+	RegisterOutputConnection(&Value);
+
+	UDataflow* DataflowAsset = Cast<UDataflow>(InParam.OwningObject);
+	WeakDataflowPtr = DataflowAsset;
+}
+
+FGetDataflowVariableNode::~FGetDataflowVariableNode()
+{
+#if WITH_EDITOR
+	FCoreUObjectDelegates::OnObjectPropertyChanged.RemoveAll(this);
+#endif
+}
+
+void FGetDataflowVariableNode::PostSerialize(const FArchive& Ar)
+{
+	if (Ar.IsLoading())
+	{
+		// make sure the variable is up to date
+		if (!VariableName.IsNone())
+		{
+			if (UDataflow* DataflowAsset = WeakDataflowPtr.Get())
+			{
+				SetVariable(DataflowAsset, VariableName);
+			}
+		}
+	}
+}
+
+void FGetDataflowVariableNode::OnObjectPropertyChanged(UObject* InObject, FPropertyChangedEvent& InPropertyChangedEvent)
+{
+	if (UDataflow* DataflowAsset = WeakDataflowPtr.Get())
+	{
+		if (InObject == DataflowAsset)
+		{
+			const FName VariablePropertyBagMemberName = GET_MEMBER_NAME_CHECKED(UDataflow, Variables);
+			if (InPropertyChangedEvent.GetMemberPropertyName() == VariablePropertyBagMemberName)
+			{
+				const FName PropertyName = InPropertyChangedEvent.GetPropertyName();
+				if (PropertyName == VariableName || PropertyName.IsNone())
+				{
+					// reset the variable - how do we handle renames ? 
+					SetVariable(DataflowAsset, VariableName);
+				}
+			}
+		}
+	}
+}
+
+void FGetDataflowVariableNode::SetVariable(UDataflow* Dataflow, FName InVariableName)
+{
+#if WITH_EDITOR
+	FCoreUObjectDelegates::OnObjectPropertyChanged.RemoveAll(this);
+#endif
+
+	WeakDataflowPtr = Dataflow;
+	EvaluateFunction = nullptr;
+
+	VariablePropertyBag.Reset();
+
+	if (Dataflow)
+	{
+		if (const FPropertyBagPropertyDesc* Desc = Dataflow->Variables.FindPropertyDescByName(InVariableName))
+		{
+			VariableName = InVariableName;
+
+			VariablePropertyBag.AddProperties({ *Desc });
+			VariablePropertyBag.SetValue(VariableName, Desc->CachedProperty, Dataflow->Variables.GetValue().GetMemory());
+
+			UpdateOutputTypes(*Desc);
+			Invalidate();
+		}
+
+#if WITH_EDITOR
+		FCoreUObjectDelegates::OnObjectPropertyChanged.AddRaw(this, &FGetDataflowVariableNode::OnObjectPropertyChanged);
+#endif
+	}
+}
+
+void FGetDataflowVariableNode::UpdateOutputTypes(const FPropertyBagPropertyDesc& Desc)
+{
+	// compute type
+	if (!Desc.CachedProperty)
+	{
+		return;
+	}
+
+	FDataflowOutput* Output = FindOutput(&Value);
+	check(Output);
+
+	Output->UnlockType();
+
+	FString ExtendedType;
+	const FString CPPType = Desc.CachedProperty->GetCPPType(&ExtendedType);
+	const FName ConcreteType(CPPType + ExtendedType);
+
+	switch (Desc.ValueType)
+	{
+	case EPropertyBagPropertyType::Bool:
+		SetOutputConcreteType(&Value, TDataflowSingleTypePolicy<bool>::TypeName);
+		EvaluateFunction = &EvaluateBool;
+		break;
+	case EPropertyBagPropertyType::Byte:
+		SetOutputConcreteType(&Value, TDataflowSingleTypePolicy<uint8>::TypeName);
+		EvaluateFunction = &EvaluateByte;
+		break;
+	case EPropertyBagPropertyType::Int32:
+		SetOutputConcreteType(&Value, TDataflowSingleTypePolicy<int32>::TypeName);
+		EvaluateFunction = &EvaluateInt32;
+		break;
+	case EPropertyBagPropertyType::Int64:
+		SetOutputConcreteType(&Value, TDataflowSingleTypePolicy<int64>::TypeName);
+		EvaluateFunction = &EvaluateInt64;
+		break;
+	case EPropertyBagPropertyType::Float:
+		SetOutputConcreteType(&Value, TDataflowSingleTypePolicy<float>::TypeName);
+		EvaluateFunction = &EvaluateFloat;
+		break;
+	case EPropertyBagPropertyType::Double:
+		SetOutputConcreteType(&Value, TDataflowSingleTypePolicy<double>::TypeName);
+		EvaluateFunction = &EvaluateDouble;
+		break;
+	case EPropertyBagPropertyType::Name:
+		SetOutputConcreteType(&Value, TDataflowSingleTypePolicy<FName>::TypeName);
+		EvaluateFunction = &EvaluateName;
+		break;
+	case EPropertyBagPropertyType::String:
+		SetOutputConcreteType(&Value, TDataflowSingleTypePolicy<FString>::TypeName);
+		EvaluateFunction = &EvaluateString;
+		break;
+	case EPropertyBagPropertyType::Text:
+		SetOutputConcreteType(&Value, TDataflowSingleTypePolicy<FString>::TypeName);
+		EvaluateFunction = &EvaluateText;
+		break;
+	case EPropertyBagPropertyType::Object:
+		SetOutputConcreteType(&Value, ConcreteType);
+		EvaluateFunction = &EvaluateObject;
+		break;
+		// unsupported type for now 
+	case EPropertyBagPropertyType::Enum:
+	case EPropertyBagPropertyType::Struct:
+	case EPropertyBagPropertyType::SoftObject:
+	case EPropertyBagPropertyType::Class:
+	case EPropertyBagPropertyType::SoftClass:
+	case EPropertyBagPropertyType::UInt32:	// Type not fully supported at UI, will work with restrictions to type editing
+	case EPropertyBagPropertyType::UInt64: // Type not fully supported at UI, will work with restrictions to type editing
+	default:
+		EvaluateFunction = nullptr;
+		break;
+	}
+
+	Output->LockType();
+}
+
+void FGetDataflowVariableNode::OnPropertyChanged(UE::Dataflow::FContext& Context, const FPropertyChangedEvent& InPropertyChangedEvent)
+{
+#if WITH_EDITOR
+	const FName PropertyName = InPropertyChangedEvent.GetPropertyName();
+	const FName MemberPropertyName = InPropertyChangedEvent.GetMemberPropertyName();
+	const FName VariablePropertyBagMemberName = GET_MEMBER_NAME_CHECKED(FGetDataflowVariableNode, VariablePropertyBag);
+
+	if (PropertyName == VariablePropertyBagMemberName || MemberPropertyName == VariablePropertyBagMemberName)
+	{
+		// update the variable value in the dataflow asset
+		if (const TStrongObjectPtr<UDataflow> DataflowAsset = WeakDataflowPtr.Pin())
+		{
+			if (const FPropertyBagPropertyDesc* SourceDesc = VariablePropertyBag.FindPropertyDescByName(VariableName))
+			{
+				if (SourceDesc->CachedProperty)
+				{
+					EPropertyBagResult Result = DataflowAsset->Variables.SetValue(VariableName, SourceDesc->CachedProperty, VariablePropertyBag.GetValue().GetMemory());
+					if (Result == EPropertyBagResult::Success)
+					{
+						DataflowAsset->Modify();
+						FDataflowAssetDelegates::OnVariablesChanged.Broadcast(DataflowAsset.Get(), VariableName);
+
+						FPropertyChangedEvent PropertyChangedEvent(nullptr);
+						if (UClass* DataflowClass = DataflowAsset->GetClass())
+						{
+							FProperty* MemberProperty = DataflowClass->FindPropertyByName(GET_MEMBER_NAME_CHECKED(UDataflow, Variables));
+							PropertyChangedEvent.SetActiveMemberProperty(MemberProperty);
+						}
+						DataflowAsset->PostEditChangeProperty(PropertyChangedEvent);
+						Invalidate();
+					}
+				}
+			}
+		}
+	}
+#endif
+}
+
+void FGetDataflowVariableNode::Evaluate(UE::Dataflow::FContext& Context, const FDataflowOutput* Out) const
+{
+	if (Out->IsA(&Value))
+	{
+		if (EvaluateFunction)
+		{
+			if (const UE::Dataflow::FEngineContext* EngineContext = Context.AsType<UE::Dataflow::FEngineContext>())
+			{
+				if (IDataflowInstanceInterface* Interface = Cast<IDataflowInstanceInterface>(EngineContext->Owner))
+				{
+					const FDataflowInstance& DataflowInstance = Interface->GetDataflowInstance();
+					if (DataflowInstance.GetVariableOverrides().IsVariableOverridden(VariableName))
+					{
+						EvaluateFunction(*this, DataflowInstance.GetVariableOverrides().GetVariables(), Context, *Out);
+					}
+					else if (const UDataflow* DataflowAsset = DataflowInstance.GetDataflowAsset())
+					{
+						if (DataflowAsset->Variables.FindPropertyDescByName(VariableName) != nullptr)
+						{
+							EvaluateFunction(*this, DataflowAsset->Variables, Context, *Out);
+						}
+					}
+					return;
+				}
+			}
+		}
+		Out->SetNullValue(Context);
+	}
+}
+
+/*static*/ void FGetDataflowVariableNode::EvaluateBool(const FGetDataflowVariableNode& Node, const FInstancedPropertyBag& Variables, UE::Dataflow::FContext& Context, const FDataflowOutput& Out)
+{
+	bool OutBool = false;
+	if (const bool* Value = Variables.GetValueBool(Node.VariableName).TryGetValue())
+	{
+		OutBool = *Value;
+	}
+	Out.SetValue<bool>(MoveTemp(OutBool), Context);
+}
+
+/*static*/ void FGetDataflowVariableNode::EvaluateByte(const FGetDataflowVariableNode& Node, const FInstancedPropertyBag& Variables, UE::Dataflow::FContext& Context, const FDataflowOutput& Out)
+{
+	uint8 OutUInt8 = 0;
+	if (const uint8* Value = Variables.GetValueByte(Node.VariableName).TryGetValue())
+	{
+		OutUInt8 = *Value;
+	}
+	Out.SetValue<uint8>(MoveTemp(OutUInt8), Context);
+}
+
+/*static*/ void FGetDataflowVariableNode::EvaluateInt32(const FGetDataflowVariableNode& Node, const FInstancedPropertyBag& Variables, UE::Dataflow::FContext& Context, const FDataflowOutput& Out)
+{
+	int32 OutInt32 = 0;
+	if (const int32* Value = Variables.GetValueInt32(Node.VariableName).TryGetValue())
+	{
+		OutInt32 = *Value;
+	}
+	Out.SetValue<int32>(MoveTemp(OutInt32), Context);
+}
+
+/*static*/ void FGetDataflowVariableNode::EvaluateInt64(const FGetDataflowVariableNode& Node, const FInstancedPropertyBag& Variables, UE::Dataflow::FContext& Context, const FDataflowOutput& Out)
+{
+	int64 OutInt64 = 0;
+	if (const int64* Value = Variables.GetValueInt64(Node.VariableName).TryGetValue())
+	{
+		OutInt64 = *Value;
+	}
+	Out.SetValue<int64>(MoveTemp(OutInt64), Context);
+}
+
+/*static*/ void FGetDataflowVariableNode::EvaluateFloat(const FGetDataflowVariableNode& Node, const FInstancedPropertyBag& Variables, UE::Dataflow::FContext& Context, const FDataflowOutput& Out)
+{
+	float OutFloat = 0.0f;
+	if (const float* Value = Variables.GetValueFloat(Node.VariableName).TryGetValue())
+	{
+		OutFloat = *Value;
+	}
+	Out.SetValue<float>(MoveTemp(OutFloat), Context);
+}
+
+/*static*/ void FGetDataflowVariableNode::EvaluateDouble(const FGetDataflowVariableNode& Node, const FInstancedPropertyBag& Variables, UE::Dataflow::FContext& Context, const FDataflowOutput& Out)
+{
+	double OutDouble = 0.0;
+	if (const double* Value = Variables.GetValueDouble(Node.VariableName).TryGetValue())
+	{
+		OutDouble = *Value;
+	}
+	Out.SetValue<double>(MoveTemp(OutDouble), Context);
+}
+
+/*static*/ void FGetDataflowVariableNode::EvaluateName(const FGetDataflowVariableNode& Node, const FInstancedPropertyBag& Variables, UE::Dataflow::FContext& Context, const FDataflowOutput& Out)
+{
+	FName OutName;
+	if (const FName* Value = Variables.GetValueName(Node.VariableName).TryGetValue())
+	{
+		OutName = *Value;
+	}
+	Out.SetValue<FName>(MoveTemp(OutName), Context);
+}
+
+/*static*/ void FGetDataflowVariableNode::EvaluateString(const FGetDataflowVariableNode& Node, const FInstancedPropertyBag& Variables, UE::Dataflow::FContext& Context, const FDataflowOutput& Out)
+{
+	FString OutString;
+	if (const FString* Value = Variables.GetValueString(Node.VariableName).TryGetValue())
+	{
+		OutString = *Value;
+	}
+	Out.SetValue<FString>(MoveTemp(OutString), Context);
+}
+
+/*static*/ void FGetDataflowVariableNode::EvaluateText(const FGetDataflowVariableNode& Node, const FInstancedPropertyBag& Variables, UE::Dataflow::FContext& Context, const FDataflowOutput& Out)
+{
+	FString OutString;
+	if (const FText* Value = Variables.GetValueText(Node.VariableName).TryGetValue())
+	{
+		OutString = Value->ToString();
+	}
+	Out.SetValue<FString>(MoveTemp(OutString), Context);
+}
+
+/*static*/ void FGetDataflowVariableNode::EvaluateObject(const FGetDataflowVariableNode& Node, const FInstancedPropertyBag& Variables, UE::Dataflow::FContext& Context, const FDataflowOutput& Out)
+{
+	TObjectPtr<UObject> OutObject = nullptr;
+	if (UObject* const* Value = Variables.GetValueObject<UObject>(Node.VariableName).TryGetValue())
+	{
+		OutObject = *Value;
+	}
+	Out.SetValue<TObjectPtr<UObject>>(MoveTemp(OutObject), Context);
+}
+
diff --git a/Engine/Source/Runtime/Experimental/Dataflow/Engine/Public/Dataflow/DataflowInstance.h b/Engine/Source/Runtime/Experimental/Dataflow/Engine/Public/Dataflow/DataflowInstance.h
--- a/Engine/Source/Runtime/Experimental/Dataflow/Engine/Public/Dataflow/DataflowInstance.h
+++ b/Engine/Source/Runtime/Experimental/Dataflow/Engine/Public/Dataflow/DataflowInstance.h
@@ -10,10 +10,73 @@
 
 namespace UE::Dataflow::Private
 {
-	class FVariablesDetails;
+	class FVariablesOverridesDetails;
 }
 
 class UDataflow;
+struct FDataflowInstance;
+
+/**
+* This wraps the variable overrides
+* This is also separate from FDataflowInstance to allow for customization to display override checkboxes in front of each  variable property
+* ( see DataflowDetails.h )
+*/
+USTRUCT()
+struct FDataflowVariableOverrides
+{
+	GENERATED_USTRUCT_BODY();
+
+	DATAFLOWENGINE_API FDataflowVariableOverrides(FDataflowInstance* InOwner = nullptr);
+
+	/** Remove all variables */
+	DATAFLOWENGINE_API void RemoveAllVariables();
+
+	/** Sync variables with from the original dataflow asset */
+	DATAFLOWENGINE_API void SyncVariables();
+
+	/** check if a specific variable exists */
+	DATAFLOWENGINE_API bool HasVariable(FName VariableName) const;
+
+	DATAFLOWENGINE_API bool IsVariableOverridden(FName VariableName) const;
+
+	DATAFLOWENGINE_API const FInstancedPropertyBag& GetVariables() const;
+
+	DATAFLOWENGINE_API const FInstancedPropertyBag& GetVariables();
+
+	DATAFLOWENGINE_API static FName GetVariablePropertyName();
+
+#if WITH_EDITOR
+
+	void OnOwnerPostEditChangeProperty(struct FPropertyChangedEvent& PropertyChangedEvent);
+	void OnDataflowVariablesChanged(const UDataflow* DataflowAsset, FName VariableName);
+#endif
+
+private:
+	friend class UE::Dataflow::Private::FVariablesOverridesDetails;
+
+	DATAFLOWENGINE_API const FInstancedPropertyBag* GetDefaultVariablesFromAsset() const;
+
+	DATAFLOWENGINE_API FInstancedPropertyBag& GetOverridenVariables();
+
+	/** @return true if the property of specified ID is overridden. */
+	DATAFLOWENGINE_API bool IsVariableOverridden(const FGuid PropertyID) const;
+
+	/** Sets the override status of specified property by ID. */
+	DATAFLOWENGINE_API void SetVariableOverridden(const FGuid PropertyID, const bool bIsOverridden);
+
+	void RemoveOverridenVariablesNotInDataflowAsset();
+
+private:
+	/** Variables to override */
+	UPROPERTY(EditAnywhere, Category = "Dataflow", meta = (FixedLayout, ShowOnlyInnerProperties))
+	FInstancedPropertyBag Variables;
+
+	/** Array of overridden variable GUIDs. Non-overridden properties will inherit the values from the DataflowAsset default parameters. */
+	UPROPERTY()
+	TArray<FGuid> OverriddenVariableGuids;
+
+	FDataflowInstance* Owner = nullptr;;
+};
 
 /**
 * This structure is to be embedded in any asset that need generation from a dataflow
@@ -24,6 +87,7 @@ struct FDataflowInstance
 	GENERATED_USTRUCT_BODY();
 
 	DATAFLOWENGINE_API FDataflowInstance(UObject* InOwner = nullptr);
+	DATAFLOWENGINE_API ~FDataflowInstance();
 
 public:
 	DATAFLOWENGINE_API void SetDataflowAsset(UDataflow* DataflowAsset);
@@ -38,18 +102,20 @@ public:
 
 	DATAFLOWENGINE_API const FInstancedPropertyBag& GetVariables();
 
-	/** @return true if the property of specified ID is overridden. */
-	DATAFLOWENGINE_API bool IsVariableOverridden(const FGuid PropertyID) const;
+	/** Get the dataflow terminal member property name */
+	DATAFLOWENGINE_API static FName GetDataflowTerminalPropertyName();
 
-	/** Sets the override status of specified property by ID. */
-	DATAFLOWENGINE_API void SetVariableOverridden(const FGuid PropertyID, const bool bIsOverridden);
+	/** Get the dataflow asset member property name */
+	DATAFLOWENGINE_API static FName GetDataflowAssetPropertyName();
 
-	/** Get the variable member property name */
-	DATAFLOWENGINE_API static FName GetVariablePropertyName();
+	DATAFLOWENGINE_API const FDataflowVariableOverrides& GetVariableOverrides() const;
 
-	/** sync variables overriden in the asset with the one in the dataflow asset */
 	DATAFLOWENGINE_API void SyncVariables();
 
+#if WITH_EDITOR
+	DATAFLOWENGINE_API TSharedPtr<FStructOnScope> MakeStructOnScope() const;
+#endif 
+
 private:
 	/** Dataflow asset to use */
 	UPROPERTY(EditAnywhere, Category = "Dataflow")
@@ -60,19 +126,12 @@ private:
 	FName DataflowTerminal;
 
 	/** Variables to override */
-	UPROPERTY(EditAnywhere, Category = "Dataflow", meta=(FixedLayout))
-	FInstancedPropertyBag Variables;
-
-	/** Array of overridden variable GUIDs. Non-overridden properties will inherit the values from the DataflowAsset default parameters. */
-	UPROPERTY()
-	TArray<FGuid> OverriddenVariables;
+	UPROPERTY(EditAnywhere, Category = "Dataflow", meta=(ShowOnlyInnerProperties))
+	FDataflowVariableOverrides VariableOverrides;
 
 private:
-	friend class UE::Dataflow::Private::FVariablesDetails;
-	DATAFLOWENGINE_API FInstancedPropertyBag& GetMutableVariables();
-
-	/** remove overriden variables that do not exist in the dataflow asset variables */
-	void RemoveOverridenVariablesNotInDataflowAsset();
+	friend class UE::Dataflow::Private::FVariablesOverridesDetails;
+	DATAFLOWENGINE_API FDataflowVariableOverrides& GetVariableOverrides();
 
 #if WITH_EDITOR
 	void OnOwnerPostEditChangeProperty(struct FPropertyChangedEvent& PropertyChangedEvent);
diff --git a/Engine/Source/Runtime/Experimental/Dataflow/Engine/Public/Dataflow/DataflowVariableNodes.h b/Engine/Source/Runtime/Experimental/Dataflow/Engine/Public/Dataflow/DataflowVariableNodes.h
new file mode 100644
--- /dev/null
+++ b/Engine/Source/Runtime/Experimental/Dataflow/Engine/Public/Dataflow/DataflowVariableNodes.h
@@ -0,0 +1,74 @@
+// Copyright Epic Games, Inc. All Rights Reserved.
+
+#pragma once 
+
+#include "Dataflow/DataflowNode.h"
+
+#include "DataflowVariableNodes.generated.h"
+
+class UDataflow;
+class UObject;
+struct FInstancedPropertyBag;
+struct FPropertyChangedEvent;
+
+USTRUCT()
+struct FGetDataflowVariableNode : public FDataflowNode
+{
+	GENERATED_USTRUCT_BODY()
+	DATAFLOW_NODE_DEFINE_INTERNAL(FGetDataflowVariableNode, "GetVariable", "Dataflow", "")
+
+public:
+	DATAFLOWENGINE_API FGetDataflowVariableNode(const UE::Dataflow::FNodeParameters& InParam, FGuid InGuid = FGuid::NewGuid());
+	DATAFLOWENGINE_API ~FGetDataflowVariableNode();
+
+	DATAFLOWENGINE_API void SetVariable(UDataflow* DataflowAsset, FName VariableName);
+
+	DATAFLOWENGINE_API FName GetVariableName() const { return VariableName; }
+
+private:
+	UPROPERTY(meta=(DataflowOutput))
+	FDataflowAnyType Value;
+
+	// Called when UPROPERTY members of the dataflow node have been changed in the editor
+	virtual void OnPropertyChanged(UE::Dataflow::FContext& Context, const FPropertyChangedEvent& InPropertyChangedEvent) override;
+
+	/** Override this method to provide custom post-serialization for this node. This method will be called after Serialize. It is also called after copy-paste with ArchiveState IsLoading. */
+	virtual void PostSerialize(const FArchive& Ar);
+
+	virtual void Evaluate(UE::Dataflow::FContext& Context, const FDataflowOutput* Out) const;
+
+	static void EvaluateBool(const FGetDataflowVariableNode& Node, const FInstancedPropertyBag& Variables, UE::Dataflow::FContext& Context, const FDataflowOutput& Out);
+	static void EvaluateByte(const FGetDataflowVariableNode& Node, const FInstancedPropertyBag& Variables, UE::Dataflow::FContext& Context, const FDataflowOutput& Out);
+	static void EvaluateInt32(const FGetDataflowVariableNode& Node, const FInstancedPropertyBag& Variables, UE::Dataflow::FContext& Context, const FDataflowOutput& Out);
+	static void EvaluateInt64(const FGetDataflowVariableNode& Node, const FInstancedPropertyBag& Variables, UE::Dataflow::FContext& Context, const FDataflowOutput& Out);
+	static void EvaluateFloat(const FGetDataflowVariableNode& Node, const FInstancedPropertyBag& Variables, UE::Dataflow::FContext& Context, const FDataflowOutput& Out);
+	static void EvaluateDouble(const FGetDataflowVariableNode& Node, const FInstancedPropertyBag& Variables, UE::Dataflow::FContext& Context, const FDataflowOutput& Out);
+	static void EvaluateName(const FGetDataflowVariableNode& Node, const FInstancedPropertyBag& Variables, UE::Dataflow::FContext& Context, const FDataflowOutput& Out);
+	static void EvaluateString(const FGetDataflowVariableNode& Node, const FInstancedPropertyBag& Variables, UE::Dataflow::FContext& Context, const FDataflowOutput& Out);
+	static void EvaluateText(const FGetDataflowVariableNode& Node, const FInstancedPropertyBag& Variables, UE::Dataflow::FContext& Context, const FDataflowOutput& Out);
+	static void EvaluateObject(const FGetDataflowVariableNode& Node, const FInstancedPropertyBag& Variables, UE::Dataflow::FContext& Context, const FDataflowOutput& Out);
+
+	void UpdateOutputTypes(const FPropertyBagPropertyDesc& Desc);
+
+
+private:
+	UPROPERTY(EditAnywhere, Category = Variables, meta = (FixedLayout, ShowOnlyInnerProperties))
+	FInstancedPropertyBag VariablePropertyBag;
+
+	UPROPERTY()
+	FName VariableName;
+
+	using FStaticEvaluationFunctionPtr = void(*)(const FGetDataflowVariableNode&, const FInstancedPropertyBag&, UE::Dataflow::FContext&, const FDataflowOutput&);
+	FStaticEvaluationFunctionPtr EvaluateFunction = nullptr;
+
+	void OnObjectPropertyChanged(UObject* InObject, FPropertyChangedEvent& InPropertyChangedEvent);
+
+	TWeakObjectPtr<UDataflow> WeakDataflowPtr;
+
+};
+
+namespace UE::Dataflow
+{
+	void RegisterVariableNodes();
+}
+
