Dataflow Editor: the Construction and Simulation viewports now each track their own currently active preview scene profile, rather than sharing an editor-wide global selected profile
- The details panel on the right hand side of the DF editor corresponds to the currently active profile in the simulation viewport
- Selected profiles for both viewports are saved/loaded using the DataflowEditorOptions class
- The defaults in DataflowEditorOptions are "Grey Wireframe" for the Construction viewport and "Grey Ambient" for the Simulation viewport, so these should be the profiles the user sees the first time they enter the Dataflow Editor
- Also updating the default construction material to a brighter white to look better in the dark Construction viewport

#virtualized

diff --git a/Engine/Plugins/Experimental/Dataflow/Source/DataflowEditor/DataflowEditor.Build.cs b/Engine/Plugins/Experimental/Dataflow/Source/DataflowEditor/DataflowEditor.Build.cs
--- a/Engine/Plugins/Experimental/Dataflow/Source/DataflowEditor/DataflowEditor.Build.cs
+++ b/Engine/Plugins/Experimental/Dataflow/Source/DataflowEditor/DataflowEditor.Build.cs
@@ -57,6 +57,7 @@ namespace UnrealBuildTool.Rules
 					"RenderCore",
 					"RHI",
 					"SceneOutliner",
+					"SharedSettingsWidgets",
 					"SkeletonEditor",
 					"Slate",
 					"SlateCore",
diff --git a/Engine/Plugins/Experimental/Dataflow/Source/DataflowEditor/Private/Dataflow/DataflowAdvancedPreviewDetailsTab.cpp b/Engine/Plugins/Experimental/Dataflow/Source/DataflowEditor/Private/Dataflow/DataflowAdvancedPreviewDetailsTab.cpp
new file mode 100644
--- /dev/null
+++ b/Engine/Plugins/Experimental/Dataflow/Source/DataflowEditor/Private/Dataflow/DataflowAdvancedPreviewDetailsTab.cpp
@@ -0,0 +1,348 @@
+// Copyright Epic Games, Inc. All Rights Reserved.
+
+#include "DataflowAdvancedPreviewDetailsTab.h"
+#include "Modules/ModuleManager.h"
+#include "Widgets/SBoxPanel.h"
+#include "Widgets/Text/STextBlock.h"
+#include "Widgets/Input/SButton.h"
+#include "Widgets/Input/STextComboBox.h"
+#include "IDetailsView.h"
+#include "PropertyEditorModule.h"
+#include "ScopedTransaction.h"
+
+#include "AssetViewerSettings.h"
+
+#include "AdvancedPreviewScene.h"
+#include "IDetailRootObjectCustomization.h"
+
+#define LOCTEXT_NAMESPACE "SPrettyPreview"
+
+SDataflowAdvancedPreviewDetailsTab::SDataflowAdvancedPreviewDetailsTab()
+{
+	DefaultSettings = UAssetViewerSettings::Get();
+	if (DefaultSettings)
+	{
+		RefreshDelegate = DefaultSettings->OnAssetViewerSettingsChanged().AddRaw(this, &SDataflowAdvancedPreviewDetailsTab::OnAssetViewerSettingsRefresh);
+			
+		AddRemoveProfileDelegate = DefaultSettings->OnAssetViewerProfileAddRemoved().AddLambda([this]() { this->Refresh(); });
+
+		PostUndoDelegate = DefaultSettings->OnAssetViewerSettingsPostUndo().AddRaw(this, &SDataflowAdvancedPreviewDetailsTab::OnAssetViewerSettingsPostUndo);
+	}
+}
+
+SDataflowAdvancedPreviewDetailsTab::~SDataflowAdvancedPreviewDetailsTab()
+{
+	DefaultSettings = UAssetViewerSettings::Get();
+	if (DefaultSettings)
+	{
+		DefaultSettings->OnAssetViewerSettingsChanged().Remove(RefreshDelegate);
+		DefaultSettings->OnAssetViewerProfileAddRemoved().Remove(AddRemoveProfileDelegate);
+		DefaultSettings->OnAssetViewerSettingsPostUndo().Remove(PostUndoDelegate);
+		DefaultSettings->Save();
+	}
+}
+
+void SDataflowAdvancedPreviewDetailsTab::Construct(const FArguments& InArgs, const TSharedRef<FAdvancedPreviewScene>& InPreviewScene)
+{
+	PreviewScenePtr = InPreviewScene;
+	DefaultSettings = UAssetViewerSettings::Get();
+	ProfileIndexStorage = InArgs._ProfileIndexStorage;
+	AdditionalSettings = InArgs._AdditionalSettings;
+	DetailCustomizations = InArgs._DetailCustomizations;
+	PropertyTypeCustomizations = InArgs._PropertyTypeCustomizations;
+	Delegates = InArgs._Delegates;
+
+	ProfileIndex = ProfileIndexStorage->RetrieveProfileIndex();
+	ProfileIndex = DefaultSettings->Profiles.IsValidIndex(ProfileIndex) ? ProfileIndex : 0;
+
+	CreateSettingsView();
+
+	this->ChildSlot
+	[
+		SNew(SVerticalBox)
+		+ SVerticalBox::Slot()
+		.Padding(2.0f, 1.0f, 2.0f, 1.0f)
+		[
+			SNew(SHorizontalBox)
+			+SHorizontalBox::Slot()
+			[
+				SettingsView->AsShared()
+			]
+		]
+
+		+SVerticalBox::Slot()		
+		.Padding(2.0f, 1.0f, 2.0f, 1.0f)
+		.AutoHeight()
+		[
+			SNew(SHorizontalBox)
+			+SHorizontalBox::Slot()
+			.Padding(2.0f)
+			[
+				SNew(SHorizontalBox)
+				.ToolTipText(LOCTEXT("SceneProfileComboBoxToolTip", "Allows for switching between scene environment and lighting profiles."))
+				+ SHorizontalBox::Slot()
+				.Padding(0.0f, 0.0f, 2.0f, 0.0f)
+				.AutoWidth()
+				.VAlign(VAlign_Center)
+				[
+					SNew(STextBlock)
+					.Text(LOCTEXT("SceneProfileSettingsLabel", "Profile"))
+				]
+				+ SHorizontalBox::Slot()
+				.VAlign(VAlign_Fill)
+				[
+					SAssignNew(ProfileComboBox, STextComboBox)
+					.OptionsSource(&ProfileNames)
+					.OnSelectionChanged(this, &SDataflowAdvancedPreviewDetailsTab::ComboBoxSelectionChanged)				
+					.IsEnabled_Lambda([this]() -> bool
+					{
+						return ProfileNames.Num() > 1;
+					})
+				]
+			]
+			+ SHorizontalBox::Slot()
+			.Padding(2.0f)
+			.AutoWidth()
+			[
+				SNew(SButton)
+				.OnClicked(this, &SDataflowAdvancedPreviewDetailsTab::AddProfileButtonClick)
+				.Text(LOCTEXT("AddProfileButton", "Add Profile"))
+				.ToolTipText(LOCTEXT("SceneProfileAddProfile", "Adds a new profile."))
+			]
+			+ SHorizontalBox::Slot()
+			.Padding(2.0f)
+			.AutoWidth()
+			[
+				SNew(SButton)
+				.OnClicked(this, &SDataflowAdvancedPreviewDetailsTab::RemoveOrResetProfileButtonClick)
+				.Text_Lambda([this]()
+				{
+					if (DefaultSettings->Profiles[ProfileIndex].bIsEngineDefaultProfile)
+					{
+						return LOCTEXT("ResetProfileButton", "Reset Profile");
+					}
+					return LOCTEXT("RemoveProfileButton", "Remove Profile");
+				})
+				.ToolTipText_Lambda([this]()
+				{
+					if (DefaultSettings->Profiles[ProfileIndex].bIsEngineDefaultProfile)
+					{
+						return LOCTEXT("SceneProfileResetProfile", "Resets this engine profile to default settings. Cannot delete engine profiles.");
+					}
+					return LOCTEXT("SceneProfileRemoveProfile", "Removes the currently selected profile.");
+				})
+				.IsEnabled_Lambda([this]()->bool
+				{
+					return ProfileNames.Num() > 1; 
+				})
+			]
+		]
+	];
+
+	UpdateProfileNames();
+	UpdateSettingsView();
+}
+
+void SDataflowAdvancedPreviewDetailsTab::ComboBoxSelectionChanged(TSharedPtr<FString> NewSelection, ESelectInfo::Type SelectInfo)
+{
+	int32 NewSelectionIndex;
+
+	if (ProfileNames.Find(NewSelection, NewSelectionIndex))
+	{
+		ProfileIndex = NewSelectionIndex;
+		ProfileIndexStorage->StoreProfileIndex(ProfileIndex);
+		UpdateSettingsView();	
+		check(PreviewScenePtr.IsValid());
+		if (SelectInfo == ESelectInfo::Type::OnMouseClick)
+		{
+			PreviewScenePtr.Pin()->SetProfileIndex(ProfileIndex);	
+		}
+	}
+}
+
+void SDataflowAdvancedPreviewDetailsTab::UpdateSettingsView()
+{	
+	TArray<UObject*> Objects;
+	if (AdditionalSettings)
+	{
+		Objects.Add(AdditionalSettings);
+	}
+	Objects.Add(DefaultSettings);
+
+	SettingsView->SetObjects(Objects, true);
+	//SettingsView->Invalidate(EInvalidateWidgetReason::Paint);
+}
+
+void SDataflowAdvancedPreviewDetailsTab::UpdateProfileNames()
+{
+	checkf(DefaultSettings->Profiles.Num(), TEXT("There should always be at least one profile available"));
+	ProfileNames.Empty();
+	for (FPreviewSceneProfile& Profile : DefaultSettings->Profiles)
+	{
+		FString Suffix = TEXT("");
+		if (Profile.bSharedProfile)
+		{
+			Suffix = TEXT(" (Shared)");
+		}
+		if (Profile.bIsEngineDefaultProfile)
+		{
+			Suffix = TEXT(" (Engine Default)");
+		}
+		ProfileNames.Add(TSharedPtr<FString>(new FString(Profile.ProfileName + Suffix)));
+	}
+
+	ProfileComboBox->RefreshOptions();
+	ProfileComboBox->SetSelectedItem(ProfileNames[ProfileIndex]);
+}
+
+FReply SDataflowAdvancedPreviewDetailsTab::AddProfileButtonClick()
+{
+	const FScopedTransaction Transaction(LOCTEXT("AddSceneProfile", "Adding Preview Scene Profile"));
+	DefaultSettings->Modify();
+
+	// Add new profile to settings instance
+	ProfileIndex = DefaultSettings->Profiles.AddDefaulted();
+	FPreviewSceneProfile& NewProfile = DefaultSettings->Profiles.Last();
+	
+	// Try to create a valid profile name when one is added
+	bool bFoundValidName = false;
+	int32 ProfileAppendNum = FMath::Max(0, DefaultSettings->Profiles.Num() - 1);
+	FString NewProfileName;
+	while (!bFoundValidName)
+	{
+		NewProfileName = FString::Printf(TEXT("Profile_%i"), ProfileAppendNum);
+
+		bool bValidName = true;
+		for (const FPreviewSceneProfile& Profile : DefaultSettings->Profiles)
+		{
+			if (Profile.ProfileName == NewProfileName)
+			{
+				bValidName = false;				
+				break;
+			}
+		}
+
+		if (!bValidName)
+		{
+			++ProfileAppendNum;
+		}
+
+		bFoundValidName = bValidName;
+	}
+
+	NewProfile.ProfileName = NewProfileName;
+	ProfileIndexStorage->StoreProfileIndex(ProfileIndex);
+	DefaultSettings->PostEditChange();
+
+	// Change selection to new profile so the user directly sees the profile that was added
+	Refresh();
+	ProfileComboBox->SetSelectedItem(ProfileNames.Last());
+	
+	return FReply::Handled();
+}
+
+FReply SDataflowAdvancedPreviewDetailsTab::RemoveOrResetProfileButtonClick()
+{
+	const FPreviewSceneProfile& CurrentProfile = DefaultSettings->Profiles[ProfileIndex];
+	
+	if (CurrentProfile.bIsEngineDefaultProfile)
+	{
+		const FScopedTransaction Transaction(LOCTEXT("ResetSceneProfile", "Reset Preview Scene Profile"));
+		DefaultSettings->Modify();
+
+		const FPreviewSceneProfile* DefaultEditorProfile = GetMutableDefault<UDefaultEditorProfiles>()->GetProfile(CurrentProfile.ProfileName);
+		if (DefaultEditorProfile)
+		{
+			// Reset currently selected profile 
+			DefaultSettings->Profiles[ProfileIndex] = *DefaultEditorProfile;
+			DefaultSettings->PostEditChange();
+			return FReply::Handled();
+		}
+
+		// if we get here, it means an engine-provided default profile was removed from the engine,
+		// in which case we should remove it...
+	}
+
+	const FScopedTransaction Transaction(LOCTEXT("RemoveSceneProfile", "Remove Preview Scene Profile"));
+	DefaultSettings->Modify();
+
+	// Remove currently selected profile 
+	DefaultSettings->Profiles.RemoveAt(ProfileIndex);
+	ProfileIndex = DefaultSettings->Profiles.IsValidIndex(ProfileIndex - 1 ) ? ProfileIndex - 1 : 0;
+	ProfileIndexStorage->StoreProfileIndex(ProfileIndex);
+	DefaultSettings->PostEditChange();
+	
+	return FReply::Handled();
+}
+
+void SDataflowAdvancedPreviewDetailsTab::OnAssetViewerSettingsRefresh(const FName& InPropertyName)
+{
+	if (!PreviewScenePtr.IsValid())
+	{
+		// this callback can fire when the editor is forcibly closed and tool modes revert the active profile
+		// when this happens, the preview scene is null even though this details tab hasn't been destroyed yet (and unregistered this delegate)
+		return;
+	}
+	
+	if (InPropertyName == GET_MEMBER_NAME_CHECKED(FPreviewSceneProfile, ProfileName) || InPropertyName == GET_MEMBER_NAME_CHECKED(FPreviewSceneProfile, bSharedProfile))
+	{
+		Refresh();
+	}
+}
+
+void SDataflowAdvancedPreviewDetailsTab::CreateSettingsView()
+{	
+	// Create a property view
+	FPropertyEditorModule& EditModule = FModuleManager::Get().GetModuleChecked<FPropertyEditorModule>("PropertyEditor");
+
+	FDetailsViewArgs DetailsViewArgs;
+	DetailsViewArgs.NameAreaSettings = FDetailsViewArgs::HideNameArea;
+	DetailsViewArgs.bHideSelectionTip = true;
+	DetailsViewArgs.DefaultsOnlyVisibility = EEditDefaultsOnlyNodeVisibility::Automatic;
+	DetailsViewArgs.bShowOptions = false;
+	DetailsViewArgs.bAllowMultipleTopLevelObjects = true;
+
+	SettingsView = EditModule.CreateDetailView(DetailsViewArgs);
+
+	for (const FAdvancedPreviewSceneModule::FDetailCustomizationInfo& DetailCustomizationInfo : DetailCustomizations)
+	{
+		SettingsView->RegisterInstancedCustomPropertyLayout(DetailCustomizationInfo.Struct, DetailCustomizationInfo.OnGetDetailCustomizationInstance);
+	}
+
+	for (const FAdvancedPreviewSceneModule::FPropertyTypeCustomizationInfo& PropertyTypeCustomizationInfo : PropertyTypeCustomizations)
+	{
+		SettingsView->RegisterInstancedCustomPropertyTypeLayout(PropertyTypeCustomizationInfo.StructName, PropertyTypeCustomizationInfo.OnGetPropertyTypeCustomizationInstance);
+	}
+
+	for (FAdvancedPreviewSceneModule::FDetailDelegates& DetailDelegate : Delegates)
+	{
+		DetailDelegate.OnPreviewSceneChangedDelegate.AddSP(this, &SDataflowAdvancedPreviewDetailsTab::OnPreviewSceneChanged);
+	}
+
+	UpdateSettingsView();
+}
+
+void SDataflowAdvancedPreviewDetailsTab::Refresh()
+{	
+	const int32 StoredIndex = ProfileIndexStorage->RetrieveProfileIndex();
+	ProfileIndex = DefaultSettings->Profiles.IsValidIndex(StoredIndex) ? StoredIndex : 0;
+	ProfileIndexStorage->StoreProfileIndex(ProfileIndex);
+	UpdateProfileNames();
+	PreviewScenePtr.Pin()->SetProfileIndex(ProfileIndex);
+	UpdateSettingsView();
+}
+
+void SDataflowAdvancedPreviewDetailsTab::OnAssetViewerSettingsPostUndo()
+{
+	Refresh();
+	PreviewScenePtr.Pin()->UpdateScene(DefaultSettings->Profiles[ProfileIndex]);
+}
+
+void SDataflowAdvancedPreviewDetailsTab::OnPreviewSceneChanged(TSharedRef<FAdvancedPreviewScene> PreviewScene)
+{
+  	PreviewScenePtr = PreviewScene;
+	Refresh();
+
+}
+
+#undef LOCTEXT_NAMESPACE
\ No newline at end of file
diff --git a/Engine/Plugins/Experimental/Dataflow/Source/DataflowEditor/Private/Dataflow/DataflowAdvancedPreviewDetailsTab.h b/Engine/Plugins/Experimental/Dataflow/Source/DataflowEditor/Private/Dataflow/DataflowAdvancedPreviewDetailsTab.h
new file mode 100644
--- /dev/null
+++ b/Engine/Plugins/Experimental/Dataflow/Source/DataflowEditor/Private/Dataflow/DataflowAdvancedPreviewDetailsTab.h
@@ -0,0 +1,89 @@
+// Copyright Epic Games, Inc. All Rights Reserved.
+#pragma once
+
+#include "CoreMinimal.h"
+#include "SlateFwd.h"
+#include "Input/Reply.h"
+#include "Widgets/DeclarativeSyntaxSupport.h"
+#include "Widgets/SCompoundWidget.h"
+#include "PropertyEditorDelegates.h"
+#include "AdvancedPreviewSceneModule.h"
+#include "Dataflow/DataflowPreviewProfileController.h"
+
+class FAdvancedPreviewScene;
+class IDetailsView;
+class UAssetViewerSettings;
+
+// This class is almost identical to SAdvancedPreviewDetailsTab except that it doesn't use UEditorPerProjectUserSettings
+// to get the current scene profile index. Instead it is supplied with an IProfileIndexStorage object which stores and 
+// loads the scene profile index. This allows a separate details tab to be created for each AdvancedPreviewScene.
+
+class SDataflowAdvancedPreviewDetailsTab : public SCompoundWidget
+{
+public:
+	SLATE_BEGIN_ARGS(SDataflowAdvancedPreviewDetailsTab)
+		: _ProfileIndexStorage(nullptr),
+		_AdditionalSettings(nullptr)
+	{}
+
+	SLATE_ARGUMENT(TSharedPtr<FDataflowPreviewProfileController::IProfileIndexStorage>, ProfileIndexStorage)
+
+	/** Additional settings object to display in the view */
+	SLATE_ARGUMENT(UObject*, AdditionalSettings)
+
+	/** Customizations to use for this details tab */
+	SLATE_ARGUMENT(TArray<FAdvancedPreviewSceneModule::FDetailCustomizationInfo>, DetailCustomizations)
+
+	/** Customizations to use for this details tab */
+	SLATE_ARGUMENT(TArray<FAdvancedPreviewSceneModule::FPropertyTypeCustomizationInfo>, PropertyTypeCustomizations)
+
+	/** Delegates to use for this details tab */
+	SLATE_ARGUMENT(TArray<FAdvancedPreviewSceneModule::FDetailDelegates>, Delegates)
+
+	SLATE_END_ARGS()
+
+public:
+	/** **/
+	SDataflowAdvancedPreviewDetailsTab();
+	~SDataflowAdvancedPreviewDetailsTab();
+
+	/** SWidget functions */
+	void Construct(const FArguments& InArgs, const TSharedRef<FAdvancedPreviewScene>& InPreviewScene);
+
+	void Refresh();
+
+protected:
+	void CreateSettingsView();
+	void ComboBoxSelectionChanged(TSharedPtr<FString> NewSelection, ESelectInfo::Type /*SelectInfo*/);
+	void UpdateSettingsView();
+	void UpdateProfileNames();
+	FReply AddProfileButtonClick();
+	FReply RemoveOrResetProfileButtonClick();
+protected:
+	void OnAssetViewerSettingsRefresh(const FName& InPropertyName);
+	void OnAssetViewerSettingsPostUndo();
+protected:
+	/** Property viewing widget */
+	TSharedPtr<IDetailsView> SettingsView;
+	TSharedPtr<STextComboBox> ProfileComboBox;
+	TWeakPtr<FAdvancedPreviewScene> PreviewScenePtr;
+	UAssetViewerSettings* DefaultSettings;
+	UObject* AdditionalSettings;
+
+	TArray<TSharedPtr<FString>> ProfileNames;
+	int32 ProfileIndex;
+
+	void OnPreviewSceneChanged(TSharedRef<FAdvancedPreviewScene> PreviewScene);
+
+	FDelegateHandle RefreshDelegate;
+	FDelegateHandle AddRemoveProfileDelegate;
+	FDelegateHandle PostUndoDelegate;
+
+	TSharedPtr<FDataflowPreviewProfileController::IProfileIndexStorage> ProfileIndexStorage;
+
+	TArray<FAdvancedPreviewSceneModule::FDetailCustomizationInfo> DetailCustomizations;
+
+	TArray<FAdvancedPreviewSceneModule::FPropertyTypeCustomizationInfo> PropertyTypeCustomizations;
+
+	TArray<FAdvancedPreviewSceneModule::FDetailDelegates> Delegates;
+};
diff --git a/Engine/Plugins/Experimental/Dataflow/Source/DataflowEditor/Private/Dataflow/DataflowAssetViewerSettingsCustomization.cpp b/Engine/Plugins/Experimental/Dataflow/Source/DataflowEditor/Private/Dataflow/DataflowAssetViewerSettingsCustomization.cpp
new file mode 100644
--- /dev/null
+++ b/Engine/Plugins/Experimental/Dataflow/Source/DataflowEditor/Private/Dataflow/DataflowAssetViewerSettingsCustomization.cpp
@@ -0,0 +1,196 @@
+// Copyright Epic Games, Inc. All Rights Reserved.
+
+#include "DataflowAssetViewerSettingsCustomization.h"
+
+#include "AssetViewerSettings.h"
+#include "Containers/Array.h"
+#include "Containers/Map.h"
+#include "Delegates/Delegate.h"
+#include "DetailCategoryBuilder.h"
+#include "DetailLayoutBuilder.h"
+#include "DetailWidgetRow.h"
+#include "Fonts/SlateFontInfo.h"
+#include "GenericPlatform/GenericPlatformFile.h"
+#include "HAL/PlatformCrt.h"
+#include "HAL/PlatformFileManager.h"
+#include "IDetailGroup.h"
+#include "IDetailPropertyRow.h"
+#include "Internationalization/Internationalization.h"
+#include "Misc/AssertionMacros.h"
+#include "Misc/Attribute.h"
+#include "PropertyHandle.h"
+#include "SSettingsEditorCheckoutNotice.h"
+#include "ScopedTransaction.h"
+#include "UObject/NameTypes.h"
+#include "UObject/UObjectGlobals.h"
+#include "UObject/UnrealType.h"
+#include "Widgets/DeclarativeSyntaxSupport.h"
+#include "Widgets/Input/SEditableTextBox.h"
+
+#define LOCTEXT_NAMESPACE "AssetViewerSettingsCustomizations"
+
+FDataflowAssetViewerSettingsCustomization::FDataflowAssetViewerSettingsCustomization(TSharedPtr<FDataflowPreviewProfileController::IProfileIndexStorage> ProfileIndexStorage) :
+	ProfileIndexStorage(ProfileIndexStorage)
+{
+}
+
+void FDataflowAssetViewerSettingsCustomization::CustomizeDetails(IDetailLayoutBuilder& DetailBuilder)
+{
+	ViewerSettings = UAssetViewerSettings::Get();
+
+	// Create the watcher widget for the default config file (checks file status / SCC state)
+	FileWatcherWidget =
+		SNew(SSettingsEditorCheckoutNotice)
+		.ConfigFilePath(this, &FDataflowAssetViewerSettingsCustomization::SharedProfileConfigFilePath)
+		.Visibility(this, &FDataflowAssetViewerSettingsCustomization::ShowFileWatcherWidget);
+
+	// Find profiles array property handle and hide it from settings view
+	TSharedPtr<IPropertyHandle> ProfileHandle = DetailBuilder.GetProperty(GET_MEMBER_NAME_CHECKED(UAssetViewerSettings, Profiles));
+	check(ProfileHandle->IsValidHandle());
+	ProfileHandle->MarkHiddenByCustomization();
+
+	// Create new category
+	IDetailCategoryBuilder& CategoryBuilder = DetailBuilder.EditCategory("Settings", LOCTEXT("AssetViewerSettingsCategory", "Settings"), ECategoryPriority::Important);
+
+	check(ViewerSettings != nullptr);
+	ProfileIndex = 0;
+	if (ProfileIndexStorage)
+	{
+		ProfileIndex = ViewerSettings->Profiles.IsValidIndex(ProfileIndexStorage->RetrieveProfileIndex()) ? ProfileIndexStorage->RetrieveProfileIndex() : 0;
+	}
+
+	// Add current active profile property child nodes (rest of profiles remain hidden)
+	TSharedPtr<IPropertyHandle> ProfilePropertyHandle = ProfileHandle->GetChildHandle(ProfileIndex);
+	checkf(ProfileHandle->IsValidHandle(), TEXT("Invalid indexing into profiles child properties"));	
+	uint32 PropertyCount = 0;
+	ProfilePropertyHandle->GetNumChildren(PropertyCount);
+
+	TMap<FName, IDetailGroup*> SubcategoryToGroupMap;
+
+	FName NamePropertyName = GET_MEMBER_NAME_CHECKED(FPreviewSceneProfile, ProfileName);
+	FName SharedProfilePropertyName = GET_MEMBER_NAME_CHECKED(FPreviewSceneProfile, bSharedProfile);
+	for (uint32 PropertyIndex = 0; PropertyIndex < PropertyCount; ++PropertyIndex)
+	{
+		TSharedPtr<IPropertyHandle> ProfileProperty = ProfilePropertyHandle->GetChildHandle(PropertyIndex);
+				
+		if (ProfileProperty->GetProperty()->GetFName() == NamePropertyName)
+		{
+			NameProperty = ProfileProperty;
+			CategoryBuilder.AddCustomRow(LOCTEXT("PreviewSceneProfileDetails_ProfileNameLabel", "Profile Name"))
+					.NameContent()
+					[
+						ProfileProperty->CreatePropertyNameWidget()
+					]
+					.ValueContent()
+					.MaxDesiredWidth(250.0f)
+					[
+						// PropertyHandle->CreatePropertyValueWidget()
+						SAssignNew(NameEditTextBox, SEditableTextBox)
+						.IsEnabled_Lambda([this]()
+						{
+							return !ViewerSettings->Profiles[ProfileIndex].bIsEngineDefaultProfile;
+						})
+						.Text(this, &FDataflowAssetViewerSettingsCustomization::OnGetProfileName)
+						.OnTextChanged(this, &FDataflowAssetViewerSettingsCustomization::OnProfileNameChanged)
+						.OnTextCommitted(this, &FDataflowAssetViewerSettingsCustomization::OnProfileNameCommitted)
+						.Font(IDetailLayoutBuilder::GetDetailFont())
+					];
+		}
+		else if (ProfileProperty->GetProperty()->GetFName() == SharedProfilePropertyName)
+		{
+			IDetailPropertyRow& Row = CategoryBuilder.AddProperty(ProfileProperty);
+			Row.EditCondition(TAttribute<bool>(this, &FDataflowAssetViewerSettingsCustomization::CanSetSharedProfile), nullptr);
+			
+			CategoryBuilder.AddCustomRow(LOCTEXT("PreviewSceneProfileName_CheckoutRow", "Checkout Default Config"))
+			.Visibility(TAttribute<EVisibility>(this, &FDataflowAssetViewerSettingsCustomization::ShowFileWatcherWidget))
+			[
+				FileWatcherWidget->AsShared()
+			];
+			
+		}
+		else
+		{
+			FName DefaultCategoryName = ProfileProperty->GetDefaultCategoryName();
+
+			IDetailGroup*& Group = SubcategoryToGroupMap.FindOrAdd(DefaultCategoryName);
+
+			if (!Group)
+			{
+				Group = &CategoryBuilder.AddGroup(DefaultCategoryName, ProfileProperty->GetDefaultCategoryText());
+			}
+
+			Group->AddPropertyRow(ProfileProperty.ToSharedRef());
+		}
+	}
+}
+
+FText FDataflowAssetViewerSettingsCustomization::OnGetProfileName() const
+{
+	return FText::FromString(ViewerSettings->Profiles[ProfileIndex].ProfileName);
+}
+
+void FDataflowAssetViewerSettingsCustomization::OnProfileNameChanged(const FText& InNewText)
+{	
+	bValidProfileName = IsProfileNameValid(InNewText.ToString());
+	if (!bValidProfileName)
+	{
+		NameEditTextBox->SetError(LOCTEXT("PreviewSceneProfileName_NotValid", "This name is already in use"));
+	}
+	else
+	{
+		NameEditTextBox->SetError(FText::GetEmpty());
+	}
+}
+
+void FDataflowAssetViewerSettingsCustomization::OnProfileNameCommitted(const FText& InNewText, ETextCommit::Type InTextCommit)
+{
+	if (bValidProfileName && InTextCommit != ETextCommit::OnCleared)
+	{
+		const FScopedTransaction Transaction(LOCTEXT("RenameProfile", "Rename Profile"));
+		ViewerSettings->Modify();
+
+		ViewerSettings->Profiles[ProfileIndex].ProfileName = InNewText.ToString();		
+		FPropertyChangedEvent PropertyEvent(NameProperty->GetProperty());
+		ViewerSettings->PostEditChangeProperty(PropertyEvent);
+	}
+
+	bValidProfileName = false;
+	NameEditTextBox->SetError(FText::GetEmpty());
+}
+
+const bool FDataflowAssetViewerSettingsCustomization::IsProfileNameValid(const FString& NewName)
+{
+	for (int32 CheckProfileIndex = 0; CheckProfileIndex < ViewerSettings->Profiles.Num(); ++CheckProfileIndex)
+	{
+		if (ProfileIndex != CheckProfileIndex && ViewerSettings->Profiles[CheckProfileIndex].ProfileName == NewName)
+		{
+			return false;
+		}
+	}
+
+	return true;
+}
+
+bool FDataflowAssetViewerSettingsCustomization::CanSetSharedProfile() const
+{
+	const bool bIsEngineDefaultProfile = ViewerSettings->Profiles[ProfileIndex].bIsEngineDefaultProfile;
+	const bool bIsConfigWriteable = !FPlatformFileManager::Get().GetPlatformFile().IsReadOnly(*SharedProfileConfigFilePath());
+	return bIsConfigWriteable && !bIsEngineDefaultProfile;
+}
+
+EVisibility FDataflowAssetViewerSettingsCustomization::ShowFileWatcherWidget() const
+{
+	return CanSetSharedProfile() ? EVisibility::Collapsed : EVisibility::Visible;
+}
+
+FString FDataflowAssetViewerSettingsCustomization::SharedProfileConfigFilePath() const
+{
+	if (ViewerSettings != nullptr)
+	{
+		return ViewerSettings->GetDefaultConfigFilename();
+	}
+
+	return FString();
+}
+
+#undef LOCTEXT_NAMESPACE 
diff --git a/Engine/Plugins/Experimental/Dataflow/Source/DataflowEditor/Private/Dataflow/DataflowAssetViewerSettingsCustomization.h b/Engine/Plugins/Experimental/Dataflow/Source/DataflowEditor/Private/Dataflow/DataflowAssetViewerSettingsCustomization.h
new file mode 100644
--- /dev/null
+++ b/Engine/Plugins/Experimental/Dataflow/Source/DataflowEditor/Private/Dataflow/DataflowAssetViewerSettingsCustomization.h
@@ -0,0 +1,61 @@
+// Copyright Epic Games, Inc. All Rights Reserved.
+
+#pragma once
+
+#include "Containers/UnrealString.h"
+#include "HAL/Platform.h"
+#include "IDetailCustomization.h"
+#include "Internationalization/Text.h"
+#include "Layout/Visibility.h"
+#include "Templates/SharedPointer.h"
+#include "Types/SlateEnums.h"
+#include "Dataflow/DataflowPreviewProfileController.h"
+
+class IDetailLayoutBuilder;
+class IPropertyHandle;
+class SEditableTextBox;
+class SSettingsEditorCheckoutNotice;
+class UAssetViewerSettings;
+
+// This class is almost identical to FAssetViewerSettingsCustomization except that it doesn't use UEditorPerProjectUserSettings 
+// to get the current scene profile index. Instead it is supplied with an IProfileIndexStorage object which stores and 
+// loads the scene profile index. This allows a separate customization to be created for each AdvancedPreviewScene.
+
+class FDataflowAssetViewerSettingsCustomization : public IDetailCustomization
+{
+public:
+
+	FDataflowAssetViewerSettingsCustomization(TSharedPtr<FDataflowPreviewProfileController::IProfileIndexStorage> ProfileIndexStorage);
+
+	// IDetailCustomization interface
+	virtual void CustomizeDetails(IDetailLayoutBuilder& DetailBuilder) override;
+	// End of IDetailCustomization interface
+
+protected:
+
+	// Callbacks for customized SEditableTextBox
+	FText OnGetProfileName() const;
+	void OnProfileNameChanged(const FText& InNewText);
+	void OnProfileNameCommitted(const FText& InNewText, ETextCommit::Type InTextCommit);
+	
+	// Check whether or not the given profile name is valid by cross-referencing it with existing names
+	const bool IsProfileNameValid(const FString& NewName);
+	bool CanSetSharedProfile() const;
+	EVisibility ShowFileWatcherWidget() const;
+	FString SharedProfileConfigFilePath() const;
+
+private:
+
+	// Customized name edit text box used for the profile name
+	TSharedPtr<SEditableTextBox> NameEditTextBox;
+	// Cached data
+	TSharedPtr<IPropertyHandle> NameProperty;
+	/** Watcher widget for the default config file (checks file status / SCC state). */
+	TSharedPtr<SSettingsEditorCheckoutNotice> FileWatcherWidget;
+
+	int32 ProfileIndex;
+	UAssetViewerSettings* ViewerSettings;
+	bool bValidProfileName;
+
+	TSharedPtr<FDataflowPreviewProfileController::IProfileIndexStorage> ProfileIndexStorage;
+};
diff --git a/Engine/Plugins/Experimental/Dataflow/Source/DataflowEditor/Private/Dataflow/DataflowConstructionViewport.cpp b/Engine/Plugins/Experimental/Dataflow/Source/DataflowEditor/Private/Dataflow/DataflowConstructionViewport.cpp
--- a/Engine/Plugins/Experimental/Dataflow/Source/DataflowEditor/Private/Dataflow/DataflowConstructionViewport.cpp
+++ b/Engine/Plugins/Experimental/Dataflow/Source/DataflowEditor/Private/Dataflow/DataflowConstructionViewport.cpp
@@ -30,6 +30,19 @@ void SDataflowConstructionViewport::Construct(const FArguments& InArgs, const FA
 	Client->VisibilityDelegate.BindSP(this, &SDataflowConstructionViewport::IsVisible);
 }
 
+FDataflowConstructionScene* SDataflowConstructionViewport::GetConstructionScene() const
+{
+	const TSharedPtr<FDataflowConstructionViewportClient> DataflowClient = StaticCastSharedPtr<FDataflowConstructionViewportClient>(Client);
+	if (DataflowClient)
+	{
+		if (const TSharedPtr<FDataflowEditorToolkit> Toolkit = DataflowClient->GetDataflowEditorToolkit().Pin())
+		{
+			return Toolkit->GetConstructionScene();
+		}
+	}
+	return nullptr;
+}
+
 TSharedPtr<SWidget> SDataflowConstructionViewport::MakeViewportToolbar()
 {
 	return SNew(SDataflowConstructionViewportSelectionToolBar, SharedThis(this))
diff --git a/Engine/Plugins/Experimental/Dataflow/Source/DataflowEditor/Private/Dataflow/DataflowConstructionViewportToolbar.cpp b/Engine/Plugins/Experimental/Dataflow/Source/DataflowEditor/Private/Dataflow/DataflowConstructionViewportToolbar.cpp
--- a/Engine/Plugins/Experimental/Dataflow/Source/DataflowEditor/Private/Dataflow/DataflowConstructionViewportToolbar.cpp
+++ b/Engine/Plugins/Experimental/Dataflow/Source/DataflowEditor/Private/Dataflow/DataflowConstructionViewportToolbar.cpp
@@ -7,7 +7,9 @@
 #include "Dataflow/DataflowConstructionVisualization.h"
 #include "Dataflow/DataflowEditorCommands.h"
 #include "Dataflow/DataflowEditorMode.h"
+#include "Dataflow/DataflowPreviewProfileController.h"
 #include "Dataflow/DataflowRenderingViewMode.h"
+#include "Dataflow/DataflowSceneProfileIndexStorage.h"
 #include "EditorModeManager.h"
 #include "SEditorViewportViewMenu.h"
 #include "ToolMenus.h"
@@ -105,6 +107,17 @@ void SDataflowConstructionViewportSelectionToolBar::Construct(const FArguments&
 				.OnGetMenuContent(this, &SDataflowConstructionViewportSelectionToolBar::GenerateShowMenu)
 		];
 
+
+	MainBoxPtr->AddSlot()
+		.AutoWidth()
+		.Padding(ToolbarSlotPadding)
+		[
+			SNew(SPreviewSceneProfileSelector).PreviewProfileController(
+				MakeShared<FDataflowPreviewProfileController>(
+					MakeShared<FDataflowConstructionSceneProfileIndexStorage>(
+						InDataflowViewport->GetConstructionScene())))
+		];
+
 	// View mode (Sim2D/Sim3D/Render)
 	MainBoxPtr->AddSlot()
 		.Padding(ToolbarSlotPadding)
diff --git a/Engine/Plugins/Experimental/Dataflow/Source/DataflowEditor/Private/Dataflow/DataflowEditorOptions.cpp b/Engine/Plugins/Experimental/Dataflow/Source/DataflowEditor/Private/Dataflow/DataflowEditorOptions.cpp
--- a/Engine/Plugins/Experimental/Dataflow/Source/DataflowEditor/Private/Dataflow/DataflowEditorOptions.cpp
+++ b/Engine/Plugins/Experimental/Dataflow/Source/DataflowEditor/Private/Dataflow/DataflowEditorOptions.cpp
@@ -1,19 +1,22 @@
 // Copyright Epic Games, Inc. All Rights Reserved.
 
 #include "Dataflow/DataflowEditorOptions.h"
+#include "AssetViewerSettings.h"
 
 #include UE_INLINE_GENERATED_CPP_BY_NAME(DataflowEditorOptions)
 
 UDataflowEditorOptions::UDataflowEditorOptions(class FObjectInitializer const& ObjectInitializer)
 {
+	// Construction
 	ConstructionViewFOV = 75.0f;
-	SimulationViewFOV = 75.0f;
-	
-	// The construction viewport often becomes overexposed in 2D with auto exposure turned on, so we default to fixed.
-	// In the future we will add the ability to tweak the construction viewport profile and materials (UE-231612)
+	// The construction viewport often becomes overexposed in 2D with auto exposure turned on especially with the "Grey Wireframe" profile active, so we default to fixed
 	bConstructionViewFixedExposure = true;
+	ConstructionProfileName = UDefaultEditorProfiles::EditingProfileName.ToString();
 
+	// Simulation
+	SimulationViewFOV = 75.0f;
 	bSimulationViewFixedExposure = false;
+	SimulationProfileName = UDefaultEditorProfiles::GreyAmbientProfileName.ToString();
 }
 
 FName UDataflowEditorOptions::GetCategoryName() const
diff --git a/Engine/Plugins/Experimental/Dataflow/Source/DataflowEditor/Private/Dataflow/DataflowEditorOptions.h b/Engine/Plugins/Experimental/Dataflow/Source/DataflowEditor/Private/Dataflow/DataflowEditorOptions.h
--- a/Engine/Plugins/Experimental/Dataflow/Source/DataflowEditor/Private/Dataflow/DataflowEditorOptions.h
+++ b/Engine/Plugins/Experimental/Dataflow/Source/DataflowEditor/Private/Dataflow/DataflowEditorOptions.h
@@ -42,6 +42,14 @@ class UDataflowEditorOptions : public UDeveloperSettings
 	UPROPERTY(config)
 	bool bSimulationViewFixedExposure;
 
+	/** Preview Scene Profile for the Construction viewport */
+	UPROPERTY(config)
+	FString ConstructionProfileName;
+
+	/** Preview Scene Profile for the Simulation viewport */
+	UPROPERTY(config)
+	FString SimulationProfileName;
+
 	/** Which mouse button controls camera panning in the Construction Viewport in 2D view mode */
 	UPROPERTY(Config, EditAnywhere, Category = UI)
 	EDataflowConstructionViewportMousePanButton ConstructionViewportMousePanButton;
diff --git a/Engine/Plugins/Experimental/Dataflow/Source/DataflowEditor/Private/Dataflow/DataflowEditorPreviewSceneBase.cpp b/Engine/Plugins/Experimental/Dataflow/Source/DataflowEditor/Private/Dataflow/DataflowEditorPreviewSceneBase.cpp
--- a/Engine/Plugins/Experimental/Dataflow/Source/DataflowEditor/Private/Dataflow/DataflowEditorPreviewSceneBase.cpp
+++ b/Engine/Plugins/Experimental/Dataflow/Source/DataflowEditor/Private/Dataflow/DataflowEditorPreviewSceneBase.cpp
@@ -6,6 +6,7 @@
 #include "Components/PrimitiveComponent.h"
 #include "Elements/Framework/EngineElementsLibrary.h"
 #include "Selection.h"
+#include "AssetViewerSettings.h"
 
 #define LOCTEXT_NAMESPACE "FDataflowPreviewSceneBase"
 
@@ -20,6 +21,10 @@ FDataflowPreviewSceneBase::FDataflowPreviewSceneBase(FPreviewScene::Construction
 	: FAdvancedPreviewScene(ConstructionValues)
 	, DataflowEditor(InEditor)
 {
+	// Remove the base class's callback and add our own
+	DefaultSettings->OnAssetViewerSettingsChanged().Remove(RefreshDelegate);
+	RefreshDelegate = DefaultSettings->OnAssetViewerSettingsChanged().AddRaw(this, &FDataflowPreviewSceneBase::OnAssetViewerSettingsRefresh);
+
 	RootSceneActor = GetWorld()->SpawnActor<AActor>(AActor::StaticClass());
 	
 	check(DataflowEditor);
@@ -105,5 +110,37 @@ FBox FDataflowPreviewSceneBase::GetBoundingBox() const
 }
 
 
+void FDataflowPreviewSceneBase::OnAssetViewerSettingsRefresh(const FName& InPropertyName)
+{
+	// If the profile was changed, update the current index and the scene.
+	if (InPropertyName == GET_MEMBER_NAME_CHECKED(FPreviewSceneProfile, ProfileName))
+	{
+		CurrentProfileIndex = DefaultSettings->Profiles.IsValidIndex(CurrentProfileIndex) ? CurrentProfileIndex : 0;
+		PreviousRotation = DefaultSettings->Profiles[CurrentProfileIndex].LightingRigRotation;
+		UILightingRigRotationDelta = 0.0f;
+
+		UpdateScene(DefaultSettings->Profiles[CurrentProfileIndex]);
+	}
+	else if (DefaultSettings->Profiles.IsValidIndex(CurrentProfileIndex))
+	{
+		const bool bNameNone = InPropertyName == NAME_None;
+
+		const bool bUpdateEnvironment = (InPropertyName == GET_MEMBER_NAME_CHECKED(FPreviewSceneProfile, EnvironmentCubeMap)) || (InPropertyName == GET_MEMBER_NAME_CHECKED(FPreviewSceneProfile, LightingRigRotation) || (InPropertyName == GET_MEMBER_NAME_CHECKED(UAssetViewerSettings, Profiles)));
+		const bool bUpdateSkyLight = bUpdateEnvironment || (InPropertyName == GET_MEMBER_NAME_CHECKED(FPreviewSceneProfile, SkyLightIntensity) || InPropertyName == GET_MEMBER_NAME_CHECKED(FPreviewSceneProfile, bUseSkyLighting) || (InPropertyName == GET_MEMBER_NAME_CHECKED(UAssetViewerSettings, Profiles)));
+		const bool bUpdateDirectionalLight = (InPropertyName == GET_MEMBER_NAME_CHECKED(FPreviewSceneProfile, DirectionalLightIntensity)) || (InPropertyName == GET_MEMBER_NAME_CHECKED(FPreviewSceneProfile, DirectionalLightColor));
+		const bool bUpdatePostProcessing = (InPropertyName == GET_MEMBER_NAME_CHECKED(FPreviewSceneProfile, PostProcessingSettings)) || (InPropertyName == GET_MEMBER_NAME_CHECKED(FPreviewSceneProfile, bPostProcessingEnabled));
+
+		UILightingRigRotationDelta += PreviousRotation - DefaultSettings->Profiles[CurrentProfileIndex].LightingRigRotation;
+		PreviousRotation = DefaultSettings->Profiles[CurrentProfileIndex].LightingRigRotation;
+
+		UpdateScene(DefaultSettings->Profiles[CurrentProfileIndex], bUpdateSkyLight || bNameNone, bUpdateEnvironment || bNameNone, bUpdatePostProcessing || bNameNone, bUpdateDirectionalLight || bNameNone);
+	}
+}
+
+void FDataflowPreviewSceneBase::SetCurrentProfileIndex(int32 NewProfileIndex)
+{
+	CurrentProfileIndex = NewProfileIndex;
+}
+
 #undef LOCTEXT_NAMESPACE
 
diff --git a/Engine/Plugins/Experimental/Dataflow/Source/DataflowEditor/Private/Dataflow/DataflowEditorToolkit.cpp b/Engine/Plugins/Experimental/Dataflow/Source/DataflowEditor/Private/Dataflow/DataflowEditorToolkit.cpp
--- a/Engine/Plugins/Experimental/Dataflow/Source/DataflowEditor/Private/Dataflow/DataflowEditorToolkit.cpp
+++ b/Engine/Plugins/Experimental/Dataflow/Source/DataflowEditor/Private/Dataflow/DataflowEditorToolkit.cpp
@@ -5,8 +5,11 @@
 #include "AdvancedPreviewScene.h"
 #include "AdvancedPreviewSceneModule.h"
 #include "AssetEditorModeManager.h"
+#include "AssetViewerSettings.h"
 #include "DetailCategoryBuilder.h"
 #include "Animation/Skeleton.h"
+#include "Dataflow/DataflowAdvancedPreviewDetailsTab.h"
+#include "Dataflow/DataflowAssetViewerSettingsCustomization.h"
 #include "Dataflow/DataflowCore.h"
 #include "Dataflow/DataflowEditor.h"
 #include "Dataflow/DataflowContent.h"
@@ -30,6 +33,7 @@
 #include "Dataflow/DataflowObjectInterface.h"
 #include "Dataflow/DataflowEditorPreviewSceneBase.h"
 #include "Dataflow/DataflowRenderingFactory.h"
+#include "Dataflow/DataflowSceneProfileIndexStorage.h"
 #include "Dataflow/DataflowSimulationScene.h"
 #include "Dataflow/DataflowSimulationVisualization.h"
 #include "Dataflow/DataflowSchema.h"
@@ -291,8 +295,13 @@ FDataflowEditorToolkit::FDataflowEditorToolkit(UAssetEditor* InOwningAssetEditor
 	PreviewSceneArgs.bShouldSimulatePhysics = 1;
 	PreviewSceneArgs.bCreatePhysicsScene = 1;
 	
-	ObjectScene = MakeUnique<FDataflowConstructionScene>(PreviewSceneArgs, DataflowEditor);
+	TUniquePtr<FDataflowConstructionScene> NewConstructionScene = MakeUnique<FDataflowConstructionScene>(PreviewSceneArgs, DataflowEditor);
+	ConstructionScene = NewConstructionScene.Get();
+	ObjectScene = MoveTemp(NewConstructionScene);
+
 	SimulationScene = MakeShared<FDataflowSimulationScene>(PreviewSceneArgs, DataflowEditor);
+
+	SimulationSceneProfileIndexStorage = MakeShared<FDataflowSimulationSceneProfileIndexStorage>(SimulationScene.Get());
 }
 
 FDataflowEditorToolkit::~FDataflowEditorToolkit()
@@ -487,6 +496,10 @@ void FDataflowEditorToolkit::PostInitAssetEditor()
 
 		Section.AddEntry(FToolMenuEntry::InitToolBarButton(NodeAndAddCommand.Value));
 	}
+
+
+	// Force scenes to update loaded asset viewer settings
+	UAssetViewerSettings::Get()->OnAssetViewerSettingsChanged().Broadcast(NAME_None);
 }
 
 void FDataflowEditorToolkit::InitializeEdMode(UBaseCharacterFXEditorMode* EdMode)
@@ -497,7 +510,7 @@ void FDataflowEditorToolkit::InitializeEdMode(UBaseCharacterFXEditorMode* EdMode
 
 	// We first set the preview scene in order to store the dynamic mesh elements
 	// generated by the tools
-	DataflowMode->SetDataflowConstructionScene(static_cast<FDataflowConstructionScene*>(ObjectScene.Get()));
+	DataflowMode->SetDataflowConstructionScene(ConstructionScene);
 
 	checkf(SimulationScene, TEXT("Expected SimulationScene to have been created in FDataflowEditorToolkit constructor"));
 	DataflowMode->SetDataflowSimlationScene(SimulationScene.Get());
@@ -600,6 +613,12 @@ TSharedRef<class IDetailCustomization> FDataflowEditorToolkit::CustomizePreviewS
 	return MakeShareable(new FDataflowPreviewSceneDescriptionCustomization(SimulationContents));
 }
 
+TSharedRef<class IDetailCustomization> FDataflowEditorToolkit::CustomizeAssetViewer() const
+{
+	return MakeShareable(new FDataflowAssetViewerSettingsCustomization(SimulationSceneProfileIndexStorage));
+}
+
+
 //~ Begin FBaseAssetToolkit overrides
 
 void FDataflowEditorToolkit::CreateWidgets()
@@ -626,18 +645,22 @@ void FDataflowEditorToolkit::CreateWidgets()
 			EditorContent->SetSelectedNode(InitialSelectedNode);
 
 			CreateSimulationViewportClient();
-
-			FAdvancedPreviewSceneModule& AdvancedPreviewSceneModule = FModuleManager::LoadModuleChecked<FAdvancedPreviewSceneModule>("AdvancedPreviewScene");
-			
+		
 			TArray<FAdvancedPreviewSceneModule::FDetailCustomizationInfo> DetailsCustomizations;
+
 			DetailsCustomizations.Add({ UDataflowSimulationSceneDescription::StaticClass(),
 				FOnGetDetailCustomizationInstance::CreateSP(const_cast<FDataflowEditorToolkit*>(this), &FDataflowEditorToolkit::CustomizePreviewSceneDescription) });
+
+			DetailsCustomizations.Add({ UAssetViewerSettings::StaticClass(),
+				FOnGetDetailCustomizationInstance::CreateSP(const_cast<FDataflowEditorToolkit*>(this), &FDataflowEditorToolkit::CustomizeAssetViewer) });
+
+			AdvancedPreviewSettingsWidget = SNew(SDataflowAdvancedPreviewDetailsTab, SimulationScene.ToSharedRef())
+				.AdditionalSettings(SimulationScene->GetPreviewSceneDescription())
+				.ProfileIndexStorage(SimulationSceneProfileIndexStorage)
+				.DetailCustomizations(DetailsCustomizations)
+				.PropertyTypeCustomizations(TArray<FAdvancedPreviewSceneModule::FPropertyTypeCustomizationInfo>())
+				.Delegates(TArray<FAdvancedPreviewSceneModule::FDetailDelegates>());
 			
-			AdvancedPreviewSettingsWidget = AdvancedPreviewSceneModule.CreateAdvancedPreviewSceneSettingsWidget(SimulationScene.ToSharedRef(), 
-				SimulationScene->GetPreviewSceneDescription(),
-				DetailsCustomizations, 
-				TArray<FAdvancedPreviewSceneModule::FPropertyTypeCustomizationInfo>(),
-				TArray<FAdvancedPreviewSceneModule::FDetailDelegates>());
 		}
 	}
 }
@@ -798,7 +821,6 @@ void FDataflowEditorToolkit::OnNodeSelectionChanged(const TSet<UObject*>& InNewS
 	auto SelectComponentsInView = [&](TObjectPtr<UDataflowEdNode> Node)
 	{
 		TArray<AActor*> FoundActors;
-		FDataflowConstructionScene* ConstructionScene = static_cast<FDataflowConstructionScene*>(ObjectScene.Get());
 		if (USelection* SelectedComponents = ConstructionScene->GetDataflowModeManager()->GetSelectedComponents())
 		{
 			SelectedComponents->Modify();
@@ -818,7 +840,7 @@ void FDataflowEditorToolkit::OnNodeSelectionChanged(const TSet<UObject*>& InNewS
 				}
 			}
 
-			if (TObjectPtr<AActor> RootActor = static_cast<FDataflowPreviewSceneBase*>(ObjectScene.Get())->GetRootActor())
+			if (TObjectPtr<AActor> RootActor = ConstructionScene->GetRootActor())
 			{
 				for (UActorComponent* ActorComponent : RootActor->GetComponents())
 				{
@@ -1745,6 +1767,11 @@ FLinearColor FDataflowEditorToolkit::GetWorldCentricTabColorScale() const
 	return FLinearColor(0.3f, 0.2f, 0.5f, 0.5f);
 }
 
+FDataflowConstructionScene* FDataflowEditorToolkit::GetConstructionScene() const
+{ 
+	return ConstructionScene;
+}
+
 void FDataflowEditorToolkit::AddReferencedObjects(FReferenceCollector& Collector)
 {
 	Collector.AddReferencedObjects(SelectedDataflowNodes);
diff --git a/Engine/Plugins/Experimental/Dataflow/Source/DataflowEditor/Private/Dataflow/DataflowPreviewProfileController.cpp b/Engine/Plugins/Experimental/Dataflow/Source/DataflowEditor/Private/Dataflow/DataflowPreviewProfileController.cpp
new file mode 100644
--- /dev/null
+++ b/Engine/Plugins/Experimental/Dataflow/Source/DataflowEditor/Private/Dataflow/DataflowPreviewProfileController.cpp
@@ -0,0 +1,120 @@
+// Copyright Epic Games, Inc. All Rights Reserved.
+
+#include "Dataflow/DataflowPreviewProfileController.h"
+#include "AssetViewerSettings.h"
+
+FDataflowPreviewProfileController::FDataflowPreviewProfileController(TSharedPtr<IProfileIndexStorage> ProfileIndexStorage) :
+	ProfileIndexStorage(ProfileIndexStorage)
+{
+	AssetViewerSettings = UAssetViewerSettings::Get();
+	if (AssetViewerSettings)
+	{
+		AssetViewerSettingsProfileAddRemoveHandle = AssetViewerSettings->OnAssetViewerProfileAddRemoved().AddLambda([this]()
+		{
+			UpdateAssetViewerProfiles();
+			OnPreviewProfileListChanged().Broadcast();
+		});
+
+		AssetViewerSettingsChangedHandle = AssetViewerSettings->OnAssetViewerSettingsChanged().AddLambda([this](const FName& InPropertyName)
+		{
+			FString CurrProfileName = AssetViewerProfileNames[CurrentProfileIndex];
+			UpdateAssetViewerProfiles();
+			if (CurrProfileName != AssetViewerProfileNames[CurrentProfileIndex])
+			{
+				OnPreviewProfileChanged().Broadcast();
+			}
+		});
+			
+		UpdateAssetViewerProfiles();
+	}
+}
+
+FDataflowPreviewProfileController::~FDataflowPreviewProfileController()
+{
+	if (IsValid(AssetViewerSettings))
+	{
+		AssetViewerSettings->OnAssetViewerProfileAddRemoved().Remove(AssetViewerSettingsProfileAddRemoveHandle);
+		AssetViewerSettings->OnAssetViewerSettingsChanged().Remove(AssetViewerSettingsChangedHandle);
+	}
+}
+
+void FDataflowPreviewProfileController::UpdateAssetViewerProfiles()
+{
+	AssetViewerProfileNames.Empty();
+
+	if (AssetViewerSettings && ProfileIndexStorage)
+	{
+		// Rebuild the profile list.
+		for (const FPreviewSceneProfile& Profile : AssetViewerSettings->Profiles)
+		{
+			AssetViewerProfileNames.Add(Profile.ProfileName);
+		}
+
+		CurrentProfileIndex = ProfileIndexStorage->RetrieveProfileIndex();
+
+		EnsureProfilesStateCoherence();
+	}
+}
+
+TArray<FString> FDataflowPreviewProfileController::GetPreviewProfiles(int32& OutCurrentProfileIndex) const
+{
+	if (AssetViewerSettings)
+	{
+		EnsureProfilesStateCoherence();
+		OutCurrentProfileIndex = CurrentProfileIndex;
+	}
+	return AssetViewerProfileNames;
+}
+
+bool FDataflowPreviewProfileController::SetActiveProfile(const FString& ProfileName)
+{
+	if (ProfileIndexStorage && IsValid(AssetViewerSettings))
+	{
+		EnsureProfilesStateCoherence();
+
+		int32 Index = AssetViewerProfileNames.IndexOfByKey(ProfileName);
+		if (Index != INDEX_NONE && Index != ProfileIndexStorage->RetrieveProfileIndex())
+		{
+			// Store the settings.
+			ProfileIndexStorage->StoreProfileIndex(Index);
+			CurrentProfileIndex = Index;
+
+			// Notify the observer about the change.
+			AssetViewerSettings->OnAssetViewerSettingsChanged().Broadcast(GET_MEMBER_NAME_CHECKED(FPreviewSceneProfile, ProfileName));
+			return true;
+		}
+	}
+	return false;
+}
+
+FString FDataflowPreviewProfileController::GetActiveProfile() const
+{
+	if (AssetViewerSettings)
+	{
+		EnsureProfilesStateCoherence();
+		return AssetViewerProfileNames[CurrentProfileIndex];
+	}
+	return FString();
+}
+
+bool FDataflowPreviewProfileController::HasAnyUserProfiles() const
+{
+	if (AssetViewerSettings)
+	{
+		for (const FPreviewSceneProfile& Profile : AssetViewerSettings->Profiles)
+		{
+			if (!Profile.bIsEngineDefaultProfile)
+			{
+				return true;
+			}
+		}
+	}
+	
+	return false;
+}
+
+void FDataflowPreviewProfileController::EnsureProfilesStateCoherence() const
+{
+	ensureMsgf(AssetViewerProfileNames.Num() == AssetViewerSettings->Profiles.Num(), TEXT("List of profiles is out of sync with the list of corresponding profile names."));
+	ensureMsgf(AssetViewerProfileNames.Num() > 0, TEXT("The list of profiles is expected to always have at least one default profile"));
+}
diff --git a/Engine/Plugins/Experimental/Dataflow/Source/DataflowEditor/Private/Dataflow/DataflowPreviewProfileController.h b/Engine/Plugins/Experimental/Dataflow/Source/DataflowEditor/Private/Dataflow/DataflowPreviewProfileController.h
new file mode 100644
--- /dev/null
+++ b/Engine/Plugins/Experimental/Dataflow/Source/DataflowEditor/Private/Dataflow/DataflowPreviewProfileController.h
@@ -0,0 +1,69 @@
+// Copyright Epic Games, Inc. All Rights Reserved.
+
+#pragma once
+
+#include "CoreMinimal.h"
+#include "IPreviewProfileController.h"
+
+class UAssetViewerSettings;
+
+// This class is almost identical to FPreviewProfileController except that it doesn't use UEditorPerProjectUserSettings 
+// to get the current scene profile index. Instead it is supplied with an IProfileIndexStorage object which stores and 
+// loads the scene profile index. This allows a separate profile controller to be created for each AdvancedPreviewScene.
+
+/** Preview Profile Controller that interfaces with a user-supplied index storage object */
+class FDataflowPreviewProfileController : public IPreviewProfileController
+{
+public:
+
+	class IProfileIndexStorage
+	{
+	public:
+		virtual void StoreProfileIndex(int32) = 0;
+		virtual int32 RetrieveProfileIndex() = 0;
+	};
+
+	FDataflowPreviewProfileController(TSharedPtr<IProfileIndexStorage> ProfileIndexStorage);
+	virtual ~FDataflowPreviewProfileController();
+
+	/** Returns the list of available preview profiles names. */
+	virtual TArray<FString> GetPreviewProfiles(int32& OutCurrentProfileIndex) const override;
+
+	/** Set the specified preview profiles as the active one. */
+	virtual bool SetActiveProfile(const FString& ProfileName) override;
+
+	/** Returns the preview profiles currently active. */
+	virtual FString GetActiveProfile() const override;
+
+	/** Returns true if user has added one or more of their own profiles */
+	virtual bool HasAnyUserProfiles() const override;
+
+	/** Invoked after the list of available profiles has changed. */
+	virtual FOnPreviewProfileListChanged& OnPreviewProfileListChanged() override { return OnPreviewProfileListChangedDelegate; }
+
+	/** Invoked after the active preview profile changed. */
+	virtual FOnPreviewProfileChanged& OnPreviewProfileChanged() override { return OnPreviewSettingChangedDelegate; }
+
+private:
+	void UpdateAssetViewerProfiles();
+	void EnsureProfilesStateCoherence() const;
+
+private:
+
+	TSharedPtr<IProfileIndexStorage> ProfileIndexStorage;
+
+	/** Provides the list of available preview profiles along with the profile change delegates (to update the profile combo box selection) */
+	UAssetViewerSettings* AssetViewerSettings = nullptr;
+
+	/** The list of available profiles. */
+	TArray<FString> AssetViewerProfileNames;
+
+	/** Holds the current profile index in the list. This is kept consistent with the cache list of names. */
+	int32 CurrentProfileIndex = 0;
+
+	FOnPreviewProfileListChanged OnPreviewProfileListChangedDelegate;
+	FOnPreviewProfileChanged OnPreviewSettingChangedDelegate;
+
+	FDelegateHandle AssetViewerSettingsProfileAddRemoveHandle;
+	FDelegateHandle AssetViewerSettingsChangedHandle;
+};
diff --git a/Engine/Plugins/Experimental/Dataflow/Source/DataflowEditor/Private/Dataflow/DataflowSceneProfileIndexStorage.cpp b/Engine/Plugins/Experimental/Dataflow/Source/DataflowEditor/Private/Dataflow/DataflowSceneProfileIndexStorage.cpp
new file mode 100644
--- /dev/null
+++ b/Engine/Plugins/Experimental/Dataflow/Source/DataflowEditor/Private/Dataflow/DataflowSceneProfileIndexStorage.cpp
@@ -0,0 +1,102 @@
+// Copyright Epic Games, Inc. All Rights Reserved.
+
+#include "DataflowSceneProfileIndexStorage.h"
+#include "Dataflow/DataflowConstructionScene.h"
+#include "Dataflow/DataflowSimulationScene.h"
+#include "Dataflow/DataflowEditorOptions.h"
+#include "AssetViewerSettings.h"
+
+
+namespace Private
+{
+	int32 GetProfileIndex(const FString& SearchName)
+	{
+		if (UAssetViewerSettings* AssetViewerSettings = UAssetViewerSettings::Get())
+		{
+			const int32 FoundIndex = AssetViewerSettings->Profiles.IndexOfByPredicate([SearchName](const FPreviewSceneProfile& Profile)
+			{
+				return Profile.ProfileName == SearchName;
+			});
+
+			return FoundIndex;
+		}
+
+		return INDEX_NONE;
+	}
+}
+
+FDataflowConstructionSceneProfileIndexStorage::FDataflowConstructionSceneProfileIndexStorage(FDataflowConstructionScene* ConstructionScene) :
+	ConstructionScene(ConstructionScene)
+{
+	checkf(ConstructionScene, TEXT("Received invalid construction scene pointer"));
+}
+
+void FDataflowConstructionSceneProfileIndexStorage::StoreProfileIndex(int32 Index)
+{
+	if (UAssetViewerSettings* AssetViewerSettings = UAssetViewerSettings::Get())
+	{
+		if (AssetViewerSettings->Profiles.IsValidIndex(Index))
+		{
+			if (UDataflowEditorOptions* const Options = UDataflowEditorOptions::StaticClass()->GetDefaultObject<UDataflowEditorOptions>())
+			{
+				Options->ConstructionProfileName = AssetViewerSettings->Profiles[Index].ProfileName;
+				Options->SaveConfig();
+			}
+
+			ConstructionScene->SetCurrentProfileIndex(Index);
+		}
+	}
+}
+
+int32 FDataflowConstructionSceneProfileIndexStorage::RetrieveProfileIndex()
+{
+	if (const UDataflowEditorOptions* const Options = UDataflowEditorOptions::StaticClass()->GetDefaultObject<UDataflowEditorOptions>())
+	{
+		const int32 FoundIndex = Private::GetProfileIndex(Options->ConstructionProfileName);
+		if (FoundIndex != ConstructionScene->GetCurrentProfileIndex())
+		{
+			ConstructionScene->SetCurrentProfileIndex(FoundIndex);
+		}
+		return FoundIndex;
+	}
+	return INDEX_NONE;
+}
+
+
+FDataflowSimulationSceneProfileIndexStorage::FDataflowSimulationSceneProfileIndexStorage(FDataflowSimulationScene* SimulationScene) :
+	SimulationScene(SimulationScene)
+{
+	checkf(SimulationScene, TEXT("Received invalid simulation scene pointer"));
+}
+
+void FDataflowSimulationSceneProfileIndexStorage::StoreProfileIndex(int32 Index)
+{
+	if (UAssetViewerSettings* AssetViewerSettings = UAssetViewerSettings::Get())
+	{
+		if (AssetViewerSettings->Profiles.IsValidIndex(Index))
+		{
+			if (UDataflowEditorOptions* const Options = UDataflowEditorOptions::StaticClass()->GetDefaultObject<UDataflowEditorOptions>())
+			{
+				Options->SimulationProfileName = AssetViewerSettings->Profiles[Index].ProfileName;
+				Options->SaveConfig();
+			}
+
+			SimulationScene->SetCurrentProfileIndex(Index);
+		}
+	}
+}
+
+int32 FDataflowSimulationSceneProfileIndexStorage::RetrieveProfileIndex()
+{
+	if (const UDataflowEditorOptions* const Options = UDataflowEditorOptions::StaticClass()->GetDefaultObject<UDataflowEditorOptions>())
+	{
+		const int32 FoundIndex = Private::GetProfileIndex(Options->SimulationProfileName);
+		if (FoundIndex != SimulationScene->GetCurrentProfileIndex())
+		{
+			SimulationScene->SetCurrentProfileIndex(FoundIndex);
+		}
+		return FoundIndex;
+	}
+
+	return INDEX_NONE;
+}
diff --git a/Engine/Plugins/Experimental/Dataflow/Source/DataflowEditor/Private/Dataflow/DataflowSceneProfileIndexStorage.h b/Engine/Plugins/Experimental/Dataflow/Source/DataflowEditor/Private/Dataflow/DataflowSceneProfileIndexStorage.h
new file mode 100644
--- /dev/null
+++ b/Engine/Plugins/Experimental/Dataflow/Source/DataflowEditor/Private/Dataflow/DataflowSceneProfileIndexStorage.h
@@ -0,0 +1,39 @@
+// Copyright Epic Games, Inc. All Rights Reserved.
+
+#pragma once
+
+#include "Dataflow/DataflowPreviewProfileController.h"
+
+class FDataflowConstructionScene;
+class FDataflowSimulationScene;
+
+class FDataflowConstructionSceneProfileIndexStorage : public FDataflowPreviewProfileController::IProfileIndexStorage
+{
+public:
+
+	FDataflowConstructionSceneProfileIndexStorage(FDataflowConstructionScene* ConstructionScene);
+	virtual ~FDataflowConstructionSceneProfileIndexStorage() = default;
+
+	virtual void StoreProfileIndex(int32 Index) override;
+	virtual int32 RetrieveProfileIndex() override;
+
+private:
+
+	FDataflowConstructionScene* ConstructionScene;
+};
+
+class FDataflowSimulationSceneProfileIndexStorage : public FDataflowPreviewProfileController::IProfileIndexStorage
+{
+public:
+
+	FDataflowSimulationSceneProfileIndexStorage(FDataflowSimulationScene* SimulationScene);
+	virtual ~FDataflowSimulationSceneProfileIndexStorage() = default;
+
+	virtual void StoreProfileIndex(int32 Index) override;
+	virtual int32 RetrieveProfileIndex() override;
+
+private:
+
+	FDataflowSimulationScene* SimulationScene;
+};
+
diff --git a/Engine/Plugins/Experimental/Dataflow/Source/DataflowEditor/Private/Dataflow/DataflowSimulationViewportToolbar.cpp b/Engine/Plugins/Experimental/Dataflow/Source/DataflowEditor/Private/Dataflow/DataflowSimulationViewportToolbar.cpp
--- a/Engine/Plugins/Experimental/Dataflow/Source/DataflowEditor/Private/Dataflow/DataflowSimulationViewportToolbar.cpp
+++ b/Engine/Plugins/Experimental/Dataflow/Source/DataflowEditor/Private/Dataflow/DataflowSimulationViewportToolbar.cpp
@@ -1,7 +1,11 @@
 // Copyright Epic Games, Inc. All Rights Reserved.
 #include "Dataflow/DataflowSimulationViewportToolbar.h"
+#include "Dataflow/DataflowEditorToolkit.h"
+#include "Dataflow/DataflowPreviewProfileController.h"
 #include "Dataflow/DataflowSimulationViewport.h"
+#include "Dataflow/DataflowSimulationViewportClient.h"
 #include "Dataflow/DataflowEditorCommands.h"
+#include "Dataflow/DataflowSceneProfileIndexStorage.h"
 #include "Dataflow/DataflowSimulationScene.h"
 #include "Styling/AppStyle.h"
 #include "Framework/MultiBox/MultiBoxBuilder.h"
@@ -17,8 +21,32 @@ void SDataflowSimulationViewportToolBar::Construct(const FArguments& InArgs, TSh
 	EditorViewport = InDataflowViewport;
 	CommandList = InArgs._CommandList;
 	Extenders = InArgs._Extenders;
+
+	TSharedPtr<FDataflowSimulationSceneProfileIndexStorage> ProfileIndexStorage;
+
+	if (const TSharedPtr<SDataflowSimulationViewport> Viewport = EditorViewport.Pin())
+	{
+		if (const TSharedPtr<FDataflowSimulationViewportClient> ViewportClient = StaticCastSharedPtr<FDataflowSimulationViewportClient>(Viewport->GetViewportClient()))
+		{
+			if (ViewportClient->GetDataflowEditorToolkit().IsValid())
+			{
+				if (const TSharedPtr<FDataflowEditorToolkit> Toolkit = ViewportClient->GetDataflowEditorToolkit().Pin())
+				{
+					ProfileIndexStorage = Toolkit->GetSimulationSceneProfileIndexStorage();
+				}
+			}
+		}
+	}
 	
-	SCommonEditorViewportToolbarBase::Construct(SCommonEditorViewportToolbarBase::FArguments(), InDataflowViewport);
+	if (!ProfileIndexStorage)
+	{
+		ProfileIndexStorage = MakeShared<FDataflowSimulationSceneProfileIndexStorage>(InDataflowViewport->GetSimulationScene().Get());
+	}
+	
+	SCommonEditorViewportToolbarBase::FArguments Args = SCommonEditorViewportToolbarBase::FArguments().
+		PreviewProfileController(MakeShared<FDataflowPreviewProfileController>(ProfileIndexStorage));
+
+	SCommonEditorViewportToolbarBase::Construct(Args, InDataflowViewport);
 }
 
 void SDataflowSimulationViewportToolBar::ExtendLeftAlignedToolbarSlots(TSharedPtr<SHorizontalBox> MainBoxPtr, TSharedPtr<SViewportToolBar> ParentToolBarPtr) const
diff --git a/Engine/Plugins/Experimental/Dataflow/Source/DataflowEditor/Public/Dataflow/DataflowConstructionViewport.h b/Engine/Plugins/Experimental/Dataflow/Source/DataflowEditor/Public/Dataflow/DataflowConstructionViewport.h
--- a/Engine/Plugins/Experimental/Dataflow/Source/DataflowEditor/Public/Dataflow/DataflowConstructionViewport.h
+++ b/Engine/Plugins/Experimental/Dataflow/Source/DataflowEditor/Public/Dataflow/DataflowConstructionViewport.h
@@ -9,6 +9,7 @@ class ADataflowActor;
 class FAdvancedPreviewScene;
 class UDataflowEditorMode;
 class FDataflowConstructionViewportClient;
+class FDataflowConstructionScene;
 
 // ----------------------------------------------------------------------------------
 
@@ -23,6 +24,8 @@ public:
 
 	void Construct(const FArguments& InArgs, const FAssetEditorViewportConstructionArgs& InViewportConstructionArgs);
 
+	FDataflowConstructionScene* GetConstructionScene() const;
+
 	// SEditorViewport
 	virtual void BindCommands() override;
 	virtual TSharedPtr<SWidget> MakeViewportToolbar() override;
diff --git a/Engine/Plugins/Experimental/Dataflow/Source/DataflowEditor/Public/Dataflow/DataflowEditorPreviewSceneBase.h b/Engine/Plugins/Experimental/Dataflow/Source/DataflowEditor/Public/Dataflow/DataflowEditorPreviewSceneBase.h
--- a/Engine/Plugins/Experimental/Dataflow/Source/DataflowEditor/Public/Dataflow/DataflowEditorPreviewSceneBase.h
+++ b/Engine/Plugins/Experimental/Dataflow/Source/DataflowEditor/Public/Dataflow/DataflowEditorPreviewSceneBase.h
@@ -50,9 +50,16 @@ public:
 
 	/** Check if the preview scene can run simulation */
 	virtual bool CanRunSimulation() const {return false;}
-	
+
+	/** Update the currently selected scene profile */
+	void SetCurrentProfileIndex(int32 NewProfileIndex);
+
+
 protected:
-	
+
+	/** Respond to changes in the scene profile settings */
+	void OnAssetViewerSettingsRefresh(const FName& InPropertyName);
+
 	/** Root scene actor */
 	TObjectPtr<AActor> RootSceneActor = nullptr; 
 
diff --git a/Engine/Plugins/Experimental/Dataflow/Source/DataflowEditor/Public/Dataflow/DataflowEditorToolkit.h b/Engine/Plugins/Experimental/Dataflow/Source/DataflowEditor/Public/Dataflow/DataflowEditorToolkit.h
--- a/Engine/Plugins/Experimental/Dataflow/Source/DataflowEditor/Public/Dataflow/DataflowEditorToolkit.h
+++ b/Engine/Plugins/Experimental/Dataflow/Source/DataflowEditor/Public/Dataflow/DataflowEditorToolkit.h
@@ -7,13 +7,11 @@
 #include "Dataflow/DataflowObjectInterface.h"
 #include "Misc/NotifyHook.h"
 #include "GraphEditor.h"
-#include "SAdvancedPreviewDetailsTab.h"
 #include "TickableEditorObject.h"
 #include "Dataflow/DataflowSelectionView.h"
 #include "Dataflow/DataflowCollectionSpreadSheet.h"
 #include "Dataflow/DataflowConstructionViewport.h"
 #include "Dataflow/DataflowSimulationViewport.h"
-//#include "UObject\GCObject.h"
 
 class FEditorViewportTabContent;
 class IDetailsView;
@@ -30,7 +28,7 @@ class FDataflowSimulationScene;
 class FDataflowSkeletonView;
 class UDataflowEditor;
 class SDataflowMembersWidget;
-
+class FDataflowSimulationSceneProfileIndexStorage;
 namespace UE::Dataflow
 {
 	class FDataflowNodeDetailExtensionHandler;
@@ -74,6 +72,8 @@ public:
 
 	/** Dataflow preview scenes accessor */
 	const TSharedPtr<FDataflowSimulationScene>& GetSimulationScene() const {return SimulationScene;}
+	FDataflowConstructionScene* GetConstructionScene() const;
+	const TSharedPtr<FDataflowSimulationSceneProfileIndexStorage>& GetSimulationSceneProfileIndexStorage() const { return SimulationSceneProfileIndexStorage; }
 
 	// FSerializableObject interface
 	virtual void AddReferencedObjects(FReferenceCollector& Collector) override;
@@ -180,6 +180,8 @@ private:
 	/** Customize preview scene with editor/terminal contents */
 	TSharedRef<class IDetailCustomization> CustomizePreviewSceneDescription() const;
 
+	TSharedRef<class IDetailCustomization> CustomizeAssetViewer() const;
+
 	// Utility factory functions to build the widgets
 	TSharedRef<SDataflowGraphEditor> CreateGraphEditorWidget(UDataflow* ObjectToEdit, TSharedPtr<IStructureDetailsView> PropertiesEditor);
     TSharedPtr<IDetailsView> CreateAssetDetailsEditorWidget(const TArray<UObject*>& ObjectsToEdit);
@@ -212,9 +214,14 @@ private:
 	UPROPERTY()
 	TObjectPtr<UDataflowEdNode> PrimarySelection;
 	
+	/** Pointer to the construction viewport scene. Note this is an alias of ObjectScene in FBaseCharacterFXEditorToolkit but with the specific type */
+	FDataflowConstructionScene* ConstructionScene;
+
 	/** PreviewScene showing the objects being simulated */
 	TSharedPtr<FDataflowSimulationScene> SimulationScene;
 
+	TSharedPtr<FDataflowSimulationSceneProfileIndexStorage> SimulationSceneProfileIndexStorage;
+
 	/** The editor mode manager used by the simulation preview scene */
 	TSharedPtr<FEditorModeTools> SimulationModeManager;
 
