Dataflow Editor: cache the value of the p.Dataflow.EnableGraphEval cvar to avoid calling FindConsoleVariable too often.
Also set the node rendering to be more opaque when graph evaluation is paused.

diff --git a/Engine/Plugins/Experimental/Dataflow/Source/DataflowEditor/Private/Dataflow/DataflowEditorToolkit.cpp b/Engine/Plugins/Experimental/Dataflow/Source/DataflowEditor/Private/Dataflow/DataflowEditorToolkit.cpp
--- a/Engine/Plugins/Experimental/Dataflow/Source/DataflowEditor/Private/Dataflow/DataflowEditorToolkit.cpp
+++ b/Engine/Plugins/Experimental/Dataflow/Source/DataflowEditor/Private/Dataflow/DataflowEditorToolkit.cpp
@@ -91,6 +91,7 @@ const FName FDataflowEditorToolkit::MembersWidgetTabId(TEXT("DataflowEditor_Memb
 FDataflowEditorToolkit::FDataflowEditorToolkit(UAssetEditor* InOwningAssetEditor)
 	: FBaseCharacterFXEditorToolkit(InOwningAssetEditor, FName("DataflowEditor"))
 	, DataflowEditor(Cast< UDataflowEditor>(InOwningAssetEditor))
+	, bDataflowEnableGraphEval(true)
 {
 
 	// When saving, only prompt to checkout and save assets that are actually modified
@@ -302,10 +303,27 @@ FDataflowEditorToolkit::FDataflowEditorToolkit(UAssetEditor* InOwningAssetEditor
 	SimulationScene = MakeShared<FDataflowSimulationScene>(PreviewSceneArgs, DataflowEditor);
 
 	SimulationSceneProfileIndexStorage = MakeShared<FDataflowSimulationSceneProfileIndexStorage>(SimulationScene.Get());
+
+	IConsoleVariable* const ConsoleVar = IConsoleManager::Get().FindConsoleVariable(TEXT("p.Dataflow.EnableGraphEval"));
+	bDataflowEnableGraphEval = ConsoleVar ? ConsoleVar->GetBool() : true;
+
+	if (ConsoleVar)
+	{
+		GraphEvalCVarChangedDelegateHandle = ConsoleVar->OnChangedDelegate().AddLambda([this](IConsoleVariable* Var)
+		{
+			bDataflowEnableGraphEval = Var->GetBool();
+		});
+	}
+
 }
 
 FDataflowEditorToolkit::~FDataflowEditorToolkit()
 {
+	if (IConsoleVariable* const ConsoleVar = IConsoleManager::Get().FindConsoleVariable(TEXT("p.Dataflow.EnableGraphEval")))
+	{
+		ConsoleVar->OnChangedDelegate().Remove(GraphEvalCVarChangedDelegateHandle);
+	}
+
 	if (SimulationScene && SimulationScene->GetPreviewSceneDescription())
 	{
 		SimulationScene->GetPreviewSceneDescription()->DataflowSimulationSceneDescriptionChanged.Remove(OnSimulationSceneChangedDelegateHandle);
@@ -1168,9 +1186,6 @@ void FDataflowEditorToolkit::EvaluateNode(const FDataflowNode* Node, const FData
 {
 	UE_LOG(LogChaosDataflow, VeryVerbose, TEXT("FDataflowEditorToolkit::EvaluateNode(): Node [%s], Output [%s]"), Node ? *Node->GetName().ToString() : TEXT("nullptr"), Output ? *Output->GetName().ToString() : TEXT("nullptr"));
 
-	const IConsoleVariable* const ConsoleVar = IConsoleManager::Get().FindConsoleVariable(TEXT("p.Dataflow.EnableGraphEval"));
-	const bool bDataflowEnableGraphEval = ConsoleVar ? ConsoleVar->GetBool() : true;
-
 	if (bDataflowEnableGraphEval)
 	{
 		const bool bIsInPIEOrSimulate = GEditor->PlayWorld || GEditor->bIsSimulatingInEditor;
diff --git a/Engine/Plugins/Experimental/Dataflow/Source/DataflowEditor/Private/Dataflow/DataflowGraphEditor.cpp b/Engine/Plugins/Experimental/Dataflow/Source/DataflowEditor/Private/Dataflow/DataflowGraphEditor.cpp
--- a/Engine/Plugins/Experimental/Dataflow/Source/DataflowEditor/Private/Dataflow/DataflowGraphEditor.cpp
+++ b/Engine/Plugins/Experimental/Dataflow/Source/DataflowEditor/Private/Dataflow/DataflowGraphEditor.cpp
@@ -808,8 +808,6 @@ void SDataflowGraphEditor::InitGraphEditorMessageBar()
 
 	auto UpdateGraphForEvalEnabled = [this](bool bGraphEvalEnabled)
 	{
-		SetColorAndOpacity(bGraphEvalEnabled ? FLinearColor::White : FLinearColor(1.0f, 1.0f, 1.0f, 0.85f));
-
 		if (EdGraphObj)
 		{
 			for (UEdGraphNode* const Node : EdGraphObj->Nodes)
diff --git a/Engine/Plugins/Experimental/Dataflow/Source/DataflowEditor/Public/Dataflow/DataflowEditorToolkit.h b/Engine/Plugins/Experimental/Dataflow/Source/DataflowEditor/Public/Dataflow/DataflowEditorToolkit.h
--- a/Engine/Plugins/Experimental/Dataflow/Source/DataflowEditor/Public/Dataflow/DataflowEditorToolkit.h
+++ b/Engine/Plugins/Experimental/Dataflow/Source/DataflowEditor/Public/Dataflow/DataflowEditorToolkit.h
@@ -239,4 +239,10 @@ private:
 	
 	/** Simulation default layout */
 	TSharedPtr<FTabManager::FLayout> ConstructionDefaultLayout;
+
+	/** Cached value of the p.Dataflow.EnableGraphEval cvar, to avoid calling FindConsoleVariable too often */
+	bool bDataflowEnableGraphEval;
+
+	/** Delegate for updating the cached value of p.Dataflow.EnableGraphEval */
+	FDelegateHandle GraphEvalCVarChangedDelegateHandle;
 };
