Dataflow : move SDataflowPin class out of the SDataflowNode files

diff --git a/Engine/Plugins/Experimental/Dataflow/Source/DataflowEditor/Private/Dataflow/DataflowSNode.cpp b/Engine/Plugins/Experimental/Dataflow/Source/DataflowEditor/Private/Dataflow/DataflowSNode.cpp
--- a/Engine/Plugins/Experimental/Dataflow/Source/DataflowEditor/Private/Dataflow/DataflowSNode.cpp
+++ b/Engine/Plugins/Experimental/Dataflow/Source/DataflowEditor/Private/Dataflow/DataflowSNode.cpp
@@ -21,81 +21,12 @@
 #include "Editor/Transactor.h"
 #include "GraphEditorSettings.h"
 #include "Dataflow/DataflowSchema.h"
+#include "Dataflow/DataflowSPin.h"
 
 #include UE_INLINE_GENERATED_CPP_BY_NAME(DataflowSNode)
 
 #define LOCTEXT_NAMESPACE "SDataflowEdNode"
 
-//
-// SDataflowPin
-//
-
-void SDataflowPin::Construct(const FArguments& InArgs, UEdGraphPin* InPin)
-{
-	bIsPinInvalid = InArgs._IsPinInvalid.Get();
-	PinColorOverride = InArgs._PinColorOverride.Get();
-	bIsPinColorOverriden = InArgs._bIsPinColorOverriden.Get();
-
-	const FText InvalidPinDisplayText = bIsPinInvalid.Get() ? NSLOCTEXT("DataflowGraph", "DataflowOutputPinInvalidText", "*") : NSLOCTEXT("DataflowGraph", "DataflowOutputPinValidText", " ");
-
-	SGraphPin::Construct(SGraphPin::FArguments(), InPin);
-
-	GetLabelAndValue()->AddSlot()
-		.Padding(2.0f, 0.0f, 0.0f, 0.0f)
-		[
-			SNew(STextBlock)
-			.Text_Lambda([InvalidPinDisplayText]()
-				{
-					return InvalidPinDisplayText;
-				})
-		.MinDesiredWidth(5)
-		];
-}
-
-FSlateColor SDataflowPin::GetPinColor() const
-{
-	const bool IsColorOverriden = bIsPinColorOverriden.Get();
-	const FLinearColor ColorOverride = PinColorOverride.Get();
-
-	UEdGraphPin* GraphPin = GetPinObj();
-	if (GraphPin && !GraphPin->IsPendingKill())
-	{
-		if (bIsDiffHighlighted)
-		{
-			return FSlateColor(FLinearColor(0.9f, 0.2f, 0.15f));
-		}
-		if (GraphPin->bOrphanedPin)
-		{
-			return FSlateColor(FLinearColor::Red);
-		}
-		if (const UEdGraphSchema* Schema = GraphPin->GetSchema())
-		{
-			if (!GetPinObj()->GetOwningNode()->IsNodeEnabled() || GetPinObj()->GetOwningNode()->IsDisplayAsDisabledForced() || !IsEditingEnabled() || GetPinObj()->GetOwningNode()->IsNodeUnrelated())
-			{
-				if (IsColorOverriden)
-				{
-					return FSlateColor(ColorOverride);
-				}
-
-				return Schema->GetPinTypeColor(GraphPin->PinType) * FLinearColor(1.0f, 1.0f, 1.0f, 0.5f);
-			}
-
-			if (IsColorOverriden)
-			{
-				return FSlateColor(ColorOverride);
-			}
-
-			return Schema->GetPinTypeColor(GraphPin->PinType) * PinColorModifier;
-		}
-	}
-
-	return FLinearColor::White;
-}
-
-//
-// SDataflowEdNode
-//
-
 const FSlateBrush* SDataflowEdNode::GetPinButtonImage() const
 {
 	if (DataflowGraphNode && DataflowGraphNode->ShouldWireframeRenderNode())
diff --git a/Engine/Plugins/Experimental/Dataflow/Source/DataflowEditor/Private/Dataflow/DataflowSPin.cpp b/Engine/Plugins/Experimental/Dataflow/Source/DataflowEditor/Private/Dataflow/DataflowSPin.cpp
new file mode 100644
--- /dev/null
+++ b/Engine/Plugins/Experimental/Dataflow/Source/DataflowEditor/Private/Dataflow/DataflowSPin.cpp
@@ -0,0 +1,69 @@
+// Copyright Epic Games, Inc. All Rights Reserved.
+
+#include "Dataflow/DataflowSPin.h"
+
+#define LOCTEXT_NAMESPACE "SDataflowPin"
+
+void SDataflowPin::Construct(const FArguments& InArgs, UEdGraphPin* InPin)
+{
+	bIsPinInvalid = InArgs._IsPinInvalid.Get();
+	PinColorOverride = InArgs._PinColorOverride.Get();
+	bIsPinColorOverriden = InArgs._bIsPinColorOverriden.Get();
+
+	const FText InvalidPinDisplayText = bIsPinInvalid.Get() ? LOCTEXT("DataflowOutputPinInvalidText", "*") : LOCTEXT("DataflowOutputPinValidText", " ");
+
+	SGraphPin::Construct(SGraphPin::FArguments(), InPin);
+
+	GetLabelAndValue()->AddSlot()
+		.Padding(2.0f, 0.0f, 0.0f, 0.0f)
+		[
+			SNew(STextBlock)
+			.Text_Lambda([InvalidPinDisplayText]()
+				{
+					return InvalidPinDisplayText;
+				})
+		.MinDesiredWidth(5)
+		];
+}
+
+FSlateColor SDataflowPin::GetPinColor() const
+{
+	const bool IsColorOverriden = bIsPinColorOverriden.Get();
+	const FLinearColor ColorOverride = PinColorOverride.Get();
+
+	UEdGraphPin* GraphPin = GetPinObj();
+	if (GraphPin && !GraphPin->IsPendingKill())
+	{
+		if (bIsDiffHighlighted)
+		{
+			return FSlateColor(FLinearColor(0.9f, 0.2f, 0.15f));
+		}
+		if (GraphPin->bOrphanedPin)
+		{
+			return FSlateColor(FLinearColor::Red);
+		}
+		if (const UEdGraphSchema* Schema = GraphPin->GetSchema())
+		{
+			if (!GetPinObj()->GetOwningNode()->IsNodeEnabled() || GetPinObj()->GetOwningNode()->IsDisplayAsDisabledForced() || !IsEditingEnabled() || GetPinObj()->GetOwningNode()->IsNodeUnrelated())
+			{
+				if (IsColorOverriden)
+				{
+					return FSlateColor(ColorOverride);
+				}
+
+				return Schema->GetPinTypeColor(GraphPin->PinType) * FLinearColor(1.0f, 1.0f, 1.0f, 0.5f);
+			}
+
+			if (IsColorOverriden)
+			{
+				return FSlateColor(ColorOverride);
+			}
+
+			return Schema->GetPinTypeColor(GraphPin->PinType) * PinColorModifier;
+		}
+	}
+
+	return FLinearColor::White;
+}
+#undef LOCTEXT_NAMESPACE
+
diff --git a/Engine/Plugins/Experimental/Dataflow/Source/DataflowEditor/Public/Dataflow/DataflowSNode.h b/Engine/Plugins/Experimental/Dataflow/Source/DataflowEditor/Public/Dataflow/DataflowSNode.h
--- a/Engine/Plugins/Experimental/Dataflow/Source/DataflowEditor/Public/Dataflow/DataflowSNode.h
+++ b/Engine/Plugins/Experimental/Dataflow/Source/DataflowEditor/Public/Dataflow/DataflowSNode.h
@@ -7,39 +7,12 @@
 #include "UObject/GCObject.h"
 #include "Dataflow/DataflowCore.h"
 #include "Dataflow/DataflowSEditorInterface.h"
-#include "SGraphPin.h"
 
 #include "DataflowSNode.generated.h"
 
 class UDataflowEdNode;
 class SCheckBox;
-
-//
-// SDataflowPin
-//
-
-class SDataflowPin : public SGraphPin
-{
-public:
-	SLATE_BEGIN_ARGS(SDataflowPin)
-		: _IsPinInvalid(false)
-		, _PinColorOverride(FLinearColor::White)
-		, _bIsPinColorOverriden(false)
-	{}
-		SLATE_ATTRIBUTE(bool, IsPinInvalid)
-		SLATE_ATTRIBUTE(FLinearColor, PinColorOverride)
-		SLATE_ATTRIBUTE(bool, bIsPinColorOverriden)
-	SLATE_END_ARGS()
-
-	void Construct(const FArguments& InArgs, UEdGraphPin* InPin);
-
-	virtual FSlateColor GetPinColor() const override;
-
-private:
-	TAttribute<bool> bIsPinInvalid;
-	TAttribute<FLinearColor> PinColorOverride;
-	TAttribute<bool> bIsPinColorOverriden;
-};
+class SImage;
 
 //
 // SDataflowEdNode
diff --git a/Engine/Plugins/Experimental/Dataflow/Source/DataflowEditor/Public/Dataflow/DataflowSPin.h b/Engine/Plugins/Experimental/Dataflow/Source/DataflowEditor/Public/Dataflow/DataflowSPin.h
new file mode 100644
--- /dev/null
+++ b/Engine/Plugins/Experimental/Dataflow/Source/DataflowEditor/Public/Dataflow/DataflowSPin.h
@@ -0,0 +1,32 @@
+// Copyright Epic Games, Inc. All Rights Reserved.
+
+#pragma once
+
+#include "SGraphPin.h"
+
+/**
+* class used to for all dataflow node pins 
+*/
+class SDataflowPin : public SGraphPin
+{
+public:
+	SLATE_BEGIN_ARGS(SDataflowPin)
+		: _IsPinInvalid(false)
+		, _PinColorOverride(FLinearColor::White)
+		, _bIsPinColorOverriden(false)
+	{}
+		SLATE_ATTRIBUTE(bool, IsPinInvalid)
+		SLATE_ATTRIBUTE(FLinearColor, PinColorOverride)
+		SLATE_ATTRIBUTE(bool, bIsPinColorOverriden)
+	SLATE_END_ARGS()
+
+	void Construct(const FArguments& InArgs, UEdGraphPin* InPin);
+
+	virtual FSlateColor GetPinColor() const override;
+
+private:
+	TAttribute<bool> bIsPinInvalid;
+	TAttribute<FLinearColor> PinColorOverride;
+	TAttribute<bool> bIsPinColorOverriden;
+};
+
