Dataflow : move graph schema action to their own file ( out of the DataflowSNode file )

diff --git a/Engine/Plugins/ChaosClothAssetEditor/Source/ChaosClothAssetEditor/Private/ChaosClothAsset/ClothEditorMode.cpp b/Engine/Plugins/ChaosClothAssetEditor/Source/ChaosClothAssetEditor/Private/ChaosClothAsset/ClothEditorMode.cpp
--- a/Engine/Plugins/ChaosClothAssetEditor/Source/ChaosClothAssetEditor/Private/ChaosClothAsset/ClothEditorMode.cpp
+++ b/Engine/Plugins/ChaosClothAssetEditor/Source/ChaosClothAssetEditor/Private/ChaosClothAsset/ClothEditorMode.cpp
@@ -58,7 +58,7 @@
 #include "BaseGizmos/TransformGizmoUtil.h"
 #include "Materials/Material.h"
 #include "MaterialDomain.h"
-#include "Dataflow/DataflowSNode.h"
+#include "Dataflow/DataflowGraphSchemaAction.h"
 #include "Components/SkeletalMeshComponent.h"
 #include "ChaosClothAsset/ClothGeometryTools.h"
 #include "DynamicMesh/NonManifoldMappingSupport.h"
diff --git a/Engine/Plugins/Experimental/Dataflow/Source/DataflowEditor/Private/Dataflow/DataflowGraphEditor.cpp b/Engine/Plugins/Experimental/Dataflow/Source/DataflowEditor/Private/Dataflow/DataflowGraphEditor.cpp
--- a/Engine/Plugins/Experimental/Dataflow/Source/DataflowEditor/Private/Dataflow/DataflowGraphEditor.cpp
+++ b/Engine/Plugins/Experimental/Dataflow/Source/DataflowEditor/Private/Dataflow/DataflowGraphEditor.cpp
@@ -12,7 +12,7 @@
 #include "Dataflow/DataflowSCommentNode.h"
 #include "Dataflow/DataflowNodeParameters.h"
 #include "Dataflow/DataflowXml.h"
-#include "Dataflow/DataflowSNode.h"
+#include "Dataflow/DataflowGraphSchemaAction.h"
 #include "Dataflow/DataflowSelectionNodes.h"
 #include "Dataflow/DataflowSchema.h"
 #include "HAL/PlatformApplicationMisc.h"
diff --git a/Engine/Plugins/Experimental/Dataflow/Source/DataflowEditor/Private/Dataflow/DataflowGraphSchemaAction.cpp b/Engine/Plugins/Experimental/Dataflow/Source/DataflowEditor/Private/Dataflow/DataflowGraphSchemaAction.cpp
--- a/Engine/Plugins/Experimental/Dataflow/Source/DataflowEditor/Private/Dataflow/DataflowGraphSchemaAction.cpp
+++ b/Engine/Plugins/Experimental/Dataflow/Source/DataflowEditor/Private/Dataflow/DataflowGraphSchemaAction.cpp
@@ -16,6 +16,261 @@
 
 /////////////////////////////////////////////////////////////////////////////////////////////////////////////////////////
 
+TSharedPtr<FAssetSchemaAction_Dataflow_CreateNode_DataflowEdNode> FAssetSchemaAction_Dataflow_CreateNode_DataflowEdNode::CreateAction(const UEdGraph* ParentGraph, const FName& InNodeTypeName, const FName& InOverrideNodeName)
+{
+	if (const UDataflow* Dataflow = Cast<UDataflow>(ParentGraph))
+	{
+		if (UE::Dataflow::FNodeFactory* Factory = UE::Dataflow::FNodeFactory::GetInstance())
+		{
+			const UE::Dataflow::FFactoryParameters& Param = Factory->GetParameters(InNodeTypeName);
+			if (Param.IsValid())
+			{
+				const bool bIsSimulationNode = Param.Tags.Contains(UDataflow::SimulationTag);
+				const bool bIsSimulationGraph = (Dataflow->Type == EDataflowType::Simulation);
+
+				if ((bIsSimulationGraph && bIsSimulationNode) || (!bIsSimulationGraph && !bIsSimulationNode))
+				{
+					const FText ToolTip = FText::FromString(Param.ToolTip.IsEmpty() ? FString("Add a Dataflow node.") : Param.ToolTip);
+					FText NodeName = FText::FromString(Param.DisplayName.ToString());
+					if (!InOverrideNodeName.IsNone())
+					{
+						NodeName = FText::FromName(InOverrideNodeName);
+					}
+
+					const FText Category = FText::FromString(Param.Category.ToString().IsEmpty() ? FString("Dataflow") : Param.Category.ToString());
+					const FText Tags = FText::FromString(Param.Tags);
+					TSharedPtr<FAssetSchemaAction_Dataflow_CreateNode_DataflowEdNode> NewNodeAction(
+						new FAssetSchemaAction_Dataflow_CreateNode_DataflowEdNode(InNodeTypeName, Category, NodeName, ToolTip, Tags));
+					return NewNodeAction;
+				}
+			}
+		}
+	}
+	return TSharedPtr<FAssetSchemaAction_Dataflow_CreateNode_DataflowEdNode>(nullptr);
+}
+
+static FName GetNodeUniqueName(UDataflow* Dataflow, FString NodeBaseName)
+{
+	FString Left, Right;
+	int32 NameIndex = 1;
+
+	// Check if NodeBaseName already ends with "_dd"
+	if (NodeBaseName.Split(TEXT("_"), &Left, &Right, ESearchCase::IgnoreCase, ESearchDir::FromEnd))
+	{
+		if (Right.IsNumeric())
+		{
+			NameIndex = FCString::Atoi(*Right);
+
+			NodeBaseName = Left;
+		}
+	}
+
+	FName NodeUniqueName{ NodeBaseName };
+	while (Dataflow->GetDataflow()->FindBaseNode(FName(NodeUniqueName)) != nullptr)
+	{
+		NodeUniqueName = FName(NodeBaseName + FString::Printf(TEXT("_%02d"), NameIndex));
+		NameIndex++;
+	}
+
+	return NodeUniqueName;
+}
+
+static UDataflowEdNode* CreateNode(UDataflow* Dataflow, UEdGraphPin* FromPin, const FVector2D Location, bool bSelectNewNode, const FName NodeUniqueName, const FName NodeTypeName, TSharedPtr<FDataflowNode> DataflowNodeToDuplicate, bool bCopySettings = false)
+{
+	if (UE::Dataflow::FNodeFactory* Factory = UE::Dataflow::FNodeFactory::GetInstance())
+	{
+		if (TSharedPtr<FDataflowNode> DataflowNode =
+			Factory->NewNodeFromRegisteredType(
+				*Dataflow->GetDataflow(),
+				{ FGuid::NewGuid(), NodeTypeName, NodeUniqueName, Dataflow }))
+		{
+			if (UDataflowEdNode* EdNode = NewObject<UDataflowEdNode>(Dataflow, UDataflowEdNode::StaticClass(), NodeUniqueName))
+			{
+				Dataflow->Modify();
+				if (FromPin != nullptr)
+				{
+					FromPin->Modify();
+				}
+
+				Dataflow->AddNode(EdNode, true, bSelectNewNode);
+
+				// Copy properties from DataflowNodeToDuplicate to DataflowNode
+				if (bCopySettings)
+				{
+					SDataflowEdNode::CopyDataflowNodeSettings(DataflowNodeToDuplicate, DataflowNode);
+				}
+
+				EdNode->CreateNewGuid();
+				EdNode->PostPlacedNewNode();
+
+				EdNode->SetDataflowGraph(Dataflow->GetDataflow());
+				EdNode->SetDataflowNodeGuid(DataflowNode->GetGuid());
+				EdNode->AllocateDefaultPins();
+
+				EdNode->AutowireNewNode(FromPin);
+
+				EdNode->NodePosX = Location.X;
+				EdNode->NodePosY = Location.Y;
+
+				EdNode->SetFlags(RF_Transactional);
+
+				return EdNode;
+			}
+		}
+	}
+
+	return nullptr;
+}
+
+//
+//  Created the EdGraph node and bind the guids to the Dataflow's node. 
+//
+UEdGraphNode* FAssetSchemaAction_Dataflow_CreateNode_DataflowEdNode::PerformAction(class UEdGraph* ParentGraph, UEdGraphPin* FromPin, const FVector2D Location, bool bSelectNewNode)
+{
+	if (UDataflow* Dataflow = Cast<UDataflow>(ParentGraph))
+	{
+		// by default use the type name and check if it is unique in the context of the graph
+		// if not, then generate a unique name 
+		const FString NodeBaseName = GetMenuDescription().ToString();
+		FName NodeUniqueName{ NodeBaseName };
+		int32 NameIndex = 0;
+		while (Dataflow->GetDataflow()->FindBaseNode(FName(NodeUniqueName)) != nullptr)
+		{
+			NodeUniqueName = FName(NodeBaseName + FString::Printf(TEXT("_%d"), NameIndex));
+			NameIndex++;
+		}
+
+		return CreateNode(Dataflow, FromPin, Location, bSelectNewNode, NodeUniqueName, NodeTypeName, nullptr);
+	}
+
+	return nullptr;
+}
+
+/////////////////////////////////////////////////////////////////////////////////////////////////////////////////////////
+
+TSharedPtr<FAssetSchemaAction_Dataflow_DuplicateNode_DataflowEdNode> FAssetSchemaAction_Dataflow_DuplicateNode_DataflowEdNode::CreateAction(UEdGraph* ParentGraph, const FName& InNodeTypeName)
+{
+	if (UE::Dataflow::FNodeFactory* Factory = UE::Dataflow::FNodeFactory::GetInstance())
+	{
+		const UE::Dataflow::FFactoryParameters& Param = Factory->GetParameters(InNodeTypeName);
+		if (Param.IsValid())
+		{
+			const FText ToolTip = FText::FromString(Param.ToolTip.IsEmpty() ? FString("Add a Dataflow node.") : Param.ToolTip);
+			const FText NodeName = FText::FromString(Param.DisplayName.ToString());
+			const FText Category = FText::FromString(Param.Category.ToString().IsEmpty() ? FString("Dataflow") : Param.Category.ToString());
+			const FText Tags = FText::FromString(Param.Tags);
+			TSharedPtr<FAssetSchemaAction_Dataflow_DuplicateNode_DataflowEdNode> NewNodeAction(
+				new FAssetSchemaAction_Dataflow_DuplicateNode_DataflowEdNode(InNodeTypeName, Category, NodeName, ToolTip, Tags));
+			return NewNodeAction;
+		}
+	}
+	return TSharedPtr<FAssetSchemaAction_Dataflow_DuplicateNode_DataflowEdNode>(nullptr);
+}
+
+//
+//   
+//
+UEdGraphNode* FAssetSchemaAction_Dataflow_DuplicateNode_DataflowEdNode::PerformAction(class UEdGraph* ParentGraph, UEdGraphPin* FromPin, const FVector2D Location, bool bSelectNewNode)
+{
+	if (UDataflow* Dataflow = Cast<UDataflow>(ParentGraph))
+	{
+		FString NodeToDuplicateName = DataflowNodeToDuplicate->GetName().ToString();
+
+		// Check if that is unique, if not then make it unique with an index postfix
+		FName NodeUniqueName = GetNodeUniqueName(Dataflow, NodeToDuplicateName);
+
+		return CreateNode(Dataflow, FromPin, Location, bSelectNewNode, NodeUniqueName, NodeTypeName, DataflowNodeToDuplicate, /*bCopySettings=*/true);
+	}
+
+	return nullptr;
+}
+
+/////////////////////////////////////////////////////////////////////////////////////////////////////////////////////////
+
+TSharedPtr<FAssetSchemaAction_Dataflow_PasteNode_DataflowEdNode> FAssetSchemaAction_Dataflow_PasteNode_DataflowEdNode::CreateAction(UEdGraph* ParentGraph, const FName& InNodeTypeName)
+{
+	if (UE::Dataflow::FNodeFactory* Factory = UE::Dataflow::FNodeFactory::GetInstance())
+	{
+		const UE::Dataflow::FFactoryParameters& Param = Factory->GetParameters(InNodeTypeName);
+		if (Param.IsValid())
+		{
+			const FText ToolTip = FText::FromString(Param.ToolTip.IsEmpty() ? FString("Add a Dataflow node.") : Param.ToolTip);
+			const FText NodeName = FText::FromString(Param.DisplayName.ToString());
+			const FText Category = FText::FromString(Param.Category.ToString().IsEmpty() ? FString("Dataflow") : Param.Category.ToString());
+			const FText Tags = FText::FromString(Param.Tags);
+			TSharedPtr<FAssetSchemaAction_Dataflow_PasteNode_DataflowEdNode> NewNodeAction(
+				new FAssetSchemaAction_Dataflow_PasteNode_DataflowEdNode(InNodeTypeName, Category, NodeName, ToolTip, Tags));
+			return NewNodeAction;
+		}
+	}
+	return TSharedPtr<FAssetSchemaAction_Dataflow_PasteNode_DataflowEdNode>(nullptr);
+}
+
+static UDataflowEdNode* CreateNodeFromPaste(UDataflow* Dataflow, UEdGraphPin* FromPin, const FVector2D Location, bool bSelectNewNode, const FName NodeUniqueName, const FName NodeTypeName, FString NodeProperties)
+{
+	if (UE::Dataflow::FNodeFactory* Factory = UE::Dataflow::FNodeFactory::GetInstance())
+	{
+		if (TSharedPtr<FDataflowNode> DataflowNode =
+			Factory->NewNodeFromRegisteredType(
+				*Dataflow->GetDataflow(),
+				{ FGuid::NewGuid(), NodeTypeName, NodeUniqueName, Dataflow }))
+		{
+			if (UDataflowEdNode* EdNode = NewObject<UDataflowEdNode>(Dataflow, UDataflowEdNode::StaticClass(), NodeUniqueName))
+			{
+				Dataflow->Modify();
+
+				Dataflow->AddNode(EdNode, true, bSelectNewNode);
+
+				// Copy properties to DataflowNode
+				if (!NodeProperties.IsEmpty())
+				{
+					DataflowNode->TypedScriptStruct()->ImportText(*NodeProperties, DataflowNode.Get(), nullptr, EPropertyPortFlags::PPF_None, nullptr, DataflowNode->TypedScriptStruct()->GetName(), true);
+				}
+				// Do any post-import fixup.
+				FArchive Ar;
+				Ar.SetIsLoading(true);
+				DataflowNode->PostSerialize(Ar);
+
+				EdNode->CreateNewGuid();
+				EdNode->PostPlacedNewNode();
+
+				EdNode->SetDataflowGraph(Dataflow->GetDataflow());
+				EdNode->SetDataflowNodeGuid(DataflowNode->GetGuid());
+				EdNode->AllocateDefaultPins();
+
+				EdNode->NodePosX = Location.X;
+				EdNode->NodePosY = Location.Y;
+
+				EdNode->SetFlags(RF_Transactional);
+
+				return EdNode;
+			}
+		}
+	}
+
+	return nullptr;
+}
+
+//
+//   
+//
+UEdGraphNode* FAssetSchemaAction_Dataflow_PasteNode_DataflowEdNode::PerformAction(class UEdGraph* ParentGraph, UEdGraphPin* FromPin, const FVector2D Location, bool bSelectNewNode)
+{
+	if (UDataflow* Dataflow = Cast<UDataflow>(ParentGraph))
+	{
+		FString NodeToDuplicateName = NodeName.ToString();
+
+		// Check if that is unique, if not then make it unique with an index postfix
+		FName NodeUniqueName = GetNodeUniqueName(Dataflow, NodeToDuplicateName);
+
+		return CreateNodeFromPaste(Dataflow, FromPin, Location, bSelectNewNode, NodeUniqueName, NodeTypeName, NodeProperties);
+	}
+
+	return nullptr;
+}
+
+/////////////////////////////////////////////////////////////////////////////////////////////////////////////////////////
+
 FEdGraphSchemaAction_DataflowVariable::FEdGraphSchemaAction_DataflowVariable()
 	: FEdGraphSchemaAction()
 {}
diff --git a/Engine/Plugins/Experimental/Dataflow/Source/DataflowEditor/Private/Dataflow/DataflowSNode.cpp b/Engine/Plugins/Experimental/Dataflow/Source/DataflowEditor/Private/Dataflow/DataflowSNode.cpp
--- a/Engine/Plugins/Experimental/Dataflow/Source/DataflowEditor/Private/Dataflow/DataflowSNode.cpp
+++ b/Engine/Plugins/Experimental/Dataflow/Source/DataflowEditor/Private/Dataflow/DataflowSNode.cpp
@@ -23,8 +23,6 @@
 #include "Dataflow/DataflowSchema.h"
 #include "Dataflow/DataflowSPin.h"
 
-#include UE_INLINE_GENERATED_CPP_BY_NAME(DataflowSNode)
-
 #define LOCTEXT_NAMESPACE "SDataflowEdNode"
 
 const FSlateBrush* SDataflowEdNode::GetPinButtonImage() const
@@ -384,68 +382,6 @@ EVisibility SDataflowEdNode::IsAddPinButtonVisible() const
 	return EVisibility::Collapsed;
 }
 
-//
-// Add a menu option to create a graph node.
-//
-TSharedPtr<FAssetSchemaAction_Dataflow_CreateNode_DataflowEdNode> FAssetSchemaAction_Dataflow_CreateNode_DataflowEdNode::CreateAction(const UEdGraph* ParentGraph, const FName & InNodeTypeName, const FName& InOverrideNodeName)
-{
-	if (const UDataflow* Dataflow = Cast<UDataflow>(ParentGraph))
-	{
-		if (UE::Dataflow::FNodeFactory* Factory = UE::Dataflow::FNodeFactory::GetInstance())
-		{
-			const UE::Dataflow::FFactoryParameters& Param = Factory->GetParameters(InNodeTypeName);
-			if (Param.IsValid())
-			{
-				const bool bIsSimulationNode = Param.Tags.Contains(UDataflow::SimulationTag);
-				const bool bIsSimulationGraph = (Dataflow->Type == EDataflowType::Simulation);
-				
-				if((bIsSimulationGraph && bIsSimulationNode) || (!bIsSimulationGraph && !bIsSimulationNode))
-				{
-					const FText ToolTip = FText::FromString(Param.ToolTip.IsEmpty() ? FString("Add a Dataflow node.") : Param.ToolTip);
-					FText NodeName = FText::FromString(Param.DisplayName.ToString());
-					if (!InOverrideNodeName.IsNone())
-					{
-						NodeName = FText::FromName(InOverrideNodeName);
-					}
-				
-					const FText Category = FText::FromString(Param.Category.ToString().IsEmpty() ? FString("Dataflow") : Param.Category.ToString());
-					const FText Tags = FText::FromString(Param.Tags);
-					TSharedPtr<FAssetSchemaAction_Dataflow_CreateNode_DataflowEdNode> NewNodeAction(
-						new FAssetSchemaAction_Dataflow_CreateNode_DataflowEdNode(InNodeTypeName, Category, NodeName, ToolTip, Tags));
-					return NewNodeAction;
-				}
-			}
-		}
-	}
-	return TSharedPtr<FAssetSchemaAction_Dataflow_CreateNode_DataflowEdNode>(nullptr);
-}
-
-static FName GetNodeUniqueName(UDataflow* Dataflow, FString NodeBaseName)
-{
-	FString Left, Right;
-	int32 NameIndex = 1;
-
-	// Check if NodeBaseName already ends with "_dd"
-	if (NodeBaseName.Split(TEXT("_"), &Left, &Right, ESearchCase::IgnoreCase, ESearchDir::FromEnd))
-	{
-		if (Right.IsNumeric())
-		{
-			NameIndex = FCString::Atoi(*Right);
-
-			NodeBaseName = Left;
-		}
-	}
-
-	FName NodeUniqueName{ NodeBaseName };
-	while (Dataflow->GetDataflow()->FindBaseNode(FName(NodeUniqueName)) != nullptr)
-	{
-		NodeUniqueName = FName(NodeBaseName + FString::Printf(TEXT("_%02d"), NameIndex));
-		NameIndex++;
-	}
-
-	return NodeUniqueName;
-}
-
 void SDataflowEdNode::CopyDataflowNodeSettings(TSharedPtr<FDataflowNode> SourceDataflowNode, TSharedPtr<FDataflowNode> TargetDataflowNode)
 {
 	using namespace UE::Transaction;
@@ -458,206 +394,5 @@ void SDataflowEdNode::CopyDataflowNodeSettings(TSharedPtr<FDataflowNode> SourceD
 	TargetDataflowNode->SerializeInternal(ArReader);
 }
 
-static UDataflowEdNode* CreateNode(UDataflow* Dataflow, UEdGraphPin* FromPin, const FVector2D Location, bool bSelectNewNode, const FName NodeUniqueName, const FName NodeTypeName, TSharedPtr<FDataflowNode> DataflowNodeToDuplicate, bool bCopySettings = false)
-{
-	if (UE::Dataflow::FNodeFactory* Factory = UE::Dataflow::FNodeFactory::GetInstance())
-	{
-		if (TSharedPtr<FDataflowNode> DataflowNode =
-			Factory->NewNodeFromRegisteredType(
-				*Dataflow->GetDataflow(),
-				{ FGuid::NewGuid(), NodeTypeName, NodeUniqueName, Dataflow }))
-		{
-			if (UDataflowEdNode* EdNode = NewObject<UDataflowEdNode>(Dataflow, UDataflowEdNode::StaticClass(), NodeUniqueName))
-			{
-				Dataflow->Modify();
-				if (FromPin != nullptr)
-				{
-					FromPin->Modify();
-				}
-
-				Dataflow->AddNode(EdNode, true, bSelectNewNode);
-
-				// Copy properties from DataflowNodeToDuplicate to DataflowNode
-				if (bCopySettings)
-				{
-					SDataflowEdNode::CopyDataflowNodeSettings(DataflowNodeToDuplicate, DataflowNode);
-				}
-				
-				EdNode->CreateNewGuid();
-				EdNode->PostPlacedNewNode();
-
-				EdNode->SetDataflowGraph(Dataflow->GetDataflow());
-				EdNode->SetDataflowNodeGuid(DataflowNode->GetGuid());
-				EdNode->AllocateDefaultPins();
-
-				EdNode->AutowireNewNode(FromPin);
-
-				EdNode->NodePosX = Location.X;
-				EdNode->NodePosY = Location.Y;
-
-				EdNode->SetFlags(RF_Transactional);
-
-				return EdNode;
-			}
-		}
-	}
-
-	return nullptr;
-}
-
-//
-//  Created the EdGraph node and bind the guids to the Dataflow's node. 
-//
-UEdGraphNode* FAssetSchemaAction_Dataflow_CreateNode_DataflowEdNode::PerformAction(class UEdGraph* ParentGraph, UEdGraphPin* FromPin, const FVector2D Location, bool bSelectNewNode)
-{
-	if (UDataflow* Dataflow = Cast<UDataflow>(ParentGraph))
-	{
-		// by default use the type name and check if it is unique in the context of the graph
-		// if not, then generate a unique name 
-		const FString NodeBaseName = GetMenuDescription().ToString();
-		FName NodeUniqueName{ NodeBaseName };
-		int32 NameIndex= 0;
-		while (Dataflow->GetDataflow()->FindBaseNode(FName(NodeUniqueName)) != nullptr)
-		{ 
-			NodeUniqueName = FName(NodeBaseName + FString::Printf(TEXT("_%d"), NameIndex));
-			NameIndex++;
-		}
-
-		return CreateNode(Dataflow, FromPin, Location, bSelectNewNode, NodeUniqueName, NodeTypeName, nullptr);
-	}
-
-	return nullptr;
-}
-
-//void FAssetSchemaAction_Dataflow_CreateNode_DataflowEdNode::AddReferencedObjects(FReferenceCollector& Collector)
-//{
-//	FEdGraphSchemaAction::AddReferencedObjects(Collector);
-//	Collector.AddReferencedObject(NodeTemplate);
-//}
-
-//
-// 
-//
-TSharedPtr<FAssetSchemaAction_Dataflow_DuplicateNode_DataflowEdNode> FAssetSchemaAction_Dataflow_DuplicateNode_DataflowEdNode::CreateAction(UEdGraph* ParentGraph, const FName& InNodeTypeName)
-{
-	if (UE::Dataflow::FNodeFactory* Factory = UE::Dataflow::FNodeFactory::GetInstance())
-	{
-		const UE::Dataflow::FFactoryParameters& Param = Factory->GetParameters(InNodeTypeName);
-		if (Param.IsValid())
-		{
-			const FText ToolTip = FText::FromString(Param.ToolTip.IsEmpty() ? FString("Add a Dataflow node.") : Param.ToolTip);
-			const FText NodeName = FText::FromString(Param.DisplayName.ToString());
-			const FText Category = FText::FromString(Param.Category.ToString().IsEmpty() ? FString("Dataflow") : Param.Category.ToString());
-			const FText Tags = FText::FromString(Param.Tags);
-			TSharedPtr<FAssetSchemaAction_Dataflow_DuplicateNode_DataflowEdNode> NewNodeAction(
-				new FAssetSchemaAction_Dataflow_DuplicateNode_DataflowEdNode(InNodeTypeName, Category, NodeName, ToolTip, Tags));
-			return NewNodeAction;
-		}
-	}
-	return TSharedPtr<FAssetSchemaAction_Dataflow_DuplicateNode_DataflowEdNode>(nullptr);
-}
-
-//
-//   
-//
-UEdGraphNode* FAssetSchemaAction_Dataflow_DuplicateNode_DataflowEdNode::PerformAction(class UEdGraph* ParentGraph, UEdGraphPin* FromPin, const FVector2D Location, bool bSelectNewNode)
-{
-	if (UDataflow* Dataflow = Cast<UDataflow>(ParentGraph))
-	{
-		FString NodeToDuplicateName = DataflowNodeToDuplicate->GetName().ToString();
-
-		// Check if that is unique, if not then make it unique with an index postfix
-		FName NodeUniqueName = GetNodeUniqueName(Dataflow, NodeToDuplicateName);
-
-		return CreateNode(Dataflow, FromPin, Location, bSelectNewNode, NodeUniqueName, NodeTypeName, DataflowNodeToDuplicate, /*bCopySettings=*/true);
-	}
-
-	return nullptr;
-}
-
-//
-// 
-//
-TSharedPtr<FAssetSchemaAction_Dataflow_PasteNode_DataflowEdNode> FAssetSchemaAction_Dataflow_PasteNode_DataflowEdNode::CreateAction(UEdGraph* ParentGraph, const FName& InNodeTypeName)
-{
-	if (UE::Dataflow::FNodeFactory* Factory = UE::Dataflow::FNodeFactory::GetInstance())
-	{
-		const UE::Dataflow::FFactoryParameters& Param = Factory->GetParameters(InNodeTypeName);
-		if (Param.IsValid())
-		{
-			const FText ToolTip = FText::FromString(Param.ToolTip.IsEmpty() ? FString("Add a Dataflow node.") : Param.ToolTip);
-			const FText NodeName = FText::FromString(Param.DisplayName.ToString());
-			const FText Category = FText::FromString(Param.Category.ToString().IsEmpty() ? FString("Dataflow") : Param.Category.ToString());
-			const FText Tags = FText::FromString(Param.Tags);
-			TSharedPtr<FAssetSchemaAction_Dataflow_PasteNode_DataflowEdNode> NewNodeAction(
-				new FAssetSchemaAction_Dataflow_PasteNode_DataflowEdNode(InNodeTypeName, Category, NodeName, ToolTip, Tags));
-			return NewNodeAction;
-		}
-	}
-	return TSharedPtr<FAssetSchemaAction_Dataflow_PasteNode_DataflowEdNode>(nullptr);
-}
-
-static UDataflowEdNode* CreateNodeFromPaste(UDataflow* Dataflow, UEdGraphPin* FromPin, const FVector2D Location, bool bSelectNewNode, const FName NodeUniqueName, const FName NodeTypeName, FString NodeProperties)
-{
-	if (UE::Dataflow::FNodeFactory* Factory = UE::Dataflow::FNodeFactory::GetInstance())
-	{
-		if (TSharedPtr<FDataflowNode> DataflowNode =
-			Factory->NewNodeFromRegisteredType(
-				*Dataflow->GetDataflow(),
-				{ FGuid::NewGuid(), NodeTypeName, NodeUniqueName, Dataflow }))
-		{
-			if (UDataflowEdNode* EdNode = NewObject<UDataflowEdNode>(Dataflow, UDataflowEdNode::StaticClass(), NodeUniqueName))
-			{
-				Dataflow->Modify();
-
-				Dataflow->AddNode(EdNode, true, bSelectNewNode);
-
-				// Copy properties to DataflowNode
-				if (!NodeProperties.IsEmpty())
-				{
-					DataflowNode->TypedScriptStruct()->ImportText(*NodeProperties, DataflowNode.Get(), nullptr, EPropertyPortFlags::PPF_None, nullptr, DataflowNode->TypedScriptStruct()->GetName(), true);
-				}
-				// Do any post-import fixup.
-				FArchive Ar;
-				Ar.SetIsLoading(true);
-				DataflowNode->PostSerialize(Ar);
-
-				EdNode->CreateNewGuid();
-				EdNode->PostPlacedNewNode();
-
-				EdNode->SetDataflowGraph(Dataflow->GetDataflow());
-				EdNode->SetDataflowNodeGuid(DataflowNode->GetGuid());
-				EdNode->AllocateDefaultPins();
-
-				EdNode->NodePosX = Location.X;
-				EdNode->NodePosY = Location.Y;
-
-				EdNode->SetFlags(RF_Transactional);
-
-				return EdNode;
-			}
-		}
-	}
-
-	return nullptr;
-}
-
-//
-//   
-//
-UEdGraphNode* FAssetSchemaAction_Dataflow_PasteNode_DataflowEdNode::PerformAction(class UEdGraph* ParentGraph, UEdGraphPin* FromPin, const FVector2D Location, bool bSelectNewNode)
-{
-	if (UDataflow* Dataflow = Cast<UDataflow>(ParentGraph))
-	{
-		FString NodeToDuplicateName = NodeName.ToString();
-
-		// Check if that is unique, if not then make it unique with an index postfix
-		FName NodeUniqueName = GetNodeUniqueName(Dataflow, NodeToDuplicateName);
-
-		return CreateNodeFromPaste(Dataflow, FromPin, Location, bSelectNewNode, NodeUniqueName, NodeTypeName, NodeProperties);
-	}
-
-	return nullptr;
-}
 #undef LOCTEXT_NAMESPACE
 
diff --git a/Engine/Plugins/Experimental/Dataflow/Source/DataflowEditor/Public/Dataflow/DataflowGraphSchemaAction.h b/Engine/Plugins/Experimental/Dataflow/Source/DataflowEditor/Public/Dataflow/DataflowGraphSchemaAction.h
--- a/Engine/Plugins/Experimental/Dataflow/Source/DataflowEditor/Public/Dataflow/DataflowGraphSchemaAction.h
+++ b/Engine/Plugins/Experimental/Dataflow/Source/DataflowEditor/Public/Dataflow/DataflowGraphSchemaAction.h
@@ -5,9 +5,12 @@
 #include "GraphEditorDragDropAction.h"
 #include "DataflowGraphSchemaAction.generated.h"
 
+struct FPropertyBagPropertyDesc;
 class IDataflowInstanceInterface;
 class UDataflow;
-struct FPropertyBagPropertyDesc;
+class UEdGraphNode;
+class UEdGraphPin;
+struct FDataflowNode;
 
 namespace UE::Dataflow
 {
@@ -23,6 +26,77 @@ namespace UE::Dataflow
 
 ////////////////////////////////////////////////////////////////////////////////////////////////////////////////////////////////
 
+//
+// Action to add a node to the graph
+//
+USTRUCT()
+struct DATAFLOWEDITOR_API FAssetSchemaAction_Dataflow_CreateNode_DataflowEdNode : public FEdGraphSchemaAction
+{
+	GENERATED_USTRUCT_BODY();
+
+public:
+	FAssetSchemaAction_Dataflow_CreateNode_DataflowEdNode() : FEdGraphSchemaAction() {}
+
+	FAssetSchemaAction_Dataflow_CreateNode_DataflowEdNode(const FName& InType, FText InNodeCategory, FText InMenuDesc, FText InToolTip, FText InKeywords)
+		: FEdGraphSchemaAction(InNodeCategory, InMenuDesc, InToolTip, 0, InKeywords), NodeTypeName(InType) {}
+
+	static TSharedPtr<FAssetSchemaAction_Dataflow_CreateNode_DataflowEdNode> CreateAction(const UEdGraph* ParentGraph, const FName& NodeTypeName, const FName& InOverrideNodeName = NAME_None);
+
+	virtual UEdGraphNode* PerformAction(class UEdGraph* ParentGraph, UEdGraphPin* FromPin, const FVector2D Location, bool bSelectNewNode = true) override;
+	//virtual void AddReferencedObjects(FReferenceCollector& Collector) override;
+
+	FName NodeTypeName;
+};
+
+////////////////////////////////////////////////////////////////////////////////////////////////////////////////////////////////
+
+//
+// Action to duplicate a set of nodes in the graph
+//
+USTRUCT()
+struct DATAFLOWEDITOR_API FAssetSchemaAction_Dataflow_DuplicateNode_DataflowEdNode : public FEdGraphSchemaAction
+{
+	GENERATED_USTRUCT_BODY();
+
+public:
+	FAssetSchemaAction_Dataflow_DuplicateNode_DataflowEdNode() : FEdGraphSchemaAction() {}
+
+	FAssetSchemaAction_Dataflow_DuplicateNode_DataflowEdNode(const FName& InType, FText InNodeCategory, FText InMenuDesc, FText InToolTip, FText InKeywords)
+		: FEdGraphSchemaAction(InNodeCategory, InMenuDesc, InToolTip, 0, InKeywords), NodeTypeName(InType) {}
+
+	static TSharedPtr<FAssetSchemaAction_Dataflow_DuplicateNode_DataflowEdNode> CreateAction(UEdGraph* ParentGraph, const FName& NodeTypeName);
+
+	virtual UEdGraphNode* PerformAction(class UEdGraph* ParentGraph, UEdGraphPin* FromPin, const FVector2D Location, bool bSelectNewNode = true) override;
+
+	FName NodeTypeName;
+	TSharedPtr<FDataflowNode> DataflowNodeToDuplicate;
+};
+
+////////////////////////////////////////////////////////////////////////////////////////////////////////////////////////////////
+
+//
+// Action to paste a set of nodes in the graph
+//
+USTRUCT()
+struct DATAFLOWEDITOR_API FAssetSchemaAction_Dataflow_PasteNode_DataflowEdNode : public FEdGraphSchemaAction
+{
+	GENERATED_USTRUCT_BODY();
+
+public:
+	FAssetSchemaAction_Dataflow_PasteNode_DataflowEdNode() : FEdGraphSchemaAction() {}
+
+	FAssetSchemaAction_Dataflow_PasteNode_DataflowEdNode(const FName& InType, FText InNodeCategory, FText InMenuDesc, FText InToolTip, FText InKeywords)
+		: FEdGraphSchemaAction(InNodeCategory, InMenuDesc, InToolTip, 0, InKeywords), NodeTypeName(InType) {}
+
+	static TSharedPtr<FAssetSchemaAction_Dataflow_PasteNode_DataflowEdNode> CreateAction(UEdGraph* ParentGraph, const FName& NodeTypeName);
+
+	virtual UEdGraphNode* PerformAction(class UEdGraph* ParentGraph, UEdGraphPin* FromPin, const FVector2D Location, bool bSelectNewNode = true) override;
+
+	FName NodeTypeName;
+	FName NodeName;
+	FString NodeProperties;
+};
+
 /**
 * Action that refers to a dataflow variable
 * It is used to display the variables in the SDataflowMembersWidget
diff --git a/Engine/Plugins/Experimental/Dataflow/Source/DataflowEditor/Public/Dataflow/DataflowSNode.h b/Engine/Plugins/Experimental/Dataflow/Source/DataflowEditor/Public/Dataflow/DataflowSNode.h
--- a/Engine/Plugins/Experimental/Dataflow/Source/DataflowEditor/Public/Dataflow/DataflowSNode.h
+++ b/Engine/Plugins/Experimental/Dataflow/Source/DataflowEditor/Public/Dataflow/DataflowSNode.h
@@ -8,8 +8,6 @@
 #include "Dataflow/DataflowCore.h"
 #include "Dataflow/DataflowSEditorInterface.h"
 
-#include "DataflowSNode.generated.h"
-
 class UDataflowEdNode;
 class SCheckBox;
 class SImage;
@@ -78,69 +76,3 @@ private:
 
 };
 
-//
-// Action to add a node to the graph
-//
-USTRUCT()
-struct DATAFLOWEDITOR_API FAssetSchemaAction_Dataflow_CreateNode_DataflowEdNode : public FEdGraphSchemaAction
-{
-	GENERATED_USTRUCT_BODY();
-
-public:
-	FAssetSchemaAction_Dataflow_CreateNode_DataflowEdNode() : FEdGraphSchemaAction() {}
-
-	FAssetSchemaAction_Dataflow_CreateNode_DataflowEdNode(const FName& InType, FText InNodeCategory, FText InMenuDesc, FText InToolTip, FText InKeywords)
-		: FEdGraphSchemaAction(InNodeCategory, InMenuDesc, InToolTip, 0, InKeywords), NodeTypeName(InType) {}
-
-	static TSharedPtr<FAssetSchemaAction_Dataflow_CreateNode_DataflowEdNode> CreateAction(const UEdGraph* ParentGraph, const FName & NodeTypeName, const FName& InOverrideNodeName = NAME_None);
-
-	virtual UEdGraphNode* PerformAction(class UEdGraph* ParentGraph, UEdGraphPin* FromPin, const FVector2D Location, bool bSelectNewNode = true) override;
-	//virtual void AddReferencedObjects(FReferenceCollector& Collector) override;
-
-	FName NodeTypeName;
-};
-
-//
-// Action to duplicate a set of nodes in the graph
-//
-USTRUCT()
-struct DATAFLOWEDITOR_API FAssetSchemaAction_Dataflow_DuplicateNode_DataflowEdNode : public FEdGraphSchemaAction
-{
-	GENERATED_USTRUCT_BODY();
-
-public:
-	FAssetSchemaAction_Dataflow_DuplicateNode_DataflowEdNode() : FEdGraphSchemaAction() {}
-
-	FAssetSchemaAction_Dataflow_DuplicateNode_DataflowEdNode(const FName& InType, FText InNodeCategory, FText InMenuDesc, FText InToolTip, FText InKeywords)
-		: FEdGraphSchemaAction(InNodeCategory, InMenuDesc, InToolTip, 0, InKeywords), NodeTypeName(InType) {}
-
-	static TSharedPtr<FAssetSchemaAction_Dataflow_DuplicateNode_DataflowEdNode> CreateAction(UEdGraph* ParentGraph, const FName& NodeTypeName);
-
-	virtual UEdGraphNode* PerformAction(class UEdGraph* ParentGraph, UEdGraphPin* FromPin, const FVector2D Location, bool bSelectNewNode = true) override;
-
-	FName NodeTypeName;
-	TSharedPtr<FDataflowNode> DataflowNodeToDuplicate;
-};
-
-//
-// Action to paste a set of nodes in the graph
-//
-USTRUCT()
-struct DATAFLOWEDITOR_API FAssetSchemaAction_Dataflow_PasteNode_DataflowEdNode : public FEdGraphSchemaAction
-{
-	GENERATED_USTRUCT_BODY();
-
-public:
-	FAssetSchemaAction_Dataflow_PasteNode_DataflowEdNode() : FEdGraphSchemaAction() {}
-
-	FAssetSchemaAction_Dataflow_PasteNode_DataflowEdNode(const FName& InType, FText InNodeCategory, FText InMenuDesc, FText InToolTip, FText InKeywords)
-		: FEdGraphSchemaAction(InNodeCategory, InMenuDesc, InToolTip, 0, InKeywords), NodeTypeName(InType) {}
-
-	static TSharedPtr<FAssetSchemaAction_Dataflow_PasteNode_DataflowEdNode> CreateAction(UEdGraph* ParentGraph, const FName& NodeTypeName);
-
-	virtual UEdGraphNode* PerformAction(class UEdGraph* ParentGraph, UEdGraphPin* FromPin, const FVector2D Location, bool bSelectNewNode = true) override;
-
-	FName NodeTypeName;
-	FName NodeName;
-	FString NodeProperties;
-};
