Dataflow : Change the dataflow asset utils to be a class with static methods to prepare for future changes around SubGraphs

diff --git a/Engine/Plugins/Experimental/Dataflow/Source/DataflowEditor/Private/Dataflow/DataflowAssetEditUtils.cpp b/Engine/Plugins/Experimental/Dataflow/Source/DataflowEditor/Private/Dataflow/DataflowAssetEditUtils.cpp
--- a/Engine/Plugins/Experimental/Dataflow/Source/DataflowEditor/Private/Dataflow/DataflowAssetEditUtils.cpp
+++ b/Engine/Plugins/Experimental/Dataflow/Source/DataflowEditor/Private/Dataflow/DataflowAssetEditUtils.cpp
@@ -4,18 +4,19 @@
 
 #include "HAL/PlatformApplicationMisc.h"
 #include "Dataflow/DataflowObject.h"
+#include "Dataflow/DataflowSchema.h"
 #include "ScopedTransaction.h"
 #include "PropertyBagDetails.h"
 
 
 #define LOCTEXT_NAMESPACE "DataflowAssetEditUtils"
 
-namespace UE::Dataflow::EditAssetUtils
+namespace UE::Dataflow
 {
-	namespace Private
+	namespace EditAssetUtils::Private
 	{
 		static const TCHAR* DataflowVariableClipboardPrefix = TEXT("DataflowVariable_");
-
+		
 		enum class EChangeResult
 		{
 			None,
@@ -88,7 +89,7 @@ namespace UE::Dataflow::EditAssetUtils
 			return false;
 		}
 
-		/** Generate a Dataflow asset variable unique name from a BaseName*/
+		/** Generate a Dataflow asset variable unique name from a BaseName */
 		static FName GenerateUniqueVariableName(const UDataflow& DataflowAsset, FName BaseName)
 		{
 			int32 Counter = 1;
@@ -111,49 +112,55 @@ namespace UE::Dataflow::EditAssetUtils
 	// VARIABLES FUNCTIONS
 	// 
 	//////////////////////////////////////////////////////////////////////////////////////////////////////////
-	FName AddNewVariable(UDataflow* DataflowAsset, FName BaseName)
+	FName FEditAssetUtils::AddNewVariable(UDataflow* DataflowAsset, FName BaseName)
 	{
+		using namespace EditAssetUtils::Private;
+
 		const FText TransactionName{ LOCTEXT("AddNewDataflowVariable", "Add New Dataflow Variable") };
 
 		FName UniqueVariableName;
 
 		auto AddNewVariableInternal =
-			[&BaseName, &UniqueVariableName](UDataflow& DataflowAsset) -> Private::EChangeResult
+			[&BaseName, &UniqueVariableName](UDataflow& DataflowAsset) -> EChangeResult
 			{
-				UniqueVariableName = Private::GenerateUniqueVariableName(DataflowAsset, BaseName);
+				UniqueVariableName = GenerateUniqueVariableName(DataflowAsset, BaseName);
 				DataflowAsset.Variables.AddProperty(UniqueVariableName, EPropertyBagPropertyType::Int32);
 				FDataflowAssetDelegates::OnVariablesChanged.Broadcast(&DataflowAsset, UniqueVariableName);
-				return Private::EChangeResult::Changed;
+				return EChangeResult::Changed;
 			};
 
-		Private::ChangeDataflowAssetWithTransaction(DataflowAsset, TransactionName, AddNewVariableInternal);
+		ChangeDataflowAssetWithTransaction(DataflowAsset, TransactionName, AddNewVariableInternal);
 
 		return UniqueVariableName;
 	}
 
 	//////////////////////////////////////////////////////////////////////////////////////////////////////////
-	void DeleteVariable(UDataflow* DataflowAsset, FName VariableName)
+	void FEditAssetUtils::DeleteVariable(UDataflow* DataflowAsset, FName VariableName)
 	{
+		using namespace EditAssetUtils::Private;
+
 		const FText TransactionName{ LOCTEXT("DeleteDataflowVariable", "Delete Dataflow Variable") };
 
 		auto DeleteVariableInternal =
-			[&VariableName](UDataflow& DataflowAsset) -> Private::EChangeResult
+			[&VariableName](UDataflow& DataflowAsset) -> EChangeResult
 			{
 				DataflowAsset.Variables.RemovePropertyByName(VariableName);
 				FDataflowAssetDelegates::OnVariablesChanged.Broadcast(&DataflowAsset, VariableName);
-				return Private::EChangeResult::Changed;
+				return EChangeResult::Changed;
 			};
 
-		Private::ChangeDataflowAssetWithTransaction(DataflowAsset, TransactionName, DeleteVariableInternal);
+		ChangeDataflowAssetWithTransaction(DataflowAsset, TransactionName, DeleteVariableInternal);
 	}
 
 	//////////////////////////////////////////////////////////////////////////////////////////////////////////
-	void RenameVariable(UDataflow* DataflowAsset, FName OldVariableName, FName NewVariableName)
+	void FEditAssetUtils::RenameVariable(UDataflow* DataflowAsset, FName OldVariableName, FName NewVariableName)
 	{
+		using namespace EditAssetUtils::Private;
+
 		const FText TransactionName{ LOCTEXT("RenameDataflowVariable", "Rename Dataflow Variable") };
 
 		auto SetVariableNameInternal =
-			[OldVariableName, NewVariableName](UDataflow& DataflowAsset) -> Private::EChangeResult
+			[OldVariableName, NewVariableName](UDataflow& DataflowAsset) -> EChangeResult
 			{
 				auto ChangePropertyNameLambda =
 					[NewVariableName](FPropertyBagPropertyDesc& PropertyDesc)
@@ -161,26 +168,28 @@ namespace UE::Dataflow::EditAssetUtils
 						PropertyDesc.Name = NewVariableName;
 					};
 
-				const bool bModified = Private::ModifyVariable(DataflowAsset, OldVariableName, ChangePropertyNameLambda);
-				return (bModified) ? Private::EChangeResult::Changed : Private::EChangeResult::None;
+				const bool bModified = ModifyVariable(DataflowAsset, OldVariableName, ChangePropertyNameLambda);
+				return (bModified) ? EChangeResult::Changed : EChangeResult::None;
 			};
 
-		Private::ChangeDataflowAssetWithTransaction(DataflowAsset, TransactionName, SetVariableNameInternal);
+		ChangeDataflowAssetWithTransaction(DataflowAsset, TransactionName, SetVariableNameInternal);
 	}
 
 	//////////////////////////////////////////////////////////////////////////////////////////////////////////
-	FName DuplicateVariable(UDataflow* DataflowAsset, FName VariableName)
+	FName FEditAssetUtils::DuplicateVariable(UDataflow* DataflowAsset, FName VariableName)
 	{
+		using namespace EditAssetUtils::Private;
+
 		const FText TransactionName = FText::Format(LOCTEXT("DuplicateDataflowVariable", "Duplicate Dataflow Variable: {0}"), FText::FromName(VariableName));
 		
 		FName NewVariableName;
 
 		auto DuplicateVariableInternal =
-			[VariableName, &NewVariableName](UDataflow& DataflowAsset) -> Private::EChangeResult
+			[VariableName, &NewVariableName](UDataflow& DataflowAsset) -> EChangeResult
 			{
 				if (const FPropertyBagPropertyDesc* PropertyDescPtr = DataflowAsset.Variables.FindPropertyDescByName(VariableName))
 				{
-					NewVariableName = Private::GenerateUniqueVariableName(DataflowAsset, VariableName);
+					NewVariableName = GenerateUniqueVariableName(DataflowAsset, VariableName);
 
 					// make sure the name is unique and the GUID is invalidated to avoid copying it as is 
 					FPropertyBagPropertyDesc NewDesc(*PropertyDescPtr);
@@ -188,24 +197,26 @@ namespace UE::Dataflow::EditAssetUtils
 					NewDesc.ID.Invalidate();
 					DataflowAsset.Variables.AddProperties(MakeConstArrayView(&NewDesc, 1));
 					FDataflowAssetDelegates::OnVariablesChanged.Broadcast(&DataflowAsset, NewVariableName);
-					return Private::EChangeResult::Changed;
+					return EChangeResult::Changed;
 				}
-				return Private::EChangeResult::None;
+				return EChangeResult::None;
 			};
 
-		Private::ChangeDataflowAssetWithTransaction(DataflowAsset, TransactionName, DuplicateVariableInternal);
+		ChangeDataflowAssetWithTransaction(DataflowAsset, TransactionName, DuplicateVariableInternal);
 
 		return NewVariableName;
 	}
 
 
 	//////////////////////////////////////////////////////////////////////////////////////////////////////////
-	void SetVariableType(UDataflow* DataflowAsset, FName VariableName, const FEdGraphPinType& PinType)
+	void FEditAssetUtils::SetVariableType(UDataflow* DataflowAsset, FName VariableName, const FEdGraphPinType& PinType)
 	{
+		using namespace EditAssetUtils::Private;
+
 		const FText TransactionName{ LOCTEXT("ChangeDataflowVariableType", "Change Dataflow Variable Type") };
 
 		auto SetVariableTypeInternal =
-			[VariableName, &PinType](UDataflow& DataflowAsset) -> Private::EChangeResult
+			[VariableName, &PinType](UDataflow& DataflowAsset) -> EChangeResult
 			{
 				auto ChangePropertyTypeLambda =
 					[VariableName, &PinType](FPropertyBagPropertyDesc& PropertyDesc)
@@ -213,20 +224,22 @@ namespace UE::Dataflow::EditAssetUtils
 						UE::StructUtils::SetPropertyDescFromPin(PropertyDesc, PinType);
 					};
 
-				const bool bModified = Private::ModifyVariable(DataflowAsset, VariableName, ChangePropertyTypeLambda);
-				return (bModified) ? Private::EChangeResult::Changed : Private::EChangeResult::None;
+				const bool bModified = ModifyVariable(DataflowAsset, VariableName, ChangePropertyTypeLambda);
+				return (bModified) ? EChangeResult::Changed : EChangeResult::None;
 			};
 
-		Private::ChangeDataflowAssetWithTransaction(DataflowAsset, TransactionName, SetVariableTypeInternal);
+		ChangeDataflowAssetWithTransaction(DataflowAsset, TransactionName, SetVariableTypeInternal);
 	}
 
 	//////////////////////////////////////////////////////////////////////////////////////////////////////////
-	void SetVariableValue(UDataflow* DataflowAsset, FName VariableName, const FInstancedPropertyBag& SourceBag)
+	void FEditAssetUtils::SetVariableValue(UDataflow* DataflowAsset, FName VariableName, const FInstancedPropertyBag& SourceBag)
 	{
+		using namespace EditAssetUtils::Private;
+
 		const FText TransactionName{ LOCTEXT("ChangeDataflowVariableValue", "Change Dataflow Variable Value") };
 
 		auto SetVariableValueInternal =
-			[VariableName, &SourceBag](UDataflow& DataflowAsset) -> Private::EChangeResult
+			[VariableName, &SourceBag](UDataflow& DataflowAsset) -> EChangeResult
 			{
 				if (const FPropertyBagPropertyDesc* SourceDesc = SourceBag.FindPropertyDescByName(VariableName))
 				{
@@ -236,20 +249,22 @@ namespace UE::Dataflow::EditAssetUtils
 						if (Result == EPropertyBagResult::Success)
 						{
 							FDataflowAssetDelegates::OnVariablesChanged.Broadcast(&DataflowAsset, VariableName);
-							return Private::EChangeResult::Changed;
+							return EChangeResult::Changed;
 						}
 					}
 				}
-				return Private::EChangeResult::None;
+				return EChangeResult::None;
 			};
 
 		// do we need a transction in that case ? 
-		Private::ChangeDataflowAssetWithTransaction(DataflowAsset, TransactionName, SetVariableValueInternal);
+		ChangeDataflowAssetWithTransaction(DataflowAsset, TransactionName, SetVariableValueInternal);
 	}
 
 	//////////////////////////////////////////////////////////////////////////////////////////////////////////
-	void CopyVariableToClipboard(UDataflow* DataflowAsset, FName VariableName)
+	void FEditAssetUtils::CopyVariableToClipboard(UDataflow* DataflowAsset, FName VariableName)
 	{
+		using namespace EditAssetUtils::Private;
+
 		if (DataflowAsset)
 		{
 			// no transaction needed in that case as we write to an external system 
@@ -261,7 +276,7 @@ namespace UE::Dataflow::EditAssetUtils
 
 				if (!ClipboardPayload.IsEmpty())
 				{
-					ClipboardPayload = Private::DataflowVariableClipboardPrefix + ClipboardPayload;
+					ClipboardPayload = DataflowVariableClipboardPrefix + ClipboardPayload;
 					FPlatformApplicationMisc::ClipboardCopy(ClipboardPayload.GetCharArray().GetData());
 				}
 			}
@@ -269,18 +284,20 @@ namespace UE::Dataflow::EditAssetUtils
 	}
 
 	//////////////////////////////////////////////////////////////////////////////////////////////////////////
-	FName PasteVariableFromClipboard(UDataflow* DataflowAsset)
+	FName FEditAssetUtils::PasteVariableFromClipboard(UDataflow* DataflowAsset)
 	{
+		using namespace EditAssetUtils::Private;
+
 		FString ClipboardPayload;
 		FPlatformApplicationMisc::ClipboardPaste(ClipboardPayload);
 
-		if (!ensure(ClipboardPayload.StartsWith(Private::DataflowVariableClipboardPrefix, ESearchCase::CaseSensitive)))
+		if (!ensure(ClipboardPayload.StartsWith(DataflowVariableClipboardPrefix, ESearchCase::CaseSensitive)))
 		{
 			return NAME_None;
 		}
 
 		FStringOutputDevice Errors;
-		const TCHAR* ImportPayload = ClipboardPayload.GetCharArray().GetData() + FCString::Strlen(Private::DataflowVariableClipboardPrefix);
+		const TCHAR* ImportPayload = ClipboardPayload.GetCharArray().GetData() + FCString::Strlen(DataflowVariableClipboardPrefix);
 
 		FPropertyBagPropertyDesc PropertyDesc;
 		FPropertyBagPropertyDesc::StaticStruct()->ImportText(ImportPayload, &PropertyDesc, nullptr, PPF_None, &Errors, FPropertyBagPropertyDesc::StaticStruct()->GetName());
@@ -290,20 +307,20 @@ namespace UE::Dataflow::EditAssetUtils
 			if (DataflowAsset)
 			{
 				// make sure the name is unique and the GUID is invalidated to avoid copying it as is 
-				PropertyDesc.Name = Private::GenerateUniqueVariableName(*DataflowAsset, PropertyDesc.Name);
+				PropertyDesc.Name = GenerateUniqueVariableName(*DataflowAsset, PropertyDesc.Name);
 				PropertyDesc.ID.Invalidate();
 
 				const FText TransactionName = FText::Format(LOCTEXT("PasteDataflowVariable", "Paste Dataflow Variable: {0}"), FText::FromName(PropertyDesc.Name));
 
 				auto PasteVariableInternal =
-					[&PropertyDesc](UDataflow& DataflowAsset) -> Private::EChangeResult
+					[&PropertyDesc](UDataflow& DataflowAsset) -> EChangeResult
 					{
 						DataflowAsset.Variables.AddProperties(MakeConstArrayView(&PropertyDesc, 1));
 						FDataflowAssetDelegates::OnVariablesChanged.Broadcast(&DataflowAsset, PropertyDesc.Name);
-						return Private::EChangeResult::Changed;
+						return EChangeResult::Changed;
 					};
 
-				Private::ChangeDataflowAssetWithTransaction(DataflowAsset, TransactionName, PasteVariableInternal);
+				ChangeDataflowAssetWithTransaction(DataflowAsset, TransactionName, PasteVariableInternal);
 				return PropertyDesc.Name;
 			}
 		}
diff --git a/Engine/Plugins/Experimental/Dataflow/Source/DataflowEditor/Private/Dataflow/DataflowAssetEditUtils.h b/Engine/Plugins/Experimental/Dataflow/Source/DataflowEditor/Private/Dataflow/DataflowAssetEditUtils.h
--- a/Engine/Plugins/Experimental/Dataflow/Source/DataflowEditor/Private/Dataflow/DataflowAssetEditUtils.h
+++ b/Engine/Plugins/Experimental/Dataflow/Source/DataflowEditor/Private/Dataflow/DataflowAssetEditUtils.h
@@ -16,50 +16,59 @@ struct FPropertyBagPropertyDesc;
 * Set to functions to modify a dataflow asset programmatically 
 * all changes are wrapped inside a scoped transaction as well and propertly marked the asset as modified if needed
 */
-namespace UE::Dataflow::EditAssetUtils
+namespace UE::Dataflow
 {
 	static const FName DefaultNewVariableBaseName = "NewVariable";
 
-	/** 
-	* Add a new Dataflow asset variable
-	* If BaseName already exists, use it to generate a unique name by adding an increment to it in the form of BaseName_## 
-	*/
-	FName AddNewVariable(UDataflow* DataflowAsset, FName BaseName = DefaultNewVariableBaseName);
+	struct FEditAssetUtils
+	{
+		/////////////////////////////////////////////////////////////////////////////////////////////////////////////////////////
+		//
+		// Variables API
+		//
+		/////////////////////////////////////////////////////////////////////////////////////////////////////////////////////////
 
-	/** 
-	* Remove a Dataflow asset variable
-	*/
-	void DeleteVariable(UDataflow* DataflowAsset, FName VariableName);
+		/**
+		* Add a new Dataflow asset variable
+		* If BaseName already exists, use it to generate a unique name by adding an increment to it in the form of BaseName_##
+		*/
+		static FName AddNewVariable(UDataflow* DataflowAsset, FName BaseName = DefaultNewVariableBaseName);
 
-	/**
-	* Rename a Dataflow asset variable
-	*/
-	void RenameVariable(UDataflow* DataflowAsset, FName OldVariableName, FName NewVariableName);
+		/**
+		* Remove a Dataflow asset variable
+		*/
+		static void DeleteVariable(UDataflow* DataflowAsset, FName VariableName);
 
-	/**
-	* Duplicate a Dataflow asset variable
-	* the new variable name will be generate uniquely based on the name of the original variable
-	*/
-	FName DuplicateVariable(UDataflow* DataflowAsset, FName VariableName);
+		/**
+		* Rename a Dataflow asset variable
+		*/
+		static void RenameVariable(UDataflow* DataflowAsset, FName OldVariableName, FName NewVariableName);
 
-	/**
-	* Set the type of a Dataflow asset variable
-	*/
-	void SetVariableType(UDataflow* DataflowAsset, FName VariableName, const FEdGraphPinType& PinType);
+		/**
+		* Duplicate a Dataflow asset variable
+		* the new variable name will be generate uniquely based on the name of the original variable
+		*/
+		static FName DuplicateVariable(UDataflow* DataflowAsset, FName VariableName);
 
-	/*
-	* Set value of a Dataflow asset variable from an property in a propertybag 
-	* Name and type must match for the operation to be successful
-	*/
-	void SetVariableValue(UDataflow* DataflowAsset, FName VariableName, const FInstancedPropertyBag& SourceBag);
+		/**
+		* Set the type of a Dataflow asset variable
+		*/
+		static void SetVariableType(UDataflow* DataflowAsset, FName VariableName, const FEdGraphPinType& PinType);
 
-	/**
-	* Copy a variable to clipboard
-	*/
-	void CopyVariableToClipboard(UDataflow* DataflowAsset, FName VariableName);
+		/*
+		* Set value of a Dataflow asset variable from an property in a propertybag
+		* Name and type must match for the operation to be successful
+		*/
+		static void SetVariableValue(UDataflow* DataflowAsset, FName VariableName, const FInstancedPropertyBag& SourceBag);
 
-	/**
-	* Paste a variable from clipboard and return its name
-	*/
-	FName PasteVariableFromClipboard(UDataflow* DataflowAsset);
+		/**
+		* Copy a variable to clipboard
+		*/
+		static void CopyVariableToClipboard(UDataflow* DataflowAsset, FName VariableName);
+
+		/**
+		* Paste a variable from clipboard and return its name
+		*/
+		static FName PasteVariableFromClipboard(UDataflow* DataflowAsset);
+	};
 }
diff --git a/Engine/Plugins/Experimental/Dataflow/Source/DataflowEditor/Private/Dataflow/DataflowGraphEditor.cpp b/Engine/Plugins/Experimental/Dataflow/Source/DataflowEditor/Private/Dataflow/DataflowGraphEditor.cpp
--- a/Engine/Plugins/Experimental/Dataflow/Source/DataflowEditor/Private/Dataflow/DataflowGraphEditor.cpp
+++ b/Engine/Plugins/Experimental/Dataflow/Source/DataflowEditor/Private/Dataflow/DataflowGraphEditor.cpp
@@ -334,7 +334,7 @@ bool SDataflowGraphEditor::CanRenameNode() const
 
 void SDataflowGraphEditor::AddNewVariable() const
 {
-	UE::Dataflow::EditAssetUtils::AddNewVariable(DataflowAsset.Get());
+	UE::Dataflow::FEditAssetUtils::AddNewVariable(DataflowAsset.Get());
 }
 
 void SDataflowGraphEditor::OnSelectedNodesChanged(const TSet<UObject*>& NewSelection)
diff --git a/Engine/Plugins/Experimental/Dataflow/Source/DataflowEditor/Private/Dataflow/DataflowGraphSchemaAction.cpp b/Engine/Plugins/Experimental/Dataflow/Source/DataflowEditor/Private/Dataflow/DataflowGraphSchemaAction.cpp
--- a/Engine/Plugins/Experimental/Dataflow/Source/DataflowEditor/Private/Dataflow/DataflowGraphSchemaAction.cpp
+++ b/Engine/Plugins/Experimental/Dataflow/Source/DataflowEditor/Private/Dataflow/DataflowGraphSchemaAction.cpp
@@ -350,36 +350,36 @@ void FEdGraphSchemaAction_DataflowVariable::RenameVariable(FText NewName)
 		{
 			NewVariableName = FName(FString::Format(TEXT("{0}|{1}"), { VariableCategory, NewName.ToString() }));
 		}
-		UE::Dataflow::EditAssetUtils::RenameVariable(DataflowAsset, FullVariableName, NewVariableName);
+		UE::Dataflow::FEditAssetUtils::RenameVariable(DataflowAsset, FullVariableName, NewVariableName);
 	}
 }
 
 void FEdGraphSchemaAction_DataflowVariable::CopyVariableToClipboard()
 {
-	UE::Dataflow::EditAssetUtils::CopyVariableToClipboard(GetDataflowAsset(), GetFullVariableName());
+	UE::Dataflow::FEditAssetUtils::CopyVariableToClipboard(GetDataflowAsset(), GetFullVariableName());
 }
 
 void FEdGraphSchemaAction_DataflowVariable::PasteVariableFromClipboard()
 {
-	UE::Dataflow::EditAssetUtils::PasteVariableFromClipboard(GetDataflowAsset());
+	UE::Dataflow::FEditAssetUtils::PasteVariableFromClipboard(GetDataflowAsset());
 }
 
 void FEdGraphSchemaAction_DataflowVariable::SetVariableType(const FEdGraphPinType& PinType)
 {
 	if (UDataflow* DataflowAsset = GetDataflowAsset())
 	{
-		UE::Dataflow::EditAssetUtils::SetVariableType(DataflowAsset, FullVariableName, PinType);
+		UE::Dataflow::FEditAssetUtils::SetVariableType(DataflowAsset, FullVariableName, PinType);
 	}
 }
 
 void FEdGraphSchemaAction_DataflowVariable::DeleteVariable()
 {
-	UE::Dataflow::EditAssetUtils::DeleteVariable(GetDataflowAsset(), GetFullVariableName());
+	UE::Dataflow::FEditAssetUtils::DeleteVariable(GetDataflowAsset(), GetFullVariableName());
 }
 
 void FEdGraphSchemaAction_DataflowVariable::DuplicateVariable()
 {
-	UE::Dataflow::EditAssetUtils::DuplicateVariable(GetDataflowAsset(), GetFullVariableName());
+	UE::Dataflow::FEditAssetUtils::DuplicateVariable(GetDataflowAsset(), GetFullVariableName());
 }
 
 /** Execute this action, given the graph and schema, and possibly a pin that we were dragged from. Returns a node that was created by this action (if any). */
