Dataflow:
- Fixed namespace issue

diff --git a/Engine/Plugins/Experimental/Dataflow/Source/DataflowEditor/Private/Dataflow/DataflowSceneProfileIndexStorage.cpp b/Engine/Plugins/Experimental/Dataflow/Source/DataflowEditor/Private/Dataflow/DataflowSceneProfileIndexStorage.cpp
--- a/Engine/Plugins/Experimental/Dataflow/Source/DataflowEditor/Private/Dataflow/DataflowSceneProfileIndexStorage.cpp
+++ b/Engine/Plugins/Experimental/Dataflow/Source/DataflowEditor/Private/Dataflow/DataflowSceneProfileIndexStorage.cpp
@@ -7,7 +7,7 @@
 #include "AssetViewerSettings.h"
 
 
-namespace Private
+namespace UE::Dataflow::Private
 {
 	int32 GetProfileIndex(const FString& SearchName)
 	{
@@ -52,7 +52,7 @@ int32 FDataflowConstructionSceneProfileIndexStorage::RetrieveProfileIndex()
 {
 	if (const UDataflowEditorOptions* const Options = UDataflowEditorOptions::StaticClass()->GetDefaultObject<UDataflowEditorOptions>())
 	{
-		const int32 FoundIndex = Private::GetProfileIndex(Options->ConstructionProfileName);
+		const int32 FoundIndex = UE::Dataflow::Private::GetProfileIndex(Options->ConstructionProfileName);
 		if (FoundIndex != ConstructionScene->GetCurrentProfileIndex())
 		{
 			ConstructionScene->SetCurrentProfileIndex(FoundIndex);
@@ -90,7 +90,7 @@ int32 FDataflowSimulationSceneProfileIndexStorage::RetrieveProfileIndex()
 {
 	if (const UDataflowEditorOptions* const Options = UDataflowEditorOptions::StaticClass()->GetDefaultObject<UDataflowEditorOptions>())
 	{
-		const int32 FoundIndex = Private::GetProfileIndex(Options->SimulationProfileName);
+		const int32 FoundIndex = UE::Dataflow::Private::GetProfileIndex(Options->SimulationProfileName);
 		if (FoundIndex != SimulationScene->GetCurrentProfileIndex())
 		{
 			SimulationScene->SetCurrentProfileIndex(FoundIndex);
