Add new methods to dataflow edit asset utils class ( in preparation to simplify the actions code )
- AddNewNode / DuplicateNode

diff --git a/Engine/Plugins/Experimental/Dataflow/Source/DataflowEditor/Private/Dataflow/DataflowAssetEditUtils.cpp b/Engine/Plugins/Experimental/Dataflow/Source/DataflowEditor/Private/Dataflow/DataflowAssetEditUtils.cpp
--- a/Engine/Plugins/Experimental/Dataflow/Source/DataflowEditor/Private/Dataflow/DataflowAssetEditUtils.cpp
+++ b/Engine/Plugins/Experimental/Dataflow/Source/DataflowEditor/Private/Dataflow/DataflowAssetEditUtils.cpp
@@ -5,6 +5,9 @@
 #include "HAL/PlatformApplicationMisc.h"
 #include "Dataflow/DataflowObject.h"
 #include "Dataflow/DataflowSchema.h"
+#include "Dataflow/DataflowNode.h"
+#include "Dataflow/DataflowEdNode.h"
+#include "Dataflow/DataflowSNode.h"
 #include "ScopedTransaction.h"
 #include "PropertyBagDetails.h"
 
@@ -29,7 +32,7 @@ namespace UE::Dataflow
 		* The InFunction return paramater determines if the asset will be modified or not 
 		* if modification happens, a PostEditChangeProperty notification is sent
 		*/
-		static void ChangeDataflowAssetWithTransaction(UDataflow* DataflowAsset, const FText &TransactionName, TFunctionRef<EChangeResult(UDataflow&)> InFunction)
+		static void ChangeDataflowAssetWithTransaction(UDataflow* DataflowAsset, const FText &TransactionName, TFunctionRef<EChangeResult(UDataflow&)> InFunction, FName ChangedPropertyName)
 		{
 			if (DataflowAsset)
 			{
@@ -45,13 +48,16 @@ namespace UE::Dataflow
 				case EChangeResult::Changed:
 				{
 					DataflowAsset->Modify();
-					FPropertyChangedEvent PropertyChangedEvent(nullptr);
-					if (UClass* DataflowClass = DataflowAsset->GetClass())
+					if (!ChangedPropertyName.IsNone())
 					{
-						FProperty* MemberProperty = DataflowClass->FindPropertyByName(GET_MEMBER_NAME_CHECKED(UDataflow, Variables));
-						PropertyChangedEvent.SetActiveMemberProperty(MemberProperty);
+						FPropertyChangedEvent PropertyChangedEvent(nullptr);
+						if (UClass* DataflowClass = DataflowAsset->GetClass())
+						{
+							FProperty* MemberProperty = DataflowClass->FindPropertyByName(ChangedPropertyName);
+							PropertyChangedEvent.SetActiveMemberProperty(MemberProperty);
+						}
+						DataflowAsset->PostEditChangeProperty(PropertyChangedEvent);
 					}
-					DataflowAsset->PostEditChangeProperty(PropertyChangedEvent);
 					break;
 				}
 				case EChangeResult::None:
@@ -61,6 +67,81 @@ namespace UE::Dataflow
 			}
 		}
 
+		/** Generate a Dataflow node unique name from a BaseName */
+		static FName GenerateUniqueNodeName(UDataflow* Dataflow, FString NodeBaseName)
+		{
+			FString Left, Right;
+			int32 NameIndex = 1;
+
+			// Check if NodeBaseName already ends with "_dd"
+			if (NodeBaseName.Split(TEXT("_"), &Left, &Right, ESearchCase::IgnoreCase, ESearchDir::FromEnd))
+			{
+				if (Right.IsNumeric())
+				{
+					NameIndex = FCString::Atoi(*Right);
+					NodeBaseName = Left;
+				}
+			}
+
+			FName NodeUniqueName{ NodeBaseName };
+			while (Dataflow->GetDataflow()->FindBaseNode(FName(NodeUniqueName)) != nullptr)
+			{
+				NodeUniqueName = FName(NodeBaseName + FString::Printf(TEXT("_%d"), NameIndex));
+				NameIndex++;
+			}
+
+			return NodeUniqueName;
+		}
+
+		static TSharedPtr<FDataflowNode> AddDataflowNode(UDataflow* Dataflow, FName NodeName, FName NodeTypeName)
+		{
+			if (UE::Dataflow::FNodeFactory* Factory = UE::Dataflow::FNodeFactory::GetInstance())
+			{
+				FNewNodeParameters Parameters
+				{
+					.Guid = FGuid::NewGuid(),
+					.Type = NodeTypeName,
+					.Name = GenerateUniqueNodeName(Dataflow, NodeName.ToString()),
+					.OwningObject = Dataflow,
+				};
+
+				return Factory->NewNodeFromRegisteredType(*Dataflow->GetDataflow(), Parameters);
+			}
+			return nullptr;
+		}
+
+		static UDataflowEdNode* CreateDataflowEdNode(UEdGraph* EdGraph, TSharedPtr<FDataflowNode> DataflowNode, const FVector2D& Location, UEdGraphPin* FromPin)
+		{
+			if (UDataflow* Dataflow = UDataflow::GetDataflowAssetFromEdGraph(EdGraph))
+			{
+				if (UDataflowEdNode* EdNode = NewObject<UDataflowEdNode>(EdGraph, UDataflowEdNode::StaticClass(), DataflowNode->GetName()))
+				{
+					EdGraph->Modify();
+
+					EdNode->CreateNewGuid();
+					EdNode->PostPlacedNewNode();
+
+					EdNode->SetDataflowGraph(Dataflow->GetDataflow());
+					EdNode->SetDataflowNodeGuid(DataflowNode->GetGuid());
+					EdNode->AllocateDefaultPins();
+
+					if (FromPin)
+					{
+						FromPin->Modify();
+						EdNode->AutowireNewNode(FromPin);
+					}
+
+					EdNode->NodePosX = Location.X;
+					EdNode->NodePosY = Location.Y;
+
+					EdNode->SetFlags(RF_Transactional);
+
+					return EdNode;
+				}
+			}
+			return nullptr;
+		}
+
 		/** returns true if the variable was modified */
 		static bool ModifyVariable(UDataflow& DataflowAsset, FName Variable, TFunctionRef<void(FPropertyBagPropertyDesc& PropertyDesc)> InFunction)
 		{
@@ -109,7 +190,84 @@ namespace UE::Dataflow
 
 	//////////////////////////////////////////////////////////////////////////////////////////////////////////
 	// 
-	// VARIABLES FUNCTIONS
+	// NODE API
+	// 
+	//////////////////////////////////////////////////////////////////////////////////////////////////////////
+	UDataflowEdNode* FEditAssetUtils::AddNewNode(UEdGraph* EdGraph, const FVector2D& Location, const FName NodeName, const FName NodeTypeName, UEdGraphPin* FromPin)
+	{
+		using namespace EditAssetUtils::Private;
+
+		const FText TransactionName{ LOCTEXT("AddNewDataflowNode", "Add New Dataflow Node") };
+
+		UDataflowEdNode* EdNodeToReturn = nullptr;
+
+		if (UDataflow* Dataflow = UDataflow::GetDataflowAssetFromEdGraph(EdGraph))
+		{
+			auto AddNewNodeInternal =
+				[&Dataflow, &EdGraph, &Location, &NodeName, &NodeTypeName, &FromPin, &EdNodeToReturn]
+				(UDataflow& DataflowAsset) -> EChangeResult
+				{
+					if (TSharedPtr<FDataflowNode> DataflowNode = AddDataflowNode(Dataflow, NodeName, NodeTypeName))
+					{
+						if (UDataflowEdNode* EdNode = CreateDataflowEdNode(EdGraph, DataflowNode, Location, FromPin))
+						{
+							Dataflow->AddNode(EdNode, /*bUserAction*/true, /*bSelectNewNode*/false);
+							EdNodeToReturn = EdNode;
+							return EChangeResult::Changed;
+						}
+					}
+					return EChangeResult::None;
+				};
+
+			ChangeDataflowAssetWithTransaction(Dataflow, TransactionName, AddNewNodeInternal, NAME_None);
+		}
+		return EdNodeToReturn;
+	}
+
+	//////////////////////////////////////////////////////////////////////////////////////////////////////////
+	UDataflowEdNode* FEditAssetUtils::DuplicateNode(UEdGraph* EdGraph, TSharedPtr<FDataflowNode> NodeToDuplicate, const FVector2D& Location)
+	{
+		using namespace EditAssetUtils::Private;
+
+		const FText TransactionName{ LOCTEXT("DuplicateDataflowNode", "Duplicate a Dataflow Node") };
+
+		UDataflowEdNode* EdNodeToReturn = nullptr;
+
+		if (UDataflow* Dataflow = UDataflow::GetDataflowAssetFromEdGraph(EdGraph))
+		{
+			auto DuplicateNodeInternal =
+				[&Dataflow, &EdGraph, &Location, &NodeToDuplicate, &EdNodeToReturn]
+				(UDataflow& DataflowAsset) -> EChangeResult
+				{
+					if (NodeToDuplicate)
+					{
+						const FName NodeName = NodeToDuplicate->GetName();
+						const FName NodeTypeName = NodeToDuplicate->GetType();
+
+						if (TSharedPtr<FDataflowNode> DataflowNode = AddDataflowNode(Dataflow, NodeName, NodeTypeName))
+						{
+							SDataflowEdNode::CopyDataflowNodeSettings(NodeToDuplicate, DataflowNode);
+
+							if (UDataflowEdNode* EdNode = CreateDataflowEdNode(EdGraph, DataflowNode, Location, /*FromPin*/nullptr))
+							{
+								Dataflow->AddNode(EdNode, /*bUserAction*/true, /*bSelectNewNode*/false);
+
+								EdNodeToReturn = EdNode;
+								return EChangeResult::Changed;
+							}
+						}
+					}
+					return EChangeResult::None;
+				};
+
+			ChangeDataflowAssetWithTransaction(Dataflow, TransactionName, DuplicateNodeInternal, NAME_None);
+		}
+		return EdNodeToReturn;
+	}
+
+	//////////////////////////////////////////////////////////////////////////////////////////////////////////
+	// 
+	// VARIABLES API
 	// 
 	//////////////////////////////////////////////////////////////////////////////////////////////////////////
 	FName FEditAssetUtils::AddNewVariable(UDataflow* DataflowAsset, FName BaseName)
@@ -129,7 +287,7 @@ namespace UE::Dataflow
 				return EChangeResult::Changed;
 			};
 
-		ChangeDataflowAssetWithTransaction(DataflowAsset, TransactionName, AddNewVariableInternal);
+		ChangeDataflowAssetWithTransaction(DataflowAsset, TransactionName, AddNewVariableInternal, GET_MEMBER_NAME_CHECKED(UDataflow, Variables));
 
 		return UniqueVariableName;
 	}
@@ -149,7 +307,7 @@ namespace UE::Dataflow
 				return EChangeResult::Changed;
 			};
 
-		ChangeDataflowAssetWithTransaction(DataflowAsset, TransactionName, DeleteVariableInternal);
+		ChangeDataflowAssetWithTransaction(DataflowAsset, TransactionName, DeleteVariableInternal, GET_MEMBER_NAME_CHECKED(UDataflow, Variables));
 	}
 
 	//////////////////////////////////////////////////////////////////////////////////////////////////////////
@@ -172,7 +330,7 @@ namespace UE::Dataflow
 				return (bModified) ? EChangeResult::Changed : EChangeResult::None;
 			};
 
-		ChangeDataflowAssetWithTransaction(DataflowAsset, TransactionName, SetVariableNameInternal);
+		ChangeDataflowAssetWithTransaction(DataflowAsset, TransactionName, SetVariableNameInternal, GET_MEMBER_NAME_CHECKED(UDataflow, Variables));
 	}
 
 	//////////////////////////////////////////////////////////////////////////////////////////////////////////
@@ -202,7 +360,7 @@ namespace UE::Dataflow
 				return EChangeResult::None;
 			};
 
-		ChangeDataflowAssetWithTransaction(DataflowAsset, TransactionName, DuplicateVariableInternal);
+		ChangeDataflowAssetWithTransaction(DataflowAsset, TransactionName, DuplicateVariableInternal, GET_MEMBER_NAME_CHECKED(UDataflow, Variables));
 
 		return NewVariableName;
 	}
@@ -228,7 +386,7 @@ namespace UE::Dataflow
 				return (bModified) ? EChangeResult::Changed : EChangeResult::None;
 			};
 
-		ChangeDataflowAssetWithTransaction(DataflowAsset, TransactionName, SetVariableTypeInternal);
+		ChangeDataflowAssetWithTransaction(DataflowAsset, TransactionName, SetVariableTypeInternal, GET_MEMBER_NAME_CHECKED(UDataflow, Variables));
 	}
 
 	//////////////////////////////////////////////////////////////////////////////////////////////////////////
@@ -257,7 +415,7 @@ namespace UE::Dataflow
 			};
 
 		// do we need a transction in that case ? 
-		ChangeDataflowAssetWithTransaction(DataflowAsset, TransactionName, SetVariableValueInternal);
+		ChangeDataflowAssetWithTransaction(DataflowAsset, TransactionName, SetVariableValueInternal, GET_MEMBER_NAME_CHECKED(UDataflow, Variables));
 	}
 
 	//////////////////////////////////////////////////////////////////////////////////////////////////////////
@@ -320,7 +478,7 @@ namespace UE::Dataflow
 						return EChangeResult::Changed;
 					};
 
-				ChangeDataflowAssetWithTransaction(DataflowAsset, TransactionName, PasteVariableInternal);
+				ChangeDataflowAssetWithTransaction(DataflowAsset, TransactionName, PasteVariableInternal, GET_MEMBER_NAME_CHECKED(UDataflow, Variables));
 				return PropertyDesc.Name;
 			}
 		}
diff --git a/Engine/Plugins/Experimental/Dataflow/Source/DataflowEditor/Private/Dataflow/DataflowAssetEditUtils.h b/Engine/Plugins/Experimental/Dataflow/Source/DataflowEditor/Private/Dataflow/DataflowAssetEditUtils.h
--- a/Engine/Plugins/Experimental/Dataflow/Source/DataflowEditor/Private/Dataflow/DataflowAssetEditUtils.h
+++ b/Engine/Plugins/Experimental/Dataflow/Source/DataflowEditor/Private/Dataflow/DataflowAssetEditUtils.h
@@ -2,15 +2,20 @@
 
 #pragma once
 
+#include "Math/MathFwd.h"
 #include "UObject/NameTypes.h"
+#include "Templates/SharedPointer.h"
 
 ///////////////////////////////////////////////////////////
 
-class UDataflow;
 struct FEdGraphPinType;
-
+struct FDataflowNode;
 struct FInstancedPropertyBag;
 struct FPropertyBagPropertyDesc;
+class UDataflow;
+class UEdGraph;
+class UDataflowEdNode;
+class UEdGraphPin;
 
 /**
 * Set to functions to modify a dataflow asset programmatically 
@@ -22,6 +27,31 @@ namespace UE::Dataflow
 
 	struct FEditAssetUtils
 	{
+		/////////////////////////////////////////////////////////////////////////////////////////////////////////////////////////
+		//
+		// Nodes API
+		//
+		/////////////////////////////////////////////////////////////////////////////////////////////////////////////////////////
+		/**
+		* Add a new node to the graph
+		* @param EdGraph	Main Graph or Subgraph to add the node to 
+		* @param Location	Location on the graph where to place the node 
+		* @param NodeName	Name to give to the node ( if the name is not unique a unique name will be created out of the parameter one) 
+		* @param TypeName	Name of the type of the node to add
+		* @param FromPin	Pin to connect automatically to the node
+		* @return the newly created EdGraphNode
+		*/
+		static UDataflowEdNode* AddNewNode(UEdGraph* EdGraph, const FVector2D& Location, const FName NodeName, const FName NodeTypeName, UEdGraphPin* FromPin);
+
+		/**
+		* Duplicate a node
+		* @param EdGraph			Main Graph or Subgraph to duplicate the node in 
+		* @param NodeToDuplicate	Node to duplicate
+		* @param Location			Location on the graph where to place the node 
+		* @return the newly created EdGraphNode
+		*/
+		static UDataflowEdNode* DuplicateNode(UEdGraph* EdGraph, TSharedPtr<FDataflowNode> NodeToDuplicate, const FVector2D& Location);
+
 		/////////////////////////////////////////////////////////////////////////////////////////////////////////////////////////
 		//
 		// Variables API
