Dataflow : Add clipboard methods for nodes to dataflow edit asset utils
- code is from the original commands and commands will be updated later to use those methods

diff --git a/Engine/Plugins/Experimental/Dataflow/Source/DataflowEditor/Private/Dataflow/DataflowAssetEditUtils.cpp b/Engine/Plugins/Experimental/Dataflow/Source/DataflowEditor/Private/Dataflow/DataflowAssetEditUtils.cpp
--- a/Engine/Plugins/Experimental/Dataflow/Source/DataflowEditor/Private/Dataflow/DataflowAssetEditUtils.cpp
+++ b/Engine/Plugins/Experimental/Dataflow/Source/DataflowEditor/Private/Dataflow/DataflowAssetEditUtils.cpp
@@ -3,14 +3,15 @@
 #include "Dataflow/DataflowAssetEditUtils.h"
 
 #include "HAL/PlatformApplicationMisc.h"
+#include "Dataflow/AssetDefinition_DataflowAsset.h"
 #include "Dataflow/DataflowObject.h"
 #include "Dataflow/DataflowSchema.h"
 #include "Dataflow/DataflowNode.h"
 #include "Dataflow/DataflowEdNode.h"
 #include "Dataflow/DataflowSNode.h"
-#include "ScopedTransaction.h"
 #include "PropertyBagDetails.h"
-
+#include "ScopedTransaction.h"
+#include "Settings/EditorStyleSettings.h"
 
 #define LOCTEXT_NAMESPACE "DataflowAssetEditUtils"
 
@@ -118,6 +119,8 @@ namespace UE::Dataflow
 				{
 					EdGraph->Modify();
 
+					Dataflow->AddNode(EdNode, /*bUserAction*/true, /*bSelectNewNode*/false);
+					
 					EdNode->CreateNewGuid();
 					EdNode->PostPlacedNewNode();
 
@@ -142,6 +145,41 @@ namespace UE::Dataflow
 			return nullptr;
 		}
 
+		UEdGraphNode_Comment* CreateCommentEdNode(UEdGraph* Graph, const FVector2D& Location, const FString& Comment, const FVector2D Size, const FLinearColor& Color, int32 FontSize)
+		{
+			Graph->Modify();
+			
+			UEdGraphNode_Comment* CommentTemplate = NewObject<UEdGraphNode_Comment>(Graph, FName("Comment"));
+			{
+				CommentTemplate->bCommentBubbleVisible_InDetailsPanel = false;
+				CommentTemplate->bCommentBubbleVisible = false;
+				CommentTemplate->bCommentBubblePinned = false;
+
+				CommentTemplate->SetFlags(RF_Transactional);
+
+				// set outer to be the graph so it doesn't go away
+				CommentTemplate->Rename(NULL, Graph, REN_NonTransactional);
+				Graph->AddNode(CommentTemplate, true, /*bSelectNewNode*/false);
+
+				CommentTemplate->CreateNewGuid();
+				CommentTemplate->PostPlacedNewNode();
+				CommentTemplate->AllocateDefaultPins();
+
+				CommentTemplate->NodePosX = Location.X;
+				CommentTemplate->NodePosY = Location.Y;
+				CommentTemplate->NodeWidth = Size.X;
+				CommentTemplate->NodeHeight = Size.Y;
+				CommentTemplate->SnapToGrid(GetDefault<UEditorStyleSettings>()->GridSnapSize);
+				CommentTemplate->CommentColor = Color;
+				CommentTemplate->FontSize = FontSize;
+
+				CommentTemplate->NodeComment = Comment;
+
+				//ParentGraph->NotifyGraphChanged();
+			}
+			return CommentTemplate;
+		}
+
 		/** returns true if the variable was modified */
 		static bool ModifyVariable(UDataflow& DataflowAsset, FName Variable, TFunctionRef<void(FPropertyBagPropertyDesc& PropertyDesc)> InFunction)
 		{
@@ -186,6 +224,19 @@ namespace UE::Dataflow
 			}
 			return UniqueName;
 		}
+
+		static UEdGraphPin* FindPin(const UDataflowEdNode* Node, const EEdGraphPinDirection Direction, const FName Name)
+		{
+			for (UEdGraphPin* Pin : Node->GetAllPins())
+			{
+				if (Pin->PinName == Name && Pin->Direction == Direction)
+				{
+					return Pin;
+				}
+			}
+
+			return nullptr;
+		}
 	}
 
 	//////////////////////////////////////////////////////////////////////////////////////////////////////////
@@ -211,7 +262,6 @@ namespace UE::Dataflow
 					{
 						if (UDataflowEdNode* EdNode = CreateDataflowEdNode(EdGraph, DataflowNode, Location, FromPin))
 						{
-							Dataflow->AddNode(EdNode, /*bUserAction*/true, /*bSelectNewNode*/false);
 							EdNodeToReturn = EdNode;
 							return EChangeResult::Changed;
 						}
@@ -225,44 +275,348 @@ namespace UE::Dataflow
 	}
 
 	//////////////////////////////////////////////////////////////////////////////////////////////////////////
-	UDataflowEdNode* FEditAssetUtils::DuplicateNode(UEdGraph* EdGraph, TSharedPtr<FDataflowNode> NodeToDuplicate, const FVector2D& Location)
+	void FEditAssetUtils::DuplicateNodes(UEdGraph* EdGraph, const TArray<UEdGraphNode*>& EdNodesToDuplicate, const FVector2D& Location, TArray<UEdGraphNode*>& OutDuplicatedNodes)
 	{
 		using namespace EditAssetUtils::Private;
 
-		const FText TransactionName{ LOCTEXT("DuplicateDataflowNode", "Duplicate a Dataflow Node") };
+		UDataflow* DataflowAsset = UDataflow::GetDataflowAssetFromEdGraph(EdGraph);
+		if (!DataflowAsset)
+		{
+			// no graph to copy to
+			return;
+		}
 
-		UDataflowEdNode* EdNodeToReturn = nullptr;
+		if (EdNodesToDuplicate.Num() == 0)
+		{
+			return;
+		}
 
-		if (UDataflow* Dataflow = UDataflow::GetDataflowAssetFromEdGraph(EdGraph))
+		const FText TransactionName = FText::Format(LOCTEXT("DuplicateDataflowNode", "Duplicate {0} Dataflow Nodes"), FText::AsNumber(EdNodesToDuplicate.Num()));
+
+		// location of the first node as a reference for all the others 
+		const FVector2D RefLocation
 		{
-			auto DuplicateNodeInternal =
-				[&Dataflow, &EdGraph, &Location, &NodeToDuplicate, &EdNodeToReturn]
-				(UDataflow& DataflowAsset) -> EChangeResult
+			(double)EdNodesToDuplicate[0]->NodePosX,
+			(double)EdNodesToDuplicate[0]->NodePosY
+		};
+
+		auto DuplicateNodesInternal =
+			[&EdGraph, &Location, &RefLocation, &EdNodesToDuplicate, &OutDuplicatedNodes]
+			(UDataflow& DataflowAsset) -> EChangeResult
+			{
+				TMap<FGuid, UDataflowEdNode*> EdNodeMap;
+				TMap<FGuid, FGuid> NodeGuidMap;
+
+				// copy the nodes and comments first 
+				for (UEdGraphNode* EdNodeToDuplicate : EdNodesToDuplicate)
 				{
-					if (NodeToDuplicate)
+					const FVector2D OriginalNodeLocation
 					{
-						const FName NodeName = NodeToDuplicate->GetName();
-						const FName NodeTypeName = NodeToDuplicate->GetType();
+						(double)EdNodeToDuplicate->NodePosX,
+						(double)EdNodeToDuplicate->NodePosY
+					};
+					const FVector2D NodeLocation(Location + (OriginalNodeLocation - RefLocation));
 
-						if (TSharedPtr<FDataflowNode> DataflowNode = AddDataflowNode(Dataflow, NodeName, NodeTypeName))
+					if (UDataflowEdNode* DataflowEdNodeToDuplicate = Cast<UDataflowEdNode>(EdNodeToDuplicate))
+					{
+						if (const TSharedPtr<FDataflowNode> NodeToDuplicate = DataflowEdNodeToDuplicate->GetDataflowNode())
 						{
-							SDataflowEdNode::CopyDataflowNodeSettings(NodeToDuplicate, DataflowNode);
+							const FName NodeName = NodeToDuplicate->GetName();
+							const FName NodeTypeName = NodeToDuplicate->GetType();
 
-							if (UDataflowEdNode* EdNode = CreateDataflowEdNode(EdGraph, DataflowNode, Location, /*FromPin*/nullptr))
+							if (TSharedPtr<FDataflowNode> DataflowNode = AddDataflowNode(&DataflowAsset, NodeName, NodeTypeName))
 							{
-								Dataflow->AddNode(EdNode, /*bUserAction*/true, /*bSelectNewNode*/false);
+								SDataflowEdNode::CopyDataflowNodeSettings(NodeToDuplicate, DataflowNode);
+	
+								if (UDataflowEdNode* EdNode = CreateDataflowEdNode(EdGraph, DataflowNode, NodeLocation, /*FromPin*/nullptr))
+								{
+									EdNodeMap.Add(DataflowEdNodeToDuplicate->DataflowNodeGuid, EdNode);
+									NodeGuidMap.Add(DataflowEdNodeToDuplicate->DataflowNodeGuid, EdNode->DataflowNodeGuid);
+									OutDuplicatedNodes.Add(EdNode);
+								}
+							}
+						}
+					}
+					else if (const UEdGraphNode_Comment* CommentEdNodeToDuplicate = Cast<const UEdGraphNode_Comment>(EdNodeToDuplicate))
+					{
+						const FVector2D CommentSize
+						{
+							(double)CommentEdNodeToDuplicate->NodeWidth,
+							(double)CommentEdNodeToDuplicate->NodeHeight
+						};
+						if (UEdGraphNode_Comment* CommentEdNode = CreateCommentEdNode(EdGraph, NodeLocation, CommentEdNodeToDuplicate->NodeComment, CommentSize, CommentEdNodeToDuplicate->CommentColor, CommentEdNodeToDuplicate->FontSize))
+						{
+							OutDuplicatedNodes.Add(CommentEdNode);
+						}
+					}
+				}
 
-								EdNodeToReturn = EdNode;
-								return EChangeResult::Changed;
+				// Recreate connections between duplicated nodes
+				for (const UEdGraphNode* EdNodeToDuplicate : EdNodesToDuplicate)
+				{
+					if (const UDataflowEdNode* DataflowEdNodeToDuplicate = Cast<const UDataflowEdNode>(EdNodeToDuplicate))
+					{
+						if (const TSharedPtr<const FDataflowNode> DataflowNode = DataflowEdNodeToDuplicate->GetDataflowNode())
+						{
+							const FGuid DataflowNodeAGuid = DataflowNode->GetGuid();
+							for (const FDataflowOutput* Output : DataflowNode->GetOutputs())
+							{
+								for (const FDataflowInput* Connection : Output->Connections)
+								{
+									const FName OutputputName = Connection->GetConnection()->GetName();
+
+									// Check if the node on the end of the connection was duplicated
+									const FGuid DataflowNodeBGuid = Connection->GetOwningNode()->GetGuid();
+
+									if (NodeGuidMap.Contains(DataflowNodeBGuid))
+									{
+										const FName InputputName = Connection->GetName();
+										if (TSharedPtr<FDataflowNode> DuplicatedDataflowNodeA = DataflowAsset.GetDataflow()->FindBaseNode(NodeGuidMap[DataflowNodeAGuid]))
+										{
+											FDataflowOutput* OutputConnection = DuplicatedDataflowNodeA->FindOutput(OutputputName);
+											if (TSharedPtr<FDataflowNode> DuplicatedDataflowNodeB = DataflowAsset.GetDataflow()->FindBaseNode(NodeGuidMap[DataflowNodeBGuid]))
+											{
+												FDataflowInput* InputConnection = DuplicatedDataflowNodeB->FindInput(InputputName);
+
+												DataflowAsset.GetDataflow()->Connect(OutputConnection, InputConnection);
+
+												// Connect the UDataflowEdNode FPins as well
+												if (UEdGraphPin* OutputPin = FindPin(EdNodeMap[DataflowNodeAGuid], EEdGraphPinDirection::EGPD_Output, OutputputName))
+												{
+													if (UEdGraphPin* InputPin = FindPin(EdNodeMap[DataflowNodeBGuid], EEdGraphPinDirection::EGPD_Input, InputputName))
+													{
+														OutputPin->MakeLinkTo(InputPin);
+													}
+												}
+											}
+										}
+									}
+								}
 							}
 						}
 					}
-					return EChangeResult::None;
+				}
+				return EChangeResult::Changed;
+			};
+
+		ChangeDataflowAssetWithTransaction(DataflowAsset, TransactionName, DuplicateNodesInternal, NAME_None);
+	}
+
+	//////////////////////////////////////////////////////////////////////////////////////////////////////////
+	void FEditAssetUtils::CopyNodesToClipboard(const TArray<const UEdGraphNode*>& NodesToCopy, int32& OutNumCopiedNodes)
+	{
+		FDataflowCopyPasteContent CopyPasteContent;
+
+		TSet<FGuid> NodeGuids;
+		TArray<const FDataflowInput*> NodeInputsToSave;
+
+		// no need for transaction when copying to external system like the clipboard
+		for (const UEdGraphNode* EdNode : NodesToCopy)
+		{
+			if (const UDataflowEdNode* DataflowEdNode = Cast<const UDataflowEdNode>(EdNode))
+			{
+				if (const TSharedPtr<const FDataflowNode> DataflowNode = DataflowEdNode->GetDataflowNode())
+				{
+					NodeGuids.Add(DataflowNode->GetGuid());
+					NodeInputsToSave.Append(DataflowNode->GetInputs());
+
+					FString ContentString;
+					const TUniquePtr<FDataflowNode> DefaultElement;
+					DataflowNode->TypedScriptStruct()->ExportText(ContentString, DataflowNode.Get(), DataflowNode.Get(), nullptr, PPF_None, nullptr);
+
+					FDataflowNodeData NodeData
+					{
+						.Type = DataflowNode->GetType().ToString(),
+						.Name = DataflowNode->GetName().ToString(),
+						.Properties = MoveTemp(ContentString),
+						.Position = { (double)EdNode->NodePosX, (double)EdNode->NodePosY },
+					};
+					CopyPasteContent.NodeData.Emplace(NodeData);
+				}
+			}
+			else if (const UEdGraphNode_Comment* CommentEdNode = Cast<const UEdGraphNode_Comment>(EdNode))
+			{
+				FDataflowCommentNodeData CommentNodeData
+				{
+					.Name = CommentEdNode->NodeComment,
+					.Size = { (double)CommentEdNode->NodeWidth, (double)CommentEdNode->NodeHeight },
+					.Color = CommentEdNode->CommentColor,
+					.Position = { (double)CommentEdNode->NodePosX, (double)CommentEdNode->NodePosY },
+					.FontSize = CommentEdNode->FontSize,
+				};
+				CopyPasteContent.CommentNodeData.Emplace(CommentNodeData);
+			}
+		}
+
+		// now gather connection data 
+		// Build connection data
+		for (const FDataflowInput* Input : NodeInputsToSave)
+		{
+			if (!Input) continue;
+
+			const FDataflowNode* InputNode = Input->GetOwningNode();
+			if (!InputNode) continue;
+
+			const FDataflowOutput* Output = Input->GetConnection();
+			if (!Output) continue;
+
+			const FDataflowNode* OutputNode = Output->GetOwningNode();
+			if (!OutputNode) continue;
+
+			if (NodeGuids.Contains(OutputNode->GetGuid()))
+			{
+				FDataflowConnectionData DataflowConnectionData
+				{
+					.Out = FString::Format(TEXT("/{0}:{1}"), { OutputNode->GetName().ToString(), Output->GetName().ToString() }),
+					.In = FString::Format(TEXT("/{0}:{1}"), { InputNode->GetName().ToString(), Input->GetName().ToString() }),
 				};
+				CopyPasteContent.ConnectionData.Emplace(DataflowConnectionData);
+			}
+		}
+
+		// copy to clipboard 
+		if (CopyPasteContent.NodeData.Num() ||
+			CopyPasteContent.CommentNodeData.Num() ||
+			CopyPasteContent.ConnectionData.Num())
+		{
+			FString ClipboardContent;
+			const FDataflowCopyPasteContent DefaultContent;
+			FDataflowCopyPasteContent::StaticStruct()->ExportText(ClipboardContent, &CopyPasteContent, &DefaultContent, nullptr, PPF_None, nullptr);
 
-			ChangeDataflowAssetWithTransaction(Dataflow, TransactionName, DuplicateNodeInternal, NAME_None);
+			FPlatformApplicationMisc::ClipboardCopy(*ClipboardContent);
 		}
-		return EdNodeToReturn;
+
+		OutNumCopiedNodes = CopyPasteContent.NodeData.Num() + CopyPasteContent.CommentNodeData.Num();
+	}
+
+	//////////////////////////////////////////////////////////////////////////////////////////////////////////
+	void FEditAssetUtils::PasteNodesFromClipboard(UEdGraph* EdGraph, const FVector2D& Location, TArray<UEdGraphNode*>& OutPastedNodes)
+	{
+		using namespace EditAssetUtils::Private;
+
+		UDataflow* DataflowAsset = UDataflow::GetDataflowAssetFromEdGraph(EdGraph);
+		if (!DataflowAsset)
+		{
+			// no graph to copy to
+			return;
+		}
+
+		FString ClipboardPayload;
+		FPlatformApplicationMisc::ClipboardPaste(ClipboardPayload);
+
+		if (ClipboardPayload.IsEmpty())
+		{
+			// nothing to paste, nothing to do
+			return;
+		}
+
+		const FDataflowCopyPasteContent DefaultContent;
+		FDataflowCopyPasteContent CopyPasteContent;
+		FDataflowCopyPasteContent::StaticStruct()->ImportText(*ClipboardPayload, &CopyPasteContent, nullptr, EPropertyPortFlags::PPF_None, nullptr, FDataflowCopyPasteContent::StaticStruct()->GetName(), true);
+
+		const int32 TotalNodesToPaste = CopyPasteContent.NodeData.Num() + CopyPasteContent.CommentNodeData.Num();
+		const FText TransactionName = FText::Format(LOCTEXT("PasteDataflowNodes", "Paste {0} Dataflow Nodes"), FText::AsNumber(TotalNodesToPaste));
+
+		auto PasteNodesInternal =
+		[&EdGraph, &CopyPasteContent, &Location, &OutPastedNodes](UDataflow& DataflowAsset)->EChangeResult
+		{
+			TMap<FString, UDataflowEdNode*> OriginalNodeNameToEdNode;
+
+			// compute a ref for all nodes to refer to 
+			FVector2D RefLocation { 0, 0 };
+			if (CopyPasteContent.NodeData.Num())
+			{
+				RefLocation.X = CopyPasteContent.NodeData[0].Position.X;
+				RefLocation.Y = CopyPasteContent.NodeData[0].Position.Y;
+			}
+			else if (CopyPasteContent.CommentNodeData.Num())
+			{
+				RefLocation.X = CopyPasteContent.CommentNodeData[0].Position.X;
+				RefLocation.Y = CopyPasteContent.CommentNodeData[0].Position.Y;
+			}
+
+			// paste nodes
+			for (const FDataflowNodeData& NodeData : CopyPasteContent.NodeData)
+			{
+				const FName NodeType(*NodeData.Type);
+				const FName NodeName(*NodeData.Name);
+				const FVector2D NodeLocation(Location + (NodeData.Position - RefLocation));
+
+				if (TSharedPtr<FDataflowNode> DataflowNode = AddDataflowNode(&DataflowAsset, NodeName, NodeType))
+				{
+					// load properties to DataflowNode
+					if (!NodeData.Properties.IsEmpty())
+					{
+						DataflowNode->TypedScriptStruct()->ImportText(*NodeData.Properties, DataflowNode.Get(), nullptr, EPropertyPortFlags::PPF_None, nullptr, DataflowNode->TypedScriptStruct()->GetName(), true);
+					}
+					// Do any post-import fixup.
+					FArchive Ar;
+					Ar.SetIsLoading(true);
+					DataflowNode->PostSerialize(Ar);
+		
+					if (UDataflowEdNode* EdNode = CreateDataflowEdNode(EdGraph, DataflowNode, NodeLocation, nullptr))
+					{
+						OriginalNodeNameToEdNode.Add(NodeData.Name, EdNode);
+						OutPastedNodes.Add(EdNode);
+					}
+				}
+			}
+
+			// Paste Comment nodes
+			for (const FDataflowCommentNodeData& CommentNodeData : CopyPasteContent.CommentNodeData)
+			{
+				const FVector2D CommentNodeLocation(Location + (CommentNodeData.Position - RefLocation));
+
+				if (UEdGraphNode_Comment* CommentEdNode = CreateCommentEdNode(EdGraph, CommentNodeLocation, CommentNodeData.Name, CommentNodeData.Size, CommentNodeData.Color, CommentNodeData.FontSize))
+				{
+					OutPastedNodes.Add(CommentEdNode);
+				}
+			}
+
+			// Recreate connections
+			for (const FDataflowConnectionData& Connection : CopyPasteContent.ConnectionData)
+			{
+				const FString NodeIn = FDataflowConnectionData::GetNode(Connection.In);
+				const UDataflowEdNode* EdNodeIn = OriginalNodeNameToEdNode[NodeIn];
+				const FGuid GuidIn = EdNodeIn? EdNodeIn->DataflowNodeGuid: FGuid();
+
+				const FString PropertyIn = FDataflowConnectionData::GetProperty(Connection.In);
+				const FName InputputName = *PropertyIn;
+
+				const FString NodeOut = FDataflowConnectionData::GetNode(Connection.Out);
+				const UDataflowEdNode* EdNodeOut = OriginalNodeNameToEdNode[NodeOut];
+				const FGuid GuidOut = EdNodeOut? EdNodeOut->DataflowNodeGuid: FGuid();
+
+				const FString PropertyOut = FDataflowConnectionData::GetProperty(Connection.Out);
+				const FName OutputputName = *PropertyOut;
+
+				if (TSharedPtr<UE::Dataflow::FGraph> DataflowGraph = DataflowAsset.GetDataflow())
+				{
+					if (TSharedPtr<FDataflowNode> DataflowNodeFrom = DataflowGraph->FindBaseNode(GuidOut))
+					{
+						if (TSharedPtr<FDataflowNode> DataflowNodeTo = DataflowGraph->FindBaseNode(GuidIn))
+						{
+							// first connect the edgraph as this may affect the dataflow inputs ( for AnyType )
+							UEdGraphPin* OutputPin = FindPin(EdNodeOut, EEdGraphPinDirection::EGPD_Output, OutputputName);
+							UEdGraphPin* InputPin = FindPin(EdNodeIn, EEdGraphPinDirection::EGPD_Input, InputputName);
+							if (OutputPin && InputPin)
+							{
+								EdGraph->GetSchema()->TryCreateConnection(OutputPin, InputPin);
+							}
+
+							// now connect the dataflow
+							FDataflowInput* InputConnection = DataflowNodeTo->FindInput(InputputName);
+							FDataflowOutput* OutputConnection = DataflowNodeFrom->FindOutput(OutputputName);
+							DataflowGraph->Connect(OutputConnection, InputConnection);
+
+						}
+					}
+				}
+			}
+			return EChangeResult::Changed;
+		};
+
+		ChangeDataflowAssetWithTransaction(DataflowAsset, TransactionName, PasteNodesInternal, NAME_None);
 	}
 
 	//////////////////////////////////////////////////////////////////////////////////////////////////////////
diff --git a/Engine/Plugins/Experimental/Dataflow/Source/DataflowEditor/Private/Dataflow/DataflowAssetEditUtils.h b/Engine/Plugins/Experimental/Dataflow/Source/DataflowEditor/Private/Dataflow/DataflowAssetEditUtils.h
--- a/Engine/Plugins/Experimental/Dataflow/Source/DataflowEditor/Private/Dataflow/DataflowAssetEditUtils.h
+++ b/Engine/Plugins/Experimental/Dataflow/Source/DataflowEditor/Private/Dataflow/DataflowAssetEditUtils.h
@@ -13,8 +13,9 @@ struct FDataflowNode;
 struct FInstancedPropertyBag;
 struct FPropertyBagPropertyDesc;
 class UDataflow;
-class UEdGraph;
 class UDataflowEdNode;
+class UEdGraph;
+class UEdGraphNode;
 class UEdGraphPin;
 
 /**
@@ -46,11 +47,21 @@ namespace UE::Dataflow
 		/**
 		* Duplicate a node
 		* @param EdGraph			Main Graph or Subgraph to duplicate the node in 
-		* @param NodeToDuplicate	Node to duplicate
+		* @param NodeToDuplicate	List of nodes to duplicate
 		* @param Location			Location on the graph where to place the node 
 		* @return the newly created EdGraphNode
 		*/
-		static UDataflowEdNode* DuplicateNode(UEdGraph* EdGraph, TSharedPtr<FDataflowNode> NodeToDuplicate, const FVector2D& Location);
+		static void DuplicateNodes(UEdGraph* EdGraph, const TArray<UEdGraphNode*>& NodesToDuplicate, const FVector2D& Location, TArray<UEdGraphNode*>& OutDuplicatedNodes);
+
+		/**
+		* Copy a list of nodes and their connection to the clipboard
+		*/
+		static void CopyNodesToClipboard(const TArray<const UEdGraphNode*>& NodesToCopy, int32& OutNumCopiedNodes);
+
+		/**
+		* Attempt to paste nodes from the clipboard
+		*/
+		static void PasteNodesFromClipboard(UEdGraph* EdGraph, const FVector2D& Location, TArray<UEdGraphNode*>& OutPastedNodes);
 
 		/////////////////////////////////////////////////////////////////////////////////////////////////////////////////////////
 		//
