Dataflow : move some of the graph schema action code into the Asset util class ( Create node / suplicate node )

diff --git a/Engine/Plugins/Experimental/Dataflow/Source/DataflowEditor/Private/Dataflow/DataflowEditorCommands.cpp b/Engine/Plugins/Experimental/Dataflow/Source/DataflowEditor/Private/Dataflow/DataflowEditorCommands.cpp
--- a/Engine/Plugins/Experimental/Dataflow/Source/DataflowEditor/Private/Dataflow/DataflowEditorCommands.cpp
+++ b/Engine/Plugins/Experimental/Dataflow/Source/DataflowEditor/Private/Dataflow/DataflowEditorCommands.cpp
@@ -15,6 +15,7 @@
 #include "Dataflow/DataflowSCommentNode.h"
 #include "Dataflow/DataflowGraphSchemaAction.h"
 #include "Dataflow/DataflowToolRegistry.h"
+#include "Dataflow/DataflowAssetEditUtils.h"
 #include "DataflowEditorTools/DataflowEditorWeightMapPaintTool.h"
 #include "EdGraphNode_Comment.h"
 #include "EdGraph/EdGraphNode.h"
@@ -521,194 +522,38 @@ void FDataflowEditorCommands::DuplicateNodes(UDataflow* Graph, const TSharedPtr<
 {
 	if (ensureMsgf(Graph != nullptr, TEXT("Warning : Failed to find valid graph.")))
 	{
-		// Separate selected nodes into an array of UDataflowEdNodes and
-		// an array of UEdGraphNode_Comments
-		TSet<UDataflowEdNode*> SelectedEdNodes;
-		TSet<UEdGraphNode_Comment*> SelectedEdCommentNodes;
-
-		for (UObject* Node : SelectedNodes)
+		TArray<UEdGraphNode*> SelectedEdNodes;
+		for (UObject* SelectedNode : SelectedNodes)
 		{
-			if (UDataflowEdNode* EdNode = dynamic_cast<UDataflowEdNode*>(Node))
+			if (UEdGraphNode* EdNode = Cast<UEdGraphNode>(SelectedNode))
 			{
 				SelectedEdNodes.Add(EdNode);
 			}
-			else if (UEdGraphNode_Comment* EdCommentNode = dynamic_cast<UEdGraphNode_Comment*>(Node))
-			{
-				SelectedEdCommentNodes.Add(EdCommentNode);
-			}
 		}
+		TArray<UEdGraphNode*> DuplicatedEdNodes;
+		const FVector2D PasteLocation = DataflowGraphEditor->GetPasteLocation();
+		UE::Dataflow::FEditAssetUtils::DuplicateNodes(Graph, SelectedEdNodes, PasteLocation, DuplicatedEdNodes);
 
-		FDataflowAssetEdit Edit = Graph->EditDataflow();
-		if (UE::Dataflow::FGraph* DataflowGraph = Edit.GetGraph())
+		// Update the selection in the Editor
+		if (DuplicatedEdNodes.Num())
 		{
-			// Process Dataflow nodes
-			TSet<UDataflowEdNode*> DuplicatedEdNodes;
-			FVector2D RefLocation;
-
-			if (SelectedEdNodes.Num() > 0)
-			{
-				// Store the location of the first selected node for recreating spatial relationships
-				RefLocation.X = SelectedEdNodes.Array()[0]->NodePosX; RefLocation.Y = SelectedEdNodes.Array()[0]->NodePosY;
-
-				//
-				// Create a map to record OriginalNode.Guid -> DuplicatedNode.Guid
-				//
-				TMap<FGuid, FGuid> NodeGuidMap;				// [OriginalDataflowNodeGuid -> DuplicatedDataflowNodeGuid]
-				TMap<FGuid, UDataflowEdNode*> EdNodeMap;	// [OriginalDataflowNodeGuid -> DuplicatedEdNode]
-
-				// Duplicate selected nodes
-				for (UDataflowEdNode* SelectedEdNode : SelectedEdNodes)
-				{
-					if (TSharedPtr<FDataflowNode> DataflowNode = DataflowGraph->FindBaseNode(SelectedEdNode->DataflowNodeGuid))
-					{
-						if (TSharedPtr<FAssetSchemaAction_Dataflow_DuplicateNode_DataflowEdNode> DuplicateNodeAction = FAssetSchemaAction_Dataflow_DuplicateNode_DataflowEdNode::CreateAction(Graph, DataflowNode->GetType()))
-						{
-							DuplicateNodeAction->DataflowNodeToDuplicate = DataflowNode;
-
-							FVector2D SelectedEdNodeLocation(SelectedEdNode->NodePosX, SelectedEdNode->NodePosY);
-							FVector2D DeltaLocation = SelectedEdNodeLocation - RefLocation;
-
-							if (UDataflowEdNode* NewEdNode = (UDataflowEdNode*)DuplicateNodeAction->PerformAction(Graph, nullptr, DataflowGraphEditor->GetPasteLocation() + DeltaLocation, false))
-							{
-								DuplicatedEdNodes.Add(NewEdNode);
-
-								// Record Guids
-								NodeGuidMap.Add(SelectedEdNode->DataflowNodeGuid, NewEdNode->DataflowNodeGuid);
-								EdNodeMap.Add(SelectedEdNode->DataflowNodeGuid, NewEdNode);
-							}
-						}
-					}
-				}
-
-				// Recreate connections between duplicated nodes
-				for (UDataflowEdNode* SelectedEdNode : SelectedEdNodes)
-				{
-					if (TSharedPtr<FDataflowNode> DataflowNode = DataflowGraph->FindBaseNode(SelectedEdNode->DataflowNodeGuid))
-					{
-						FGuid DataflowNodeAGuid = DataflowNode->GetGuid();
-
-						for (FDataflowOutput* Output : DataflowNode->GetOutputs())
-						{
-							for (FDataflowInput* Connection : Output->Connections)
-							{
-								const FName OutputputName = Connection->GetConnection()->GetName();
-
-								// Check if the node on the end of the conmnection was duplicated
-								FGuid DataflowNodeBGuid = Connection->GetOwningNode()->GetGuid();
-
-								if (NodeGuidMap.Contains(DataflowNodeBGuid))
-								{
-									const FName InputputName = Connection->GetName();
-
-									if (TSharedPtr<FDataflowNode> DuplicatedDataflowNodeA = DataflowGraph->FindBaseNode(NodeGuidMap[DataflowNodeAGuid]))
-									{
-										FDataflowOutput* OutputConnection = DuplicatedDataflowNodeA->FindOutput(OutputputName);
-
-										if (TSharedPtr<FDataflowNode> DuplicatedDataflowNodeB = DataflowGraph->FindBaseNode(NodeGuidMap[DataflowNodeBGuid]))
-										{
-											FDataflowInput* InputConnection = DuplicatedDataflowNodeB->FindInput(InputputName);
-
-											DataflowGraph->Connect(OutputConnection, InputConnection);
-
-											// Connect the UDataflowEdNode FPins as well
-											if (UEdGraphPin* OutputPin = GetPin(EdNodeMap[DataflowNodeAGuid], EEdGraphPinDirection::EGPD_Output, OutputputName))
-											{
-												if (UEdGraphPin* InputPin = GetPin(EdNodeMap[DataflowNodeBGuid], EEdGraphPinDirection::EGPD_Input, InputputName))
-												{
-													OutputPin->MakeLinkTo(InputPin);
-												}
-											}
-										}
-									}
-								}
-							}
-						}
-					}
-				}
-			}
-
-			// Process Comment nodes
-			TSet<UEdGraphNode*> DuplicatedEdCommentNodes;
-
-			if (SelectedEdCommentNodes.Num() > 0)
-			{
-				if (SelectedEdNodes.Num() == 0)
-				{
-					RefLocation.X = SelectedEdCommentNodes.Array()[0]->NodePosX; RefLocation.Y = SelectedEdCommentNodes.Array()[0]->NodePosY;
-				}
-
-				const TSharedPtr<SGraphEditor>& InGraphEditor = (TSharedPtr<SGraphEditor>)DataflowGraphEditor;
-
-				for (UEdGraphNode_Comment* SelectedCommentNode : SelectedEdCommentNodes)
-				{
-					if (TSharedPtr<FAssetSchemaAction_Dataflow_DuplicateCommentNode_DataflowEdNode> DuplicateCommentNodeAction = FAssetSchemaAction_Dataflow_DuplicateCommentNode_DataflowEdNode::CreateAction(Graph, InGraphEditor))
-					{
-						DuplicateCommentNodeAction->CommentNodeToDuplicate = SelectedCommentNode;
-
-						FVector2D SelectedCommentNodeLocation(SelectedCommentNode->NodePosX, SelectedCommentNode->NodePosY);
-						FVector2D DeltaLocation = SelectedCommentNodeLocation - RefLocation;
-
-						if (UEdGraphNode* NewCommentNode = DuplicateCommentNodeAction->PerformAction(Graph, nullptr, DataflowGraphEditor->GetPasteLocation() + DeltaLocation, false))
-						{
-							DuplicatedEdCommentNodes.Add(NewCommentNode);
-						}
-					}
-				}
-			}
-
-			// Update the selection in the Editor
-			if (SelectedEdNodes.Num() > 0 || SelectedEdCommentNodes.Num() > 0)
+			DataflowGraphEditor->ClearSelectionSet();
+			for (UEdGraphNode* Node : DuplicatedEdNodes)
 			{
-				DataflowGraphEditor->ClearSelectionSet();
-
-				if (SelectedEdNodes.Num() > 0)
-				{
-					for (UDataflowEdNode* Node : DuplicatedEdNodes)
-					{
-						DataflowGraphEditor->SetNodeSelection(Node, true);
-					}
-				}
-
-				if (SelectedEdCommentNodes.Num() > 0)
-				{
-					for (UEdGraphNode* Node : DuplicatedEdCommentNodes)
-					{
-						DataflowGraphEditor->SetNodeSelection(Node, true);
-					}
-				}
+				DataflowGraphEditor->SetNodeSelection(Node, true);
 			}
 
 			// Display message stating that nodes were duplicated
-			const int32 NumDuplicatedNodes = SelectedEdNodes.Num();
-			if (NumDuplicatedNodes > 0)
+			FText MessageFormat;
+			if (DuplicatedEdNodes.Num() == 1)
 			{
-				FText MessageFormat;
-				if (NumDuplicatedNodes == 1)
-				{
-					MessageFormat = LOCTEXT("DataflowDuplicatedNodesSingleNode", "{0} node was duplicated");
-				}
-				else
-				{
-					MessageFormat = LOCTEXT("DataflowDuplicatedNodesMultipleNodes", "{0} nodes were duplicated");
-				}
-				ShowNotificationMessage(FText::Format(MessageFormat, NumDuplicatedNodes), SNotificationItem::CS_Success);
+				MessageFormat = LOCTEXT("DataflowDuplicatedNodesSingleNode", "{0} node/comment was duplicated");
 			}
-
-			// Display message stating that comment boxe(s) were duplicated
-			const int32 NumDuplicatedComments = SelectedEdCommentNodes.Num();
-			if (NumDuplicatedComments > 0)
+			else
 			{
-				FText MessageFormat;
-				if (NumDuplicatedComments == 1)
-				{
-					MessageFormat = LOCTEXT("DataflowDuplicatedNodesToClipboardSingleComment", "{0} comment was duplicated");
-				}
-				else
-				{
-					MessageFormat = LOCTEXT("DataflowDuplicatedNodesToClipboardMultipleComments", "{0} comments were duplicated");
-				}
-				ShowNotificationMessage(FText::Format(MessageFormat, NumDuplicatedComments), SNotificationItem::CS_Success);
+				MessageFormat = LOCTEXT("DataflowDuplicatedNodesMultipleNodes", "{0} nodes/comments were duplicated");
 			}
+			ShowNotificationMessage(FText::Format(MessageFormat, DuplicatedEdNodes.Num()), SNotificationItem::CS_Success);
 		}
 	}
 }
@@ -717,152 +562,33 @@ void FDataflowEditorCommands::CopyNodes(UDataflow* InGraph, const TSharedPtr<SDa
 {
 	if (ensureMsgf(InGraph != nullptr, TEXT("Warning : Failed to find valid graph.")))
 	{
-		// 
 		if (InSelectedNodes.Num() > 0)
 		{
-			// Separate selected nodes into an array of UDataflowEdNodes and
-			// an array of UEdGraphNode_Comments
-			TSet<UDataflowEdNode*> SelectedEdNodes;
-			TSet<UEdGraphNode_Comment*> SelectedEdCommentNodes;
-
-			for (UObject* Node : InSelectedNodes)
+			TArray<const UEdGraphNode*> SelectedEdNodes;
+			for (UObject* SelectedNode : InSelectedNodes)
 			{
-				if (UDataflowEdNode* EdNode = dynamic_cast<UDataflowEdNode*>(Node))
+				if (const UEdGraphNode* EdNode = Cast<const UEdGraphNode>(SelectedNode))
 				{
 					SelectedEdNodes.Add(EdNode);
 				}
-				else if (UEdGraphNode_Comment* EdCommentNode = dynamic_cast<UEdGraphNode_Comment*>(Node))
-				{
-					SelectedEdCommentNodes.Add(EdCommentNode);
-				}
 			}
 
-			FDataflowAssetEdit Edit = InGraph->EditDataflow();
-			if (UE::Dataflow::FGraph* DataflowGraph = Edit.GetGraph())
-			{
-				FDataflowCopyPasteContent CopyPasteContent;
-
-				TSet<FGuid> NodeGuids;
-				TArray<FDataflowInput*> NodeInputsToSave;
-
-				// Build node data
-				if (SelectedEdNodes.Num() > 0)
-				{
-					for (UDataflowEdNode* EdNode : SelectedEdNodes)
-					{
-						if (TSharedPtr<FDataflowNode> DataflowNode = DataflowGraph->FindBaseNode(EdNode->DataflowNodeGuid))
-						{
-							NodeGuids.Add(DataflowNode->GetGuid());
-							NodeInputsToSave.Append(DataflowNode->GetInputs());
-
-							FName NodeType = DataflowNode->GetType();
-							FName NodeName = DataflowNode->GetName();
-
-							FDataflowNodeData NodeData;
-
-							NodeData.Type = NodeType.ToString();
-							NodeData.Name = NodeName.ToString();
-							NodeData.Position.X = EdNode->NodePosX;
-							NodeData.Position.Y = EdNode->NodePosY;
-
-							FString ContentString;
-
-							TUniquePtr<FDataflowNode> DefaultElement;
-							DataflowNode->TypedScriptStruct()->ExportText(ContentString, DataflowNode.Get(), DataflowNode.Get(), nullptr, PPF_None, nullptr);
-
-							NodeData.Properties = ContentString;
-
-							CopyPasteContent.NodeData.Add(NodeData);
-						}
-					}				
-				}
-
-				// Build comment node data
-				if (SelectedEdCommentNodes.Num() > 0)
-				{
-					for (UEdGraphNode_Comment* CommentEdNode : SelectedEdCommentNodes)
-					{
-						FDataflowCommentNodeData CommentNodeData;
-
-						CommentNodeData.Name = CommentEdNode->NodeComment;
-						CommentNodeData.Size.X = CommentEdNode->NodeWidth;
-						CommentNodeData.Size.Y = CommentEdNode->NodeHeight;
-						CommentNodeData.Position.X = CommentEdNode->NodePosX;
-						CommentNodeData.Position.Y = CommentEdNode->NodePosY;
-						CommentNodeData.Color = CommentEdNode->CommentColor;
-						CommentNodeData.FontSize = CommentEdNode->FontSize;
-
-						CopyPasteContent.CommentNodeData.Add(CommentNodeData);
-					}
-				}
-
-				// Build connection data
-				for (FDataflowInput* Input : NodeInputsToSave)
-				{
-					if (!Input) continue;
+			int32 NumCopiedNodes = 0;
+			UE::Dataflow::FEditAssetUtils::CopyNodesToClipboard(SelectedEdNodes, NumCopiedNodes);
 
-					const FDataflowNode* InputNode = Input->GetOwningNode();
-					if (!InputNode) continue;
-
-					const FDataflowOutput* Output = Input->GetConnection();
-					if (!Output) continue;
-
-					const FDataflowNode* OutputNode = Output->GetOwningNode();
-					if (!OutputNode) continue;
-
-					if (!NodeGuids.Contains(OutputNode->GetGuid()))
-						continue;
-
-					FString InConnection = FString::Format(TEXT("/{0}:{1}"), {InputNode->GetName().ToString(), Input->GetName().ToString()});
-					FString OutConnection = FString::Format(TEXT("/{0}:{1}"), { OutputNode->GetName().ToString(), Output->GetName().ToString() });
-
-					FDataflowConnectionData DataflowConnectionData;
-
-					DataflowConnectionData.In = InConnection;
-					DataflowConnectionData.Out = OutConnection;
-
-					CopyPasteContent.ConnectionData.Add(DataflowConnectionData);
-				}
-
-				FString ClipboardContent;
-				FDataflowCopyPasteContent DefaultContent;
-
-				FDataflowCopyPasteContent::StaticStruct()->ExportText(ClipboardContent, &CopyPasteContent, &DefaultContent, nullptr, PPF_None, nullptr);
-
-				// Save to clipboard
-				FPlatformApplicationMisc::ClipboardCopy(*ClipboardContent);
-
-				// Display message stating that nodes were copied to clipboard
-				const int32 NumCopiedNodes = CopyPasteContent.NodeData.Num();
-				if (NumCopiedNodes > 0)
+			// Display message stating that nodes were copied to clipboard
+			if (NumCopiedNodes > 0)
+			{
+				FText MessageFormat;
+				if (NumCopiedNodes == 1)
 				{
-					FText MessageFormat;
-					if (NumCopiedNodes == 1)
-					{
-						MessageFormat = LOCTEXT("DataflowCopiedNodesToClipboardSingleNode", "{0} node was copied to clipboard");
-					}
-					else
-					{
-						MessageFormat = LOCTEXT("DataflowCopiedNodesToClipboardMultipleNodes", "{0} nodes were copied to clipboard");
-					}
-					ShowNotificationMessage(FText::Format(MessageFormat, NumCopiedNodes), SNotificationItem::CS_Success);
+					MessageFormat = LOCTEXT("DataflowCopiedNodesToClipboardSingleNode", "{0} node/comment was copied to clipboard");
 				}
-
-				// Display message stating that comment boxe(s) were copied to clipboard
-				const int32 NumCopiedComments = CopyPasteContent.CommentNodeData.Num();
-				if (NumCopiedComments > 0)
+				else
 				{
-					FText MessageFormat;
-					if (NumCopiedComments == 1)
-					{
-						MessageFormat = LOCTEXT("DataflowCopiedNodesToClipboardSingleComment", "{0} comment was copied to clipboard");
-					}
-					else
-					{
-						MessageFormat = LOCTEXT("DataflowCopiedNodesToClipboardMultipleComments", "{0} comments were copied to clipboard");
-					}
-					ShowNotificationMessage(FText::Format(MessageFormat, NumCopiedComments), SNotificationItem::CS_Success);
+					MessageFormat = LOCTEXT("DataflowCopiedNodesToClipboardMultipleNodes", "{0} nodes/comments were copied to clipboard");
 				}
+				ShowNotificationMessage(FText::Format(MessageFormat, NumCopiedNodes), SNotificationItem::CS_Success);
 			}
 		}
 	}
@@ -870,191 +596,31 @@ void FDataflowEditorCommands::CopyNodes(UDataflow* InGraph, const TSharedPtr<SDa
 
 void FDataflowEditorCommands::PasteNodes(UDataflow* Graph, const TSharedPtr<SDataflowGraphEditor>& DataflowGraphEditor)
 {
-	FString ClipboardContent;
-	FPlatformApplicationMisc::ClipboardPaste(ClipboardContent);
+	TArray<UEdGraphNode*> PastedEdNodes;
+	const FVector2D PasteLocation = DataflowGraphEditor->GetPasteLocation();
+	UE::Dataflow::FEditAssetUtils::PasteNodesFromClipboard(Graph, PasteLocation, PastedEdNodes);
 
-	if (!ClipboardContent.IsEmpty())
+	// select the pasted nodes
+	if (PastedEdNodes.Num())
 	{
-		FDataflowCopyPasteContent DefaultContent;
-		FDataflowCopyPasteContent CopyPasteContent;
-
-		FDataflowCopyPasteContent::StaticStruct()->ImportText(*ClipboardContent, &CopyPasteContent, nullptr, EPropertyPortFlags::PPF_None, nullptr, FDataflowCopyPasteContent::StaticStruct()->GetName(), true);
-
-		FVector2D AppliedTranslation; // The translation that being applied to everything. for placing comment nodes
-
-		// Paste nodes
-		TSet<UDataflowEdNode*> PastedEdNodes;
-		TSet<UEdGraphNode*> PastedEdCommentNodes;
-
-		TMap<FString, UDataflowEdNode*> EdNodeMap;	// [NmaeOfNode -> DuplicatedEdNode]
-
-		if (CopyPasteContent.NodeData.Num() > 0)
-		{
-			// Store the location of the first selected node for recreating spatial relationships
-			FVector2D RefLocation;
-			RefLocation.X = CopyPasteContent.NodeData[0].Position.X; RefLocation.Y = CopyPasteContent.NodeData[0].Position.Y;
-
-			int32 Idx = 0;
-			for (FDataflowNodeData NodeData : CopyPasteContent.NodeData)
-			{
-				FName NodeType = *NodeData.Type;
-
-				if (TSharedPtr<FAssetSchemaAction_Dataflow_PasteNode_DataflowEdNode> PasteNodeAction = FAssetSchemaAction_Dataflow_PasteNode_DataflowEdNode::CreateAction(Graph, NodeType))
-				{
-					PasteNodeAction->NodeName = *NodeData.Name;
-					PasteNodeAction->NodeProperties = NodeData.Properties;
-
-					FVector2D NodeLocation(NodeData.Position.X, NodeData.Position.Y);
-					FVector2D DeltaLocation = NodeLocation - RefLocation;
-
-					if (Idx == 0)
-					{
-						AppliedTranslation = DataflowGraphEditor->GetPasteLocation() + DeltaLocation - NodeData.Position;
-					}
-					Idx++;
-
-					if (UDataflowEdNode* NewEdNode = (UDataflowEdNode*)PasteNodeAction->PerformAction(Graph, nullptr, DataflowGraphEditor->GetPasteLocation() + DeltaLocation, false))
-					{
-						PastedEdNodes.Add(NewEdNode);
-
-						EdNodeMap.Add(NodeData.Name, NewEdNode);
-					}
-				}
-			}
-		}
-
-		// Paste Comment nodes
-		if (CopyPasteContent.CommentNodeData.Num() > 0)
-		{
-			const TSharedPtr<SGraphEditor>& InGraphEditor = (TSharedPtr<SGraphEditor>)DataflowGraphEditor;
-
-			int32 Idx = 0;
-
-			for (FDataflowCommentNodeData CommentNodeData : CopyPasteContent.CommentNodeData)
-			{
-				if (TSharedPtr<FAssetSchemaAction_Dataflow_PasteCommentNode_DataflowEdNode> PasteCommentNodeAction = FAssetSchemaAction_Dataflow_PasteCommentNode_DataflowEdNode::CreateAction(Graph, InGraphEditor))
-				{
-					if (CopyPasteContent.NodeData.Num() == 0)
-					{
-						if (Idx == 0)
-						{
-							AppliedTranslation = DataflowGraphEditor->GetPasteLocation() - CommentNodeData.Position;
-						}
-					}
-					Idx++;
-
-					PasteCommentNodeAction->NodeName = *CommentNodeData.Name;
-					PasteCommentNodeAction->Size.X = CommentNodeData.Size.X + 50; // Make it longer, because the nodes are longer after copying ('_copy' in their name)
-					PasteCommentNodeAction->Size.Y = CommentNodeData.Size.Y + 30;
-					PasteCommentNodeAction->Color = CommentNodeData.Color;
-					PasteCommentNodeAction->FontSize = CommentNodeData.FontSize;
-
-					FVector2D CommentNodeLocation(CommentNodeData.Position);
-					FVector2D NewLocation = CommentNodeLocation + AppliedTranslation;
-
-					if (UEdGraphNode* NewCommentNode = PasteCommentNodeAction->PerformAction(Graph, nullptr, NewLocation, false))
-					{
-						PastedEdCommentNodes.Add(NewCommentNode);
-					}
-				}
-			}
-		}
-
-		// Recreate connections
-		if (CopyPasteContent.ConnectionData.Num() > 0)
+		DataflowGraphEditor->ClearSelectionSet();
+		for (UEdGraphNode* Node : PastedEdNodes)
 		{
-			for (FDataflowConnectionData Connection : CopyPasteContent.ConnectionData)
-			{
-				FString NodeIn = FDataflowConnectionData::GetNode(Connection.In);
-				FGuid GuidIn = (EdNodeMap[NodeIn])->DataflowNodeGuid;
-
-				FString PropertyIn = FDataflowConnectionData::GetProperty(Connection.In);
-				const FName InputputName = *PropertyIn;
-
-				FString NodeOut = FDataflowConnectionData::GetNode(Connection.Out);
-				FGuid GuidOut = (EdNodeMap[NodeOut])->DataflowNodeGuid;
-
-				FString PropertyOut = FDataflowConnectionData::GetProperty(Connection.Out);
-				const FName OutputputName = *PropertyOut;
-
-				FDataflowAssetEdit Edit = Graph->EditDataflow();
-				if (UE::Dataflow::FGraph* DataflowGraph = Edit.GetGraph())
-				{
-					if (TSharedPtr<FDataflowNode> DataflowNodeFrom = DataflowGraph->FindBaseNode(GuidOut))
-					{
-						if (TSharedPtr<FDataflowNode> DataflowNodeTo = DataflowGraph->FindBaseNode(GuidIn))
-						{
-							// first connect the edgraph as this may affect the dataflow inputs ( for AnyType )
-							UEdGraphPin* OutputPin = GetPin(EdNodeMap[NodeOut], EEdGraphPinDirection::EGPD_Output, OutputputName);
-							UEdGraphPin* InputPin = GetPin(EdNodeMap[NodeIn], EEdGraphPinDirection::EGPD_Input, InputputName);
-							if (OutputPin && InputPin)
-							{
-								if (UEdGraph* EdGraph = EdNodeMap[NodeOut]->GetGraph())
-								{
-									EdGraph->GetSchema()->TryCreateConnection(OutputPin, InputPin);
-								}
-							}
-
-							// now connect the dataflow
-							FDataflowInput* InputConnection = DataflowNodeTo->FindInput(InputputName);
-							FDataflowOutput* OutputConnection = DataflowNodeFrom->FindOutput(OutputputName);
-							DataflowGraph->Connect(OutputConnection, InputConnection);
-
-						}
-					}
-				}
-			}
+			DataflowGraphEditor->SetNodeSelection(Node, true);
 		}
-
-		// Update the selection in the Editor
-		if (PastedEdNodes.Num() > 0 || PastedEdCommentNodes.Num() > 0)
-		{
-			DataflowGraphEditor->ClearSelectionSet();
-
-			for (UDataflowEdNode* Node : PastedEdNodes)
-			{
-				DataflowGraphEditor->SetNodeSelection(Node, true);
-			}
-
-			for (UEdGraphNode* Node : PastedEdCommentNodes)
-			{
-				DataflowGraphEditor->SetNodeSelection(Node, true);
-			}
-		}
-
+	
 		// Display message stating that nodes were pasted from clipboard
-		const int32 NumPastedNodes = CopyPasteContent.NodeData.Num();
-		if (NumPastedNodes > 0)
+		FText MessageFormat;
+		if (PastedEdNodes.Num() == 1)
 		{
-			FText MessageFormat;
-			if (NumPastedNodes == 1)
-			{
-				MessageFormat = LOCTEXT("DataflowPastedNodesFromClipboardSingleNode", "{0} node was pasted from clipboard");
-			}
-			else
-			{
-				MessageFormat = LOCTEXT("DataflowPastedNodesFromClipboardMultipleNodes", "{0} nodes were pasted from clipboard");
-			}
-			ShowNotificationMessage(FText::Format(MessageFormat, NumPastedNodes), SNotificationItem::CS_Success);
+			MessageFormat = LOCTEXT("DataflowPastedNodesFromClipboardSingleNode", "{0} node/comment was pasted from clipboard");
 		}
-
-		// Display message stating that comment box(es) were pasted to clipboard
-		const int32 NumPastedComments = CopyPasteContent.CommentNodeData.Num();
-		if (NumPastedComments > 0)
+		else
 		{
-			FText MessageFormat;
-			if (NumPastedComments == 1)
-			{
-				MessageFormat = LOCTEXT("DataflowPastedNodesToClipboardSingleComment", "{0} comment was pasted from clipboard");
-			}
-			else
-			{
-				MessageFormat = LOCTEXT("DataflowPastedNodesToClipboardMultipleComments", "{0} comments were pasted from clipboard");
-			}
-			ShowNotificationMessage(FText::Format(MessageFormat, NumPastedComments), SNotificationItem::CS_Success);
+			MessageFormat = LOCTEXT("DataflowPastedNodesFromClipboardMultipleNodes", "{0} nodes/comments were pasted from clipboard");
 		}
+		ShowNotificationMessage(FText::Format(MessageFormat, PastedEdNodes.Num()), SNotificationItem::CS_Success);
 	}
 }
 
-
 #undef LOCTEXT_NAMESPACE
diff --git a/Engine/Plugins/Experimental/Dataflow/Source/DataflowEditor/Private/Dataflow/DataflowGraphSchemaAction.cpp b/Engine/Plugins/Experimental/Dataflow/Source/DataflowEditor/Private/Dataflow/DataflowGraphSchemaAction.cpp
--- a/Engine/Plugins/Experimental/Dataflow/Source/DataflowEditor/Private/Dataflow/DataflowGraphSchemaAction.cpp
+++ b/Engine/Plugins/Experimental/Dataflow/Source/DataflowEditor/Private/Dataflow/DataflowGraphSchemaAction.cpp
@@ -6,6 +6,7 @@
 #include "Dataflow/DataflowInstance.h"
 #include "Dataflow/DataflowObject.h"
 #include "Dataflow/DataflowSNode.h"
+#include "Dataflow/DataflowAssetEditUtils.h"
 #include "Dataflow/DataflowVariableNodes.h"
 #include "EdGraph/EdGraphSchema.h"
 #include "PropertyBagDetails.h"
@@ -49,224 +50,17 @@ TSharedPtr<FAssetSchemaAction_Dataflow_CreateNode_DataflowEdNode> FAssetSchemaAc
 	return TSharedPtr<FAssetSchemaAction_Dataflow_CreateNode_DataflowEdNode>(nullptr);
 }
 
-static FName GetNodeUniqueName(UDataflow* Dataflow, FString NodeBaseName)
-{
-	FString Left, Right;
-	int32 NameIndex = 1;
-
-	// Check if NodeBaseName already ends with "_dd"
-	if (NodeBaseName.Split(TEXT("_"), &Left, &Right, ESearchCase::IgnoreCase, ESearchDir::FromEnd))
-	{
-		if (Right.IsNumeric())
-		{
-			NameIndex = FCString::Atoi(*Right);
-
-			NodeBaseName = Left;
-		}
-	}
-
-	FName NodeUniqueName{ NodeBaseName };
-	while (Dataflow->GetDataflow()->FindBaseNode(FName(NodeUniqueName)) != nullptr)
-	{
-		NodeUniqueName = FName(NodeBaseName + FString::Printf(TEXT("_%02d"), NameIndex));
-		NameIndex++;
-	}
-
-	return NodeUniqueName;
-}
-
-static UDataflowEdNode* CreateNode(UDataflow* Dataflow, UEdGraphPin* FromPin, const FVector2D Location, bool bSelectNewNode, const FName NodeUniqueName, const FName NodeTypeName, TSharedPtr<FDataflowNode> DataflowNodeToDuplicate, bool bCopySettings = false)
-{
-	if (UE::Dataflow::FNodeFactory* Factory = UE::Dataflow::FNodeFactory::GetInstance())
-	{
-		if (TSharedPtr<FDataflowNode> DataflowNode =
-			Factory->NewNodeFromRegisteredType(
-				*Dataflow->GetDataflow(),
-				{ FGuid::NewGuid(), NodeTypeName, NodeUniqueName, Dataflow }))
-		{
-			if (UDataflowEdNode* EdNode = NewObject<UDataflowEdNode>(Dataflow, UDataflowEdNode::StaticClass(), NodeUniqueName))
-			{
-				Dataflow->Modify();
-				if (FromPin != nullptr)
-				{
-					FromPin->Modify();
-				}
-
-				Dataflow->AddNode(EdNode, true, bSelectNewNode);
-
-				// Copy properties from DataflowNodeToDuplicate to DataflowNode
-				if (bCopySettings)
-				{
-					SDataflowEdNode::CopyDataflowNodeSettings(DataflowNodeToDuplicate, DataflowNode);
-				}
-
-				EdNode->CreateNewGuid();
-				EdNode->PostPlacedNewNode();
-
-				EdNode->SetDataflowGraph(Dataflow->GetDataflow());
-				EdNode->SetDataflowNodeGuid(DataflowNode->GetGuid());
-				EdNode->AllocateDefaultPins();
-
-				EdNode->AutowireNewNode(FromPin);
-
-				EdNode->NodePosX = Location.X;
-				EdNode->NodePosY = Location.Y;
-
-				EdNode->SetFlags(RF_Transactional);
-
-				return EdNode;
-			}
-		}
-	}
-
-	return nullptr;
-}
-
-//
-//  Created the EdGraph node and bind the guids to the Dataflow's node. 
-//
 UEdGraphNode* FAssetSchemaAction_Dataflow_CreateNode_DataflowEdNode::PerformAction(class UEdGraph* ParentGraph, UEdGraphPin* FromPin, const FVector2D Location, bool bSelectNewNode)
 {
-	if (UDataflow* Dataflow = Cast<UDataflow>(ParentGraph))
-	{
-		// by default use the type name and check if it is unique in the context of the graph
-		// if not, then generate a unique name 
-		const FString NodeBaseName = GetMenuDescription().ToString();
-		FName NodeUniqueName{ NodeBaseName };
-		int32 NameIndex = 0;
-		while (Dataflow->GetDataflow()->FindBaseNode(FName(NodeUniqueName)) != nullptr)
-		{
-			NodeUniqueName = FName(NodeBaseName + FString::Printf(TEXT("_%d"), NameIndex));
-			NameIndex++;
-		}
-
-		return CreateNode(Dataflow, FromPin, Location, bSelectNewNode, NodeUniqueName, NodeTypeName, nullptr);
-	}
-
-	return nullptr;
-}
-
-/////////////////////////////////////////////////////////////////////////////////////////////////////////////////////////
+	// by default use the type name 
+	const FString NodeBaseName = GetMenuDescription().ToString();
+	UEdGraphNode* NewNode = UE::Dataflow::FEditAssetUtils::AddNewNode(ParentGraph, Location, FName(NodeBaseName), NodeTypeName, FromPin);
 
-TSharedPtr<FAssetSchemaAction_Dataflow_DuplicateNode_DataflowEdNode> FAssetSchemaAction_Dataflow_DuplicateNode_DataflowEdNode::CreateAction(UEdGraph* ParentGraph, const FName& InNodeTypeName)
-{
-	if (UE::Dataflow::FNodeFactory* Factory = UE::Dataflow::FNodeFactory::GetInstance())
+	if (NewNode && bSelectNewNode)
 	{
-		const UE::Dataflow::FFactoryParameters& Param = Factory->GetParameters(InNodeTypeName);
-		if (Param.IsValid())
-		{
-			const FText ToolTip = FText::FromString(Param.ToolTip.IsEmpty() ? FString("Add a Dataflow node.") : Param.ToolTip);
-			const FText NodeName = FText::FromString(Param.DisplayName.ToString());
-			const FText Category = FText::FromString(Param.Category.ToString().IsEmpty() ? FString("Dataflow") : Param.Category.ToString());
-			const FText Tags = FText::FromString(Param.Tags);
-			TSharedPtr<FAssetSchemaAction_Dataflow_DuplicateNode_DataflowEdNode> NewNodeAction(
-				new FAssetSchemaAction_Dataflow_DuplicateNode_DataflowEdNode(InNodeTypeName, Category, NodeName, ToolTip, Tags));
-			return NewNodeAction;
-		}
-	}
-	return TSharedPtr<FAssetSchemaAction_Dataflow_DuplicateNode_DataflowEdNode>(nullptr);
-}
-
-//
-//   
-//
-UEdGraphNode* FAssetSchemaAction_Dataflow_DuplicateNode_DataflowEdNode::PerformAction(class UEdGraph* ParentGraph, UEdGraphPin* FromPin, const FVector2D Location, bool bSelectNewNode)
-{
-	if (UDataflow* Dataflow = Cast<UDataflow>(ParentGraph))
-	{
-		FString NodeToDuplicateName = DataflowNodeToDuplicate->GetName().ToString();
-
-		// Check if that is unique, if not then make it unique with an index postfix
-		FName NodeUniqueName = GetNodeUniqueName(Dataflow, NodeToDuplicateName);
-
-		return CreateNode(Dataflow, FromPin, Location, bSelectNewNode, NodeUniqueName, NodeTypeName, DataflowNodeToDuplicate, /*bCopySettings=*/true);
-	}
-
-	return nullptr;
-}
-
-/////////////////////////////////////////////////////////////////////////////////////////////////////////////////////////
-
-TSharedPtr<FAssetSchemaAction_Dataflow_PasteNode_DataflowEdNode> FAssetSchemaAction_Dataflow_PasteNode_DataflowEdNode::CreateAction(UEdGraph* ParentGraph, const FName& InNodeTypeName)
-{
-	if (UE::Dataflow::FNodeFactory* Factory = UE::Dataflow::FNodeFactory::GetInstance())
-	{
-		const UE::Dataflow::FFactoryParameters& Param = Factory->GetParameters(InNodeTypeName);
-		if (Param.IsValid())
-		{
-			const FText ToolTip = FText::FromString(Param.ToolTip.IsEmpty() ? FString("Add a Dataflow node.") : Param.ToolTip);
-			const FText NodeName = FText::FromString(Param.DisplayName.ToString());
-			const FText Category = FText::FromString(Param.Category.ToString().IsEmpty() ? FString("Dataflow") : Param.Category.ToString());
-			const FText Tags = FText::FromString(Param.Tags);
-			TSharedPtr<FAssetSchemaAction_Dataflow_PasteNode_DataflowEdNode> NewNodeAction(
-				new FAssetSchemaAction_Dataflow_PasteNode_DataflowEdNode(InNodeTypeName, Category, NodeName, ToolTip, Tags));
-			return NewNodeAction;
-		}
-	}
-	return TSharedPtr<FAssetSchemaAction_Dataflow_PasteNode_DataflowEdNode>(nullptr);
-}
-
-static UDataflowEdNode* CreateNodeFromPaste(UDataflow* Dataflow, UEdGraphPin* FromPin, const FVector2D Location, bool bSelectNewNode, const FName NodeUniqueName, const FName NodeTypeName, FString NodeProperties)
-{
-	if (UE::Dataflow::FNodeFactory* Factory = UE::Dataflow::FNodeFactory::GetInstance())
-	{
-		if (TSharedPtr<FDataflowNode> DataflowNode =
-			Factory->NewNodeFromRegisteredType(
-				*Dataflow->GetDataflow(),
-				{ FGuid::NewGuid(), NodeTypeName, NodeUniqueName, Dataflow }))
-		{
-			if (UDataflowEdNode* EdNode = NewObject<UDataflowEdNode>(Dataflow, UDataflowEdNode::StaticClass(), NodeUniqueName))
-			{
-				Dataflow->Modify();
-
-				Dataflow->AddNode(EdNode, true, bSelectNewNode);
-
-				// Copy properties to DataflowNode
-				if (!NodeProperties.IsEmpty())
-				{
-					DataflowNode->TypedScriptStruct()->ImportText(*NodeProperties, DataflowNode.Get(), nullptr, EPropertyPortFlags::PPF_None, nullptr, DataflowNode->TypedScriptStruct()->GetName(), true);
-				}
-				// Do any post-import fixup.
-				FArchive Ar;
-				Ar.SetIsLoading(true);
-				DataflowNode->PostSerialize(Ar);
-
-				EdNode->CreateNewGuid();
-				EdNode->PostPlacedNewNode();
-
-				EdNode->SetDataflowGraph(Dataflow->GetDataflow());
-				EdNode->SetDataflowNodeGuid(DataflowNode->GetGuid());
-				EdNode->AllocateDefaultPins();
-
-				EdNode->NodePosX = Location.X;
-				EdNode->NodePosY = Location.Y;
-
-				EdNode->SetFlags(RF_Transactional);
-
-				return EdNode;
-			}
-		}
-	}
-
-	return nullptr;
-}
-
-//
-//   
-//
-UEdGraphNode* FAssetSchemaAction_Dataflow_PasteNode_DataflowEdNode::PerformAction(class UEdGraph* ParentGraph, UEdGraphPin* FromPin, const FVector2D Location, bool bSelectNewNode)
-{
-	if (UDataflow* Dataflow = Cast<UDataflow>(ParentGraph))
-	{
-		FString NodeToDuplicateName = NodeName.ToString();
-
-		// Check if that is unique, if not then make it unique with an index postfix
-		FName NodeUniqueName = GetNodeUniqueName(Dataflow, NodeToDuplicateName);
-
-		return CreateNodeFromPaste(Dataflow, FromPin, Location, bSelectNewNode, NodeUniqueName, NodeTypeName, NodeProperties);
+		ParentGraph->SelectNodeSet({ NewNode },  /*bFromUI*/true);
 	}
-
-	return nullptr;
+	return NewNode;
 }
 
 /////////////////////////////////////////////////////////////////////////////////////////////////////////////////////////
diff --git a/Engine/Plugins/Experimental/Dataflow/Source/DataflowEditor/Public/Dataflow/DataflowGraphSchemaAction.h b/Engine/Plugins/Experimental/Dataflow/Source/DataflowEditor/Public/Dataflow/DataflowGraphSchemaAction.h
--- a/Engine/Plugins/Experimental/Dataflow/Source/DataflowEditor/Public/Dataflow/DataflowGraphSchemaAction.h
+++ b/Engine/Plugins/Experimental/Dataflow/Source/DataflowEditor/Public/Dataflow/DataflowGraphSchemaAction.h
@@ -26,9 +26,9 @@ namespace UE::Dataflow
 
 ////////////////////////////////////////////////////////////////////////////////////////////////////////////////////////////////
 
-//
-// Action to add a node to the graph
-//
+/**
+* Action to add a node to the graph
+*/
 USTRUCT()
 struct DATAFLOWEDITOR_API FAssetSchemaAction_Dataflow_CreateNode_DataflowEdNode : public FEdGraphSchemaAction
 {
@@ -42,59 +42,9 @@ public:
 
 	static TSharedPtr<FAssetSchemaAction_Dataflow_CreateNode_DataflowEdNode> CreateAction(const UEdGraph* ParentGraph, const FName& NodeTypeName, const FName& InOverrideNodeName = NAME_None);
 
-	virtual UEdGraphNode* PerformAction(class UEdGraph* ParentGraph, UEdGraphPin* FromPin, const FVector2D Location, bool bSelectNewNode = true) override;
-	//virtual void AddReferencedObjects(FReferenceCollector& Collector) override;
-
-	FName NodeTypeName;
-};
-
-////////////////////////////////////////////////////////////////////////////////////////////////////////////////////////////////
-
-//
-// Action to duplicate a set of nodes in the graph
-//
-USTRUCT()
-struct DATAFLOWEDITOR_API FAssetSchemaAction_Dataflow_DuplicateNode_DataflowEdNode : public FEdGraphSchemaAction
-{
-	GENERATED_USTRUCT_BODY();
-
-public:
-	FAssetSchemaAction_Dataflow_DuplicateNode_DataflowEdNode() : FEdGraphSchemaAction() {}
-
-	FAssetSchemaAction_Dataflow_DuplicateNode_DataflowEdNode(const FName& InType, FText InNodeCategory, FText InMenuDesc, FText InToolTip, FText InKeywords)
-		: FEdGraphSchemaAction(InNodeCategory, InMenuDesc, InToolTip, 0, InKeywords), NodeTypeName(InType) {}
-
-	static TSharedPtr<FAssetSchemaAction_Dataflow_DuplicateNode_DataflowEdNode> CreateAction(UEdGraph* ParentGraph, const FName& NodeTypeName);
-
-	virtual UEdGraphNode* PerformAction(class UEdGraph* ParentGraph, UEdGraphPin* FromPin, const FVector2D Location, bool bSelectNewNode = true) override;
-
-	FName NodeTypeName;
-	TSharedPtr<FDataflowNode> DataflowNodeToDuplicate;
-};
-
-////////////////////////////////////////////////////////////////////////////////////////////////////////////////////////////////
-
-//
-// Action to paste a set of nodes in the graph
-//
-USTRUCT()
-struct DATAFLOWEDITOR_API FAssetSchemaAction_Dataflow_PasteNode_DataflowEdNode : public FEdGraphSchemaAction
-{
-	GENERATED_USTRUCT_BODY();
-
-public:
-	FAssetSchemaAction_Dataflow_PasteNode_DataflowEdNode() : FEdGraphSchemaAction() {}
-
-	FAssetSchemaAction_Dataflow_PasteNode_DataflowEdNode(const FName& InType, FText InNodeCategory, FText InMenuDesc, FText InToolTip, FText InKeywords)
-		: FEdGraphSchemaAction(InNodeCategory, InMenuDesc, InToolTip, 0, InKeywords), NodeTypeName(InType) {}
-
-	static TSharedPtr<FAssetSchemaAction_Dataflow_PasteNode_DataflowEdNode> CreateAction(UEdGraph* ParentGraph, const FName& NodeTypeName);
-
 	virtual UEdGraphNode* PerformAction(class UEdGraph* ParentGraph, UEdGraphPin* FromPin, const FVector2D Location, bool bSelectNewNode = true) override;
 
 	FName NodeTypeName;
-	FName NodeName;
-	FString NodeProperties;
 };
 
 /**
