Dataflow : Add methods to add a dataflow subgraph transactionally ( Edit Asset utils )

diff --git a/Engine/Plugins/Experimental/Dataflow/Source/DataflowEditor/Private/Dataflow/DataflowAssetEditUtils.cpp b/Engine/Plugins/Experimental/Dataflow/Source/DataflowEditor/Private/Dataflow/DataflowAssetEditUtils.cpp
--- a/Engine/Plugins/Experimental/Dataflow/Source/DataflowEditor/Private/Dataflow/DataflowAssetEditUtils.cpp
+++ b/Engine/Plugins/Experimental/Dataflow/Source/DataflowEditor/Private/Dataflow/DataflowAssetEditUtils.cpp
@@ -9,6 +9,7 @@
 #include "Dataflow/DataflowNode.h"
 #include "Dataflow/DataflowEdNode.h"
 #include "Dataflow/DataflowSNode.h"
+#include "Dataflow/DataflowSubGraph.h"
 #include "PropertyBagDetails.h"
 #include "ScopedTransaction.h"
 #include "Settings/EditorStyleSettings.h"
@@ -225,6 +226,24 @@ namespace UE::Dataflow
 			return UniqueName;
 		}
 
+		/** Generate a Dataflow asset subgraph unique name from a BaseName */
+		static FName GenerateUniqueSubGraphName(const UDataflow& DataflowAsset, FName BaseName)
+		{
+			int32 Counter = 1;
+			FName UniqueName{ BaseName };
+			const FString BasenameStr{ BaseName.ToString() };
+
+			while (true)
+			{
+				if (DataflowAsset.FindSubGraphByName(UniqueName) == nullptr)
+				{
+					break; // found an available name exit 
+				}
+				UniqueName = FName(FString::Format(TEXT("{0}_{1}"), { BasenameStr, FString::FromInt(Counter++) }));
+			}
+			return UniqueName;
+		}
+
 		static UEdGraphPin* FindPin(const UDataflowEdNode* Node, const EEdGraphPinDirection Direction, const FName Name)
 		{
 			for (UEdGraphPin* Pin : Node->GetAllPins())
@@ -838,6 +857,41 @@ namespace UE::Dataflow
 		}
 		return NAME_None;
 	}
+
+	//////////////////////////////////////////////////////////////////////////////////////////////////////////
+	// 
+	// SUBGRAPHS API
+	// 
+	//////////////////////////////////////////////////////////////////////////////////////////////////////////
+	FName FEditAssetUtils::AddNewSubGraph(UDataflow* DataflowAsset, FName BaseName)
+	{
+		using namespace EditAssetUtils::Private;
+
+		const FText TransactionName{ LOCTEXT("AddNewDataflowSubGraph", "Add New Dataflow SubGraph") };
+
+		FName UniqueSubGraphName;
+
+		auto AddNewVariableInternal =
+			[&BaseName, &UniqueSubGraphName](UDataflow& DataflowAsset) -> EChangeResult
+			{
+				UniqueSubGraphName = GenerateUniqueSubGraphName(DataflowAsset, BaseName);
+				UDataflowSubGraph* NewSubGraph = NewObject<UDataflowSubGraph>(&DataflowAsset, UniqueSubGraphName);
+				ensure(NewSubGraph->GetFName() == UniqueSubGraphName);
+				NewSubGraph->Schema = UDataflowSchema::StaticClass();
+				NewSubGraph->SetFlags(RF_Transient);
+
+				DataflowAsset.AddSubGraph(NewSubGraph);
+
+				// todo(dataflow) notify of a subgraph update ? 
+				//FDataflowAssetDelegates::OnVariablesChanged.Broadcast(&DataflowAsset, UniqueVariableName);
+				return EChangeResult::Changed;
+			};
+
+		ChangeDataflowAssetWithTransaction(DataflowAsset, TransactionName, AddNewVariableInternal, NAME_None);
+
+		return UniqueSubGraphName;
+	}
+
 }
 
 #undef LOCTEXT_NAMESPACE
\ No newline at end of file
diff --git a/Engine/Plugins/Experimental/Dataflow/Source/DataflowEditor/Private/Dataflow/DataflowAssetEditUtils.h b/Engine/Plugins/Experimental/Dataflow/Source/DataflowEditor/Private/Dataflow/DataflowAssetEditUtils.h
--- a/Engine/Plugins/Experimental/Dataflow/Source/DataflowEditor/Private/Dataflow/DataflowAssetEditUtils.h
+++ b/Engine/Plugins/Experimental/Dataflow/Source/DataflowEditor/Private/Dataflow/DataflowAssetEditUtils.h
@@ -111,5 +111,16 @@ namespace UE::Dataflow
 		* Paste a variable from clipboard and return its name
 		*/
 		static FName PasteVariableFromClipboard(UDataflow* DataflowAsset);
+
+		/////////////////////////////////////////////////////////////////////////////////////////////////////////////////////////
+		//
+		// SubGraph API
+		//
+		/////////////////////////////////////////////////////////////////////////////////////////////////////////////////////////
+
+		/**
+		* Add a new subgraph to teh Dataflow Asset and return its name
+		*/
+		static FName AddNewSubGraph(UDataflow* DataflowAsset, FName BaseName);
 	};
 }
