Dataflow : add a SubGraph action in preparation for the SubGraph UI support
- also unify the naming convention for action methods

diff --git a/Engine/Plugins/Experimental/Dataflow/Source/DataflowEditor/Private/Dataflow/DataflowGraphSchemaAction.cpp b/Engine/Plugins/Experimental/Dataflow/Source/DataflowEditor/Private/Dataflow/DataflowGraphSchemaAction.cpp
--- a/Engine/Plugins/Experimental/Dataflow/Source/DataflowEditor/Private/Dataflow/DataflowGraphSchemaAction.cpp
+++ b/Engine/Plugins/Experimental/Dataflow/Source/DataflowEditor/Private/Dataflow/DataflowGraphSchemaAction.cpp
@@ -116,7 +116,7 @@ UDataflow* FEdGraphSchemaAction_DataflowVariable::GetDataflowAsset() const
 	return nullptr;
 }
 
-bool FEdGraphSchemaAction_DataflowVariable::CanRenameVariable(FText NewName) const
+bool FEdGraphSchemaAction_DataflowVariable::CanRenameItem(FText NewName) const
 {
 	if (const UDataflow* DataflowAsset = GetDataflowAsset())
 	{
@@ -131,7 +131,7 @@ bool FEdGraphSchemaAction_DataflowVariable::CanRenameVariable(FText NewName) con
 	return false;
 }
 
-void FEdGraphSchemaAction_DataflowVariable::RenameVariable(FText NewName)
+void FEdGraphSchemaAction_DataflowVariable::RenameItem(FText NewName)
 {
 	if (UDataflow* DataflowAsset = GetDataflowAsset())
 	{
@@ -148,12 +148,12 @@ void FEdGraphSchemaAction_DataflowVariable::RenameVariable(FText NewName)
 	}
 }
 
-void FEdGraphSchemaAction_DataflowVariable::CopyVariableToClipboard()
+void FEdGraphSchemaAction_DataflowVariable::CopyItemToClipboard()
 {
 	UE::Dataflow::FEditAssetUtils::CopyVariableToClipboard(GetDataflowAsset(), GetFullVariableName());
 }
 
-void FEdGraphSchemaAction_DataflowVariable::PasteVariableFromClipboard()
+void FEdGraphSchemaAction_DataflowVariable::PasteItemFromClipboard()
 {
 	UE::Dataflow::FEditAssetUtils::PasteVariableFromClipboard(GetDataflowAsset());
 }
@@ -166,12 +166,12 @@ void FEdGraphSchemaAction_DataflowVariable::SetVariableType(const FEdGraphPinTyp
 	}
 }
 
-void FEdGraphSchemaAction_DataflowVariable::DeleteVariable()
+void FEdGraphSchemaAction_DataflowVariable::DeleteItem()
 {
 	UE::Dataflow::FEditAssetUtils::DeleteVariable(GetDataflowAsset(), GetFullVariableName());
 }
 
-void FEdGraphSchemaAction_DataflowVariable::DuplicateVariable()
+void FEdGraphSchemaAction_DataflowVariable::DuplicateItem()
 {
 	UE::Dataflow::FEditAssetUtils::DuplicateVariable(GetDataflowAsset(), GetFullVariableName());
 }
@@ -215,4 +215,76 @@ TSharedRef<FGraphSchemaActionDragDropAction_DataflowVariable> FGraphSchemaAction
 	return Operation;
 }
 
+/////////////////////////////////////////////////////////////////////////////////////////////////////////////////////////
+
+FEdGraphSchemaAction_DataflowSubGraph::FEdGraphSchemaAction_DataflowSubGraph()
+	: FEdGraphSchemaAction()
+{
+}
+
+FEdGraphSchemaAction_DataflowSubGraph::FEdGraphSchemaAction_DataflowSubGraph(IDataflowInstanceInterface* InDataflowInstanceInterface, FName InSubgraphName)
+	: FEdGraphSchemaAction(
+		FText::GetEmpty(),
+		LOCTEXT("Description", "Description"),
+		LOCTEXT("Tooltip", "Tooltip"),
+		/*InGrouping*/0,
+		/*Keywords*/FText::GetEmpty(),
+		(int32)UE::Dataflow::ESchemaActionSectionID::SUBGRAPHS)
+	, SubGraphName(InSubgraphName)
+	, DataflowInstanceInterface(InDataflowInstanceInterface)
+{
+}
+
+UEdGraphNode* FEdGraphSchemaAction_DataflowSubGraph::PerformAction(class UEdGraph* ParentGraph, UEdGraphPin* FromPin, const FVector2D Location, bool bSelectNewNode)
+{
+	// NOT YET IMPLEMENTED
+	// Should drop a node that represents the subgraph
+	return nullptr;
+}
+
+bool FEdGraphSchemaAction_DataflowSubGraph::CanRenameItem(FText NewName) const
+{
+	if (const UDataflow* DataflowAsset = GetDataflowAsset())
+	{
+		const bool bSubGraphAlreadyExists = (nullptr != DataflowAsset->FindSubGraphByName(FName(NewName.ToString())));
+		return !bSubGraphAlreadyExists;
+	}
+	return false;
+}
+
+void FEdGraphSchemaAction_DataflowSubGraph::RenameItem(FText NewName)
+{
+	// NOT YET IMPLEMENTED
+}
+
+void FEdGraphSchemaAction_DataflowSubGraph::CopyItemToClipboard()
+{
+	// NOT YET IMPLEMENTED
+}
+
+void FEdGraphSchemaAction_DataflowSubGraph::PasteItemFromClipboard()
+{
+	// NOT YET IMPLEMENTED
+}
+
+void FEdGraphSchemaAction_DataflowSubGraph::DeleteItem()
+{
+	// NOT YET IMPLEMENTED
+}
+
+void FEdGraphSchemaAction_DataflowSubGraph::DuplicateItem()
+{
+	// NOT YET IMPLEMENTED
+}
+
+UDataflow* FEdGraphSchemaAction_DataflowSubGraph::GetDataflowAsset() const
+{
+	if (DataflowInstanceInterface)
+	{
+		return DataflowInstanceInterface->GetDataflowInstance().GetDataflowAsset();
+	}
+	return nullptr;
+}
+
+
 #undef LOCTEXT_NAMESPACE
\ No newline at end of file
diff --git a/Engine/Plugins/Experimental/Dataflow/Source/DataflowEditor/Private/Dataflow/DataflowMembersWidget.cpp b/Engine/Plugins/Experimental/Dataflow/Source/DataflowEditor/Private/Dataflow/DataflowMembersWidget.cpp
--- a/Engine/Plugins/Experimental/Dataflow/Source/DataflowEditor/Private/Dataflow/DataflowMembersWidget.cpp
+++ b/Engine/Plugins/Experimental/Dataflow/Source/DataflowEditor/Private/Dataflow/DataflowMembersWidget.cpp
@@ -58,32 +58,32 @@ namespace UE::Dataflow::Private
 
 		virtual void OnCopy(FEdGraphSchemaAction& InAction) const
 		{
-			if (TAction* VariableAction = AsTypedAction(InAction))
+			if (TAction* TypedAction = AsTypedAction(InAction))
 			{
-				VariableAction->CopyVariableToClipboard();
+				TypedAction->CopyItemToClipboard();
 			}
 		}
 		virtual void OnPaste(FEdGraphSchemaAction& InAction) const override
 		{
-			if (TAction* VariableAction = AsTypedAction(InAction))
+			if (TAction* TypedAction = AsTypedAction(InAction))
 			{
-				VariableAction->PasteVariableFromClipboard();
+				TypedAction->PasteItemFromClipboard();
 			}
 		}
 
 		virtual void OnDuplicate(FEdGraphSchemaAction& InAction) const override
 		{
-			if (TAction* VariableAction = AsTypedAction(InAction))
+			if (TAction* TypedAction = AsTypedAction(InAction))
 			{
-				VariableAction->DuplicateVariable();
+				TypedAction->DuplicateItem();
 			}
 		}
 
 		virtual void OnDelete(FEdGraphSchemaAction& InAction) const override
 		{
-			if (TAction* VariableAction = AsTypedAction(InAction))
+			if (TAction* TypedAction = AsTypedAction(InAction))
 			{
-				VariableAction->DeleteVariable();
+				TypedAction->DeleteItem();
 			}
 		}
 
diff --git a/Engine/Plugins/Experimental/Dataflow/Source/DataflowEditor/Private/Dataflow/DataflowVariablePaletteItem.cpp b/Engine/Plugins/Experimental/Dataflow/Source/DataflowEditor/Private/Dataflow/DataflowVariablePaletteItem.cpp
--- a/Engine/Plugins/Experimental/Dataflow/Source/DataflowEditor/Private/Dataflow/DataflowVariablePaletteItem.cpp
+++ b/Engine/Plugins/Experimental/Dataflow/Source/DataflowEditor/Private/Dataflow/DataflowVariablePaletteItem.cpp
@@ -155,14 +155,14 @@ TSharedRef<SWidget> SDataflowVariablePaletteItem::CreateTextSlotWidget(FCreateWi
 
 bool SDataflowVariablePaletteItem::OnNameTextVerifyChanged(const FText& InNewText, FText& OutErrorMessage)
 {
-	return VariableAction ? VariableAction->CanRenameVariable(InNewText) : false;
+	return VariableAction ? VariableAction->CanRenameItem(InNewText) : false;
 }
 
 void SDataflowVariablePaletteItem::OnNameTextCommitted(const FText& InNewText, ETextCommit::Type InTextCommit)
 {
 	if (VariableAction)
 	{
-		VariableAction->RenameVariable(InNewText);
+		VariableAction->RenameItem(InNewText);
 	}
 }
 
diff --git a/Engine/Plugins/Experimental/Dataflow/Source/DataflowEditor/Public/Dataflow/DataflowGraphSchemaAction.h b/Engine/Plugins/Experimental/Dataflow/Source/DataflowEditor/Public/Dataflow/DataflowGraphSchemaAction.h
--- a/Engine/Plugins/Experimental/Dataflow/Source/DataflowEditor/Public/Dataflow/DataflowGraphSchemaAction.h
+++ b/Engine/Plugins/Experimental/Dataflow/Source/DataflowEditor/Public/Dataflow/DataflowGraphSchemaAction.h
@@ -21,6 +21,7 @@ namespace UE::Dataflow
 	{
 		NONE = 0,
 		VARIABLES,
+		SUBGRAPHS,
 	};
 }
 
@@ -85,14 +86,14 @@ public:
 	virtual UEdGraphNode* PerformAction(class UEdGraph* ParentGraph, UEdGraphPin* FromPin, const FVector2D Location, bool bSelectNewNode = true) override;
 	// End of FEdGraphSchemaAction interface
 
-	bool CanRenameVariable(FText NewName) const;
-	void RenameVariable(FText NewName);
+	bool CanRenameItem(FText NewName) const;
+	void RenameItem(FText NewName);
 	void SetVariableType(const FEdGraphPinType& PinType);
 
-	void CopyVariableToClipboard();
-	void PasteVariableFromClipboard();
-	void DeleteVariable();
-	void DuplicateVariable();
+	void CopyItemToClipboard();
+	void PasteItemFromClipboard();
+	void DeleteItem();
+	void DuplicateItem();
 
 	UDataflow* GetDataflowAsset() const;
 
@@ -148,3 +149,55 @@ private:
 
 	TSharedPtr<FEdGraphSchemaAction_DataflowVariable> VariableAction;
 };
+
+
+////////////////////////////////////////////////////////////////////////////////////////////////////////////////////////////////
+
+/**
+* Action that refers to a dataflow subgraph
+* It is used to display the variables in the SDataflowMembersWidget
+*/
+USTRUCT()
+struct FEdGraphSchemaAction_DataflowSubGraph : public FEdGraphSchemaAction
+{
+	GENERATED_USTRUCT_BODY()
+
+public:
+	FEdGraphSchemaAction_DataflowSubGraph();
+
+	FEdGraphSchemaAction_DataflowSubGraph(IDataflowInstanceInterface* InDataflowInstanceInterface, FName SubGraphName);
+
+	// Simple type info
+	static FName StaticGetTypeId() { static FName Type("FEdGraphSchemaAction_DataflowSubGraph"); return Type; }
+	virtual FName GetTypeId() const override { return StaticGetTypeId(); }
+
+	const FName GetSubgraphName() const
+	{
+		return SubGraphName;
+	}
+
+	// FEdGraphSchemaAction interface
+	virtual bool IsAVariable() const { return false; }
+	virtual UEdGraphNode* PerformAction(class UEdGraph* ParentGraph, UEdGraphPin* FromPin, const FVector2D Location, bool bSelectNewNode = true) override;
+	// End of FEdGraphSchemaAction interface
+
+	bool CanRenameItem(FText NewName) const;
+	void RenameItem(FText NewName);
+
+	void CopyItemToClipboard();
+	void PasteItemFromClipboard();
+	void DeleteItem();
+	void DuplicateItem();
+
+	UDataflow* GetDataflowAsset() const;
+
+private:
+	/** Name of the subgraph */
+	FName SubGraphName;
+
+	/**
+	* interface from where to read the variable information
+	* this allow access to the dataflow asset as well as the override from the asset using this dataflow asset
+	*/
+	IDataflowInstanceInterface* DataflowInstanceInterface = nullptr;;
+};
