Adding a UDataflowMesh type to Dataflow Engine. Currently it just holds a UDynamicMesh and a set of UMaterialInterfaces, but in the future we could look at making it a container for other mesh types (e.g. Static Mesh) as well.
- Also adding render support and a few utility nodes (make, split, append)

diff --git a/Engine/Plugins/Experimental/Dataflow/Source/DataflowEditor/Private/Dataflow/DataflowEngineRendering.cpp b/Engine/Plugins/Experimental/Dataflow/Source/DataflowEditor/Private/Dataflow/DataflowEngineRendering.cpp
--- a/Engine/Plugins/Experimental/Dataflow/Source/DataflowEditor/Private/Dataflow/DataflowEngineRendering.cpp
+++ b/Engine/Plugins/Experimental/Dataflow/Source/DataflowEditor/Private/Dataflow/DataflowEngineRendering.cpp
@@ -5,6 +5,7 @@
 #include "Dataflow/DataflowConnectionTypes.h"
 #include "Dataflow/DataflowEditorModule.h"
 #include "Dataflow/DataflowRenderingFactory.h"
+#include "Dataflow/DataflowMesh.h"
 #include "DynamicMesh/DynamicMesh3.h"
 #include "Field/FieldSystemTypes.h"
 #include "GeometryCollection/GeometryCollectionUtility.h"
@@ -17,6 +18,116 @@
 
 namespace UE::Dataflow
 {
+	namespace Private
+	{
+		void RenderDynamicMesh(GeometryCollection::Facades::FRenderingFacade& RenderCollection, 
+			const FString& GeometryGroupName, 
+			const UE::Geometry::FDynamicMesh3& DynamicMesh, 
+			const TArray<TObjectPtr<UMaterialInterface>>& Materials)
+		{
+			const int32 NumVertices = DynamicMesh.VertexCount();
+			const int32 NumTriangles = DynamicMesh.TriangleCount();
+
+			if (NumVertices > 0 && NumTriangles > 0)
+			{
+				// This will contain the valid triangles only
+				TArray<FIntVector> Tris; Tris.Reserve(DynamicMesh.TriangleCount());
+				TArray<int32> MaterialIDs; MaterialIDs.Reserve(DynamicMesh.TriangleCount());
+				const UE::Geometry::FDynamicMeshMaterialAttribute* const MaterialAttribute = DynamicMesh.Attributes()->GetMaterialID();
+
+				// DynamicMesh.TrianglesItr() returns the valid triangles only
+				for (const int32 TriangleID : DynamicMesh.TriangleIndicesItr())
+				{
+					const UE::Geometry::FIndex3i Tri = DynamicMesh.GetTriangle(TriangleID);
+					Tris.Add(FIntVector(Tri.A, Tri.B, Tri.C));
+					if (MaterialAttribute)
+					{
+						MaterialIDs.Add(MaterialAttribute->GetValue(TriangleID));
+					}
+				}
+
+				// This will contain all the vertices (invalid ones too)
+				// Otherwise the IDs need to be remaped
+				TArray<FVector3f> Vertices; Vertices.AddZeroed(DynamicMesh.MaxVertexID());
+
+				// DynamicMesh.VertexIndicesItr() returns the valid vertices only
+				for (int32 VertexID : DynamicMesh.VertexIndicesItr())
+				{
+					Vertices[VertexID] = (FVector3f)DynamicMesh.GetVertex(VertexID);
+				}
+
+				// Add VertexNormal and VertexColor
+				TArray<FVector3f> VertexNormals; VertexNormals.AddUninitialized(Vertices.Num());
+				TArray<FLinearColor> VertexColors; VertexColors.AddUninitialized(Vertices.Num());
+				for (int32 VertexIdx = 0; VertexIdx < VertexNormals.Num(); ++VertexIdx)
+				{
+					// TODO: Get the normal from FDynamicMesh3
+					VertexNormals[VertexIdx] = FVector3f(0.f);
+					VertexColors[VertexIdx] = FLinearColor(FDataflowEditorModule::SurfaceColor);
+				}
+
+				const int32 NumUVLayers = DynamicMesh.Attributes()->NumUVLayers();
+
+				const int32 GeometryIndex = RenderCollection.StartGeometryGroup(GeometryGroupName);
+
+				if (MaterialIDs.Num() != Tris.Num() || Materials.IsEmpty() || NumUVLayers == 0)
+				{
+					// No materials.
+					RenderCollection.AddSurface(MoveTemp(Vertices), MoveTemp(Tris), MoveTemp(VertexNormals), MoveTemp(VertexColors));
+				}
+				else
+				{
+					// GetUVs and materials.
+					// The render facade currently assumes there are no internal uv seams (i.e., each vertex has one UV, instead of wedges). 
+					// This will be the case for meshes that come from SKMs/SMs/Cloth
+					TArray<const UE::Geometry::FDynamicMeshUVOverlay*> UVLayers;
+					UVLayers.SetNumUninitialized(NumUVLayers);
+					for (int32 Layer = 0; Layer < NumUVLayers; ++Layer)
+					{
+						UVLayers[Layer] = DynamicMesh.Attributes()->GetUVLayer(Layer);
+					}
+
+					TArray<FVector2f> DefaultUVs;
+					DefaultUVs.SetNumZeroed(NumUVLayers);
+
+					TArray<TArray<FVector2f>> UVs;
+					UVs.Init(DefaultUVs, Vertices.Num()); // RenderFacade->DynamicMesh code expects all vertices to have full UV sets.
+
+					for (const int32 TriangleIndex : DynamicMesh.TriangleIndicesItr())
+					{
+						const UE::Geometry::FIndex3i Tri = DynamicMesh.GetTriangle(TriangleIndex);
+
+						for (int32 TriangleVertexIndex = 0; TriangleVertexIndex < 3; ++TriangleVertexIndex)
+						{
+							const int32 VertexIndex = Tri[TriangleVertexIndex];
+
+							for (int32 Layer = 0; Layer < NumUVLayers; ++Layer)
+							{
+								if (UVLayers[Layer])
+								{
+									UVs[VertexIndex][Layer] = UVLayers[Layer]->GetElementAtVertex(TriangleIndex, VertexIndex);
+								}
+							}
+						}
+					}
+
+					TArray<FString> MaterialPaths;
+					MaterialPaths.Reserve(Materials.Num());
+					for (const TObjectPtr<UMaterialInterface>& Material : Materials)
+					{
+						MaterialPaths.Add(Material ? Material->GetPathName() : "");
+					}
+
+					RenderCollection.AddSurface(MoveTemp(Vertices), MoveTemp(Tris), MoveTemp(VertexNormals), MoveTemp(VertexColors), MoveTemp(UVs), MoveTemp(MaterialIDs), MoveTemp(MaterialPaths));
+
+				}
+
+				RenderCollection.EndGeometryGroup(GeometryIndex);
+			}
+		}
+
+	}
+
 	void RenderBasicGeometryCollection(GeometryCollection::Facades::FRenderingFacade& RenderCollection, const UE::Dataflow::FGraphRenderingState& State, TArray<FLinearColor>* VertexColorOverride = nullptr)
 	{
 		FManagedArrayCollection Default;
@@ -422,119 +533,51 @@ namespace UE::Dataflow
 				{
 					const UE::Geometry::FDynamicMesh3& DynamicMesh = Mesh->GetMeshRef();
 
-					const int32 NumVertices = DynamicMesh.VertexCount();
-					const int32 NumTriangles = DynamicMesh.TriangleCount();
-
-					if (NumVertices > 0 && NumTriangles > 0)
+					TArray<TObjectPtr<UMaterialInterface>> Materials;
+					if (State.GetRenderOutputs().IsValidIndex(1))
 					{
-						// This will contain the valid triangles only
-						TArray<FIntVector> Tris; Tris.Reserve(DynamicMesh.TriangleCount());
-						TArray<int32> MaterialIDs; MaterialIDs.Reserve(DynamicMesh.TriangleCount());
-						const UE::Geometry::FDynamicMeshMaterialAttribute* const MaterialAttribute = DynamicMesh.Attributes()->GetMaterialID();
-
-						// DynamicMesh.TrianglesItr() returns the valid triangles only
-						for (const int32 TriangleID : DynamicMesh.TriangleIndicesItr())
-						{
-							const UE::Geometry::FIndex3i Tri = DynamicMesh.GetTriangle(TriangleID);
-							Tris.Add(FIntVector(Tri.A, Tri.B, Tri.C));
-							if (MaterialAttribute)
-							{
-								MaterialIDs.Add(MaterialAttribute->GetValue(TriangleID));
-							}
-						}
-
-						// This will contain all the vertices (invalid ones too)
-						// Otherwise the IDs need to be remaped
-						TArray<FVector3f> Vertices; Vertices.AddZeroed(DynamicMesh.MaxVertexID());
-
-						// DynamicMesh.VertexIndicesItr() returns the valid vertices only
-						for (int32 VertexID : DynamicMesh.VertexIndicesItr())
-						{
-							Vertices[VertexID] = (FVector3f)DynamicMesh.GetVertex(VertexID);
-						}
-
-						// Add VertexNormal and VertexColor
-						TArray<FVector3f> VertexNormals; VertexNormals.AddUninitialized(Vertices.Num());
-						TArray<FLinearColor> VertexColors; VertexColors.AddUninitialized(Vertices.Num());
-						for (int32 VertexIdx = 0; VertexIdx < VertexNormals.Num(); ++VertexIdx)
-						{
-							// TODO: Get the normal from FDynamicMesh3
-							VertexNormals[VertexIdx] = FVector3f(0.f);
-							VertexColors[VertexIdx] = FLinearColor(FDataflowEditorModule::SurfaceColor);
-						}
-
-						// See if we have materials to render.
-						TArray<TObjectPtr<UMaterialInterface>> Materials;
-						if (State.GetRenderOutputs().IsValidIndex(1))
-						{
-							TArray<TObjectPtr<UMaterialInterface>> MaterialsDefault;
-							Materials = State.GetValue<TArray<TObjectPtr<UMaterialInterface>>>(State.GetRenderOutputs()[1], MaterialsDefault);
-						}
-
-						const int32 NumUVLayers = DynamicMesh.Attributes()->NumUVLayers();
-
-						const int32 GeometryIndex = RenderCollection.StartGeometryGroup(State.GetGuid().ToString());
-
-						if (MaterialIDs.Num() != Tris.Num() || Materials.IsEmpty() || NumUVLayers == 0)
-						{
-							// No materials.
-							RenderCollection.AddSurface(MoveTemp(Vertices), MoveTemp(Tris), MoveTemp(VertexNormals), MoveTemp(VertexColors));
-						}
-						else
-						{
-							// GetUVs and materials.
-							// The render facade currently assumes there are no internal uv seams (i.e., each vertex has one UV, instead of wedges). 
-							// This will be the case for meshes that come from SKMs/SMs/Cloth
-							TArray<const UE::Geometry::FDynamicMeshUVOverlay*> UVLayers;
-							UVLayers.SetNumUninitialized(NumUVLayers);
-							for (int32 Layer = 0; Layer < NumUVLayers; ++Layer)
-							{
-								UVLayers[Layer] = DynamicMesh.Attributes()->GetUVLayer(Layer);
-							}
-
-							TArray<FVector2f> DefaultUVs;
-							DefaultUVs.SetNumZeroed(NumUVLayers);
-
-							TArray<TArray<FVector2f>> UVs;
-							UVs.Init(DefaultUVs, Vertices.Num()); // RenderFacade->DynamicMesh code expects all vertices to have full UV sets.
+						TArray<TObjectPtr<UMaterialInterface>> MaterialsDefault;
+						Materials = State.GetValue<TArray<TObjectPtr<UMaterialInterface>>>(State.GetRenderOutputs()[1], MaterialsDefault);
+					}
 
-							for (const int32 TriangleIndex : DynamicMesh.TriangleIndicesItr())
-							{
-								const UE::Geometry::FIndex3i Tri = DynamicMesh.GetTriangle(TriangleIndex);
+					Private::RenderDynamicMesh(RenderCollection, State.GetGuid().ToString(), DynamicMesh, Materials);
+				}
+			}
+		}
+	};
 
-								for (int32 TriangleVertexIndex = 0; TriangleVertexIndex < 3; ++TriangleVertexIndex)
-								{
-									const int32 VertexIndex = Tri[TriangleVertexIndex];
-
-									for (int32 Layer = 0; Layer < NumUVLayers; ++Layer)
-									{
-										if (UVLayers[Layer])
-										{
-											UVs[VertexIndex][Layer] = UVLayers[Layer]->GetElementAtVertex(TriangleIndex, VertexIndex);
-										}
-									}
-								}
-							}
 
-							TArray<FString> MaterialPaths;
-							MaterialPaths.Reserve(Materials.Num());
-							for (const TObjectPtr<UMaterialInterface>& Material : Materials)
-							{
-								MaterialPaths.Add(Material ? Material->GetPathName() : "");
-							}
+	class FDataflowMeshSurfaceRenderCallbacks : public FRenderingFactory::ICallbackInterface
+	{
+		virtual UE::Dataflow::FRenderKey GetRenderKey() const override
+		{
+			return { "SurfaceRender", FName("UDataflowMesh") };
+		}
 
-							RenderCollection.AddSurface(MoveTemp(Vertices), MoveTemp(Tris), MoveTemp(VertexNormals), MoveTemp(VertexColors), MoveTemp(UVs), MoveTemp(MaterialIDs), MoveTemp(MaterialPaths));
+		virtual bool CanRender(const UE::Dataflow::IDataflowConstructionViewMode& ViewMode) const override
+		{
+			return (ViewMode.GetName() == FDataflowConstruction3DViewMode::Name);
+		}
 
-						}
+		virtual void Render(GeometryCollection::Facades::FRenderingFacade& RenderCollection, const FGraphRenderingState& State) override
+		{
+			if (State.GetRenderOutputs().Num())
+			{
+				FName PrimaryOutput = State.GetRenderOutputs()[0];
 
-						RenderCollection.EndGeometryGroup(GeometryIndex);
+				TObjectPtr<UDataflowMesh> Default;
+				if (const TObjectPtr<UDataflowMesh> Mesh = State.GetValue(PrimaryOutput, Default))
+				{
+					if (Mesh->DynamicMesh)
+					{
+						const UE::Geometry::FDynamicMesh3& DynamicMesh = Mesh->DynamicMesh->GetMeshRef();
+						Private::RenderDynamicMesh(RenderCollection, State.GetGuid().ToString(), DynamicMesh, Mesh->Materials);
 					}
 				}
 			}
 		}
 	};
 
-
 	class FBoxSurfaceRenderCallbacks : public FRenderingFactory::ICallbackInterface
 	{
 		virtual UE::Dataflow::FRenderKey GetRenderKey() const override
@@ -824,6 +867,7 @@ namespace UE::Dataflow
 		FRenderingFactory::GetInstance()->RegisterCallbacks(MakeUnique<FGeometryCollectionSurfaceRenderCallbacks>());
 		FRenderingFactory::GetInstance()->RegisterCallbacks(MakeUnique<FGeometryCollectionSurfaceWeightsRenderCallbacks>());
 		FRenderingFactory::GetInstance()->RegisterCallbacks(MakeUnique<FDynamicMesh3SurfaceRenderCallbacks>());
+		FRenderingFactory::GetInstance()->RegisterCallbacks(MakeUnique<FDataflowMeshSurfaceRenderCallbacks>());
 		FRenderingFactory::GetInstance()->RegisterCallbacks(MakeUnique<FBoxSurfaceRenderCallbacks>());
 		FRenderingFactory::GetInstance()->RegisterCallbacks(MakeUnique<FSphereSurfaceRenderCallbacks>());
 		FRenderingFactory::GetInstance()->RegisterCallbacks(MakeUnique<FPointRenderCallbacks>());
diff --git a/Engine/Plugins/Experimental/GeometryCollectionPlugin/Source/GeometryCollectionNodes/Private/Dataflow/GeometryCollectionMeshNodes.cpp b/Engine/Plugins/Experimental/GeometryCollectionPlugin/Source/GeometryCollectionNodes/Private/Dataflow/GeometryCollectionMeshNodes.cpp
--- a/Engine/Plugins/Experimental/GeometryCollectionPlugin/Source/GeometryCollectionNodes/Private/Dataflow/GeometryCollectionMeshNodes.cpp
+++ b/Engine/Plugins/Experimental/GeometryCollectionPlugin/Source/GeometryCollectionNodes/Private/Dataflow/GeometryCollectionMeshNodes.cpp
@@ -51,6 +51,9 @@ namespace UE::Dataflow
 		DATAFLOW_NODE_REGISTER_CREATION_FACTORY(FCollectionToMeshDataflowNode);
 		DATAFLOW_NODE_REGISTER_CREATION_FACTORY(FStaticMeshToMeshDataflowNode);
 		DATAFLOW_NODE_REGISTER_CREATION_FACTORY(FMeshAppendDataflowNode);
+		DATAFLOW_NODE_REGISTER_CREATION_FACTORY(FDataflowMeshAppendDataflowNode);
+		DATAFLOW_NODE_REGISTER_CREATION_FACTORY(FMakeDataflowMeshDataflowNode);
+		DATAFLOW_NODE_REGISTER_CREATION_FACTORY(FSplitDataflowMeshDataflowNode);
 		DATAFLOW_NODE_REGISTER_CREATION_FACTORY(FMeshBooleanDataflowNode);
 		DATAFLOW_NODE_REGISTER_CREATION_FACTORY(FMeshCopyToPointsDataflowNode);
 		DATAFLOW_NODE_REGISTER_CREATION_FACTORY(FGetMeshDataDataflowNode);
@@ -319,6 +322,114 @@ void FMeshAppendDataflowNode::Evaluate(UE::Dataflow::FContext& Context, const FD
 }
 
 
+void FDataflowMeshAppendDataflowNode::Evaluate(UE::Dataflow::FContext& Context, const FDataflowOutput* Out) const
+{
+	if (Out->IsA(&Mesh))
+	{
+		TObjectPtr<UDataflowMesh> NewMesh = NewObject<UDataflowMesh>();
+
+		if (const UDataflowMesh* const DataflowMesh1 = GetValue(Context, &Mesh))
+		{
+			if (const UDataflowMesh* const DataflowMesh2 = GetValue(Context, &AppendMesh))
+			{
+				if (const UDynamicMesh* const UDynanicMesh1 = DataflowMesh1->DynamicMesh)
+				{
+					if (const UDynamicMesh* const UDynanicMesh2 = DataflowMesh2->DynamicMesh)
+					{
+						if (const UE::Geometry::FDynamicMesh3* const DynMesh1 = UDynanicMesh1->GetMeshPtr())
+						{
+							if (const UE::Geometry::FDynamicMesh3* const DynMesh2 = UDynanicMesh2->GetMeshPtr())
+							{
+								if (DynMesh1->VertexCount() > 0 && DynMesh2->VertexCount() > 0)
+								{
+									UE::Geometry::FDynamicMesh3 ResultDynMesh;
+									UE::Geometry::FDynamicMeshEditor MeshEditor(&ResultDynMesh);
+									ResultDynMesh.EnableAttributes();
+									ResultDynMesh.Attributes()->EnableMaterialID();
+
+									UE::Geometry::FMeshIndexMappings IndexMaps1;
+									MeshEditor.AppendMesh(DynMesh1, IndexMaps1);
+
+									UE::Geometry::FMeshIndexMappings IndexMaps2;
+									MeshEditor.AppendMesh(DynMesh2, IndexMaps2);
+
+									// Reindex material IDs
+									if (DynMesh1->HasAttributes() && DynMesh1->Attributes()->HasMaterialID() && DynMesh2->HasAttributes() && DynMesh2->Attributes()->HasMaterialID())
+									{
+										const int32 MaterialIDOffset = DataflowMesh1->Materials.Num();
+
+										for (const int32 Mesh2TriangleIndex : DynMesh2->TriangleIndicesItr())
+										{
+											int32 InputMaterialID;
+											DynMesh2->Attributes()->GetMaterialID()->GetValue(Mesh2TriangleIndex, &InputMaterialID);
+
+											const int32 NewTriangleIndex = IndexMaps2.GetNewTriangle(Mesh2TriangleIndex);
+											ResultDynMesh.Attributes()->GetMaterialID()->SetValue(NewTriangleIndex, MaterialIDOffset + InputMaterialID);
+										}
+									}
+
+									NewMesh->DynamicMesh = NewObject<UDynamicMesh>();
+									NewMesh->DynamicMesh->SetMesh(MoveTemp(ResultDynMesh));
+								}
+							}
+						}
+					}
+				}	// end if UDynanicMesh1
+
+
+				// Materials
+				NewMesh->Materials.Append(DataflowMesh1->Materials);
+				NewMesh->Materials.Append(DataflowMesh2->Materials);
+			}
+		}
+
+		SetValue(Context, NewMesh, &Mesh);
+	}
+}
+
+void FMakeDataflowMeshDataflowNode::Evaluate(UE::Dataflow::FContext& Context, const FDataflowOutput* Out) const
+{
+	if (Out->IsA(&Mesh))
+	{
+		TObjectPtr<UDataflowMesh> NewMesh = NewObject<UDataflowMesh>();
+		
+		if (UDynamicMesh* const InDynamicMesh = GetValue(Context, &InMesh))
+		{
+			NewMesh->DynamicMesh = InDynamicMesh;
+		}
+
+		const TArray<TObjectPtr<UMaterialInterface>> MaterialArray = GetValue(Context, &InMaterials);
+		NewMesh->Materials = MaterialArray;
+
+		SetValue(Context, NewMesh, &Mesh);
+	}
+}
+
+void FSplitDataflowMeshDataflowNode::Evaluate(UE::Dataflow::FContext& Context, const FDataflowOutput* Out) const
+{
+	TObjectPtr<UDynamicMesh> NewMesh;
+	TArray<TObjectPtr<UMaterialInterface>> Materials;
+
+	if (Out->IsA(&Mesh))
+	{
+		if (UDataflowMesh* const InDataflowMesh = GetValue(Context, &InMesh))
+		{
+			NewMesh = InDataflowMesh->DynamicMesh;
+		}
+		SetValue(Context, NewMesh, &Mesh);
+	}
+	else if (Out->IsA(&MaterialArray))
+	{
+		if (UDataflowMesh* const InDataflowMesh = GetValue(Context, &InMesh))
+		{
+			Materials = InDataflowMesh->Materials;
+		}
+		SetValue(Context, Materials, &MaterialArray);
+	}
+	
+}
+
+
 void FMeshBooleanDataflowNode::Evaluate(UE::Dataflow::FContext& Context, const FDataflowOutput* Out) const
 {
 	if (Out->IsA<TObjectPtr<UDynamicMesh>>(&Mesh))
diff --git a/Engine/Plugins/Experimental/GeometryCollectionPlugin/Source/GeometryCollectionNodes/Public/Dataflow/GeometryCollectionMeshNodes.h b/Engine/Plugins/Experimental/GeometryCollectionPlugin/Source/GeometryCollectionNodes/Public/Dataflow/GeometryCollectionMeshNodes.h
--- a/Engine/Plugins/Experimental/GeometryCollectionPlugin/Source/GeometryCollectionNodes/Public/Dataflow/GeometryCollectionMeshNodes.h
+++ b/Engine/Plugins/Experimental/GeometryCollectionPlugin/Source/GeometryCollectionNodes/Public/Dataflow/GeometryCollectionMeshNodes.h
@@ -4,6 +4,7 @@
 
 #include "CoreMinimal.h"
 #include "Dataflow/DataflowEngine.h"
+#include "Dataflow/DataflowMesh.h"
 #include "GeometryCollection/ManagedArrayCollection.h"
 #include "Dataflow/DataflowSelection.h"
 #include "DynamicMesh/DynamicMesh3.h"
@@ -476,6 +477,113 @@ public:
 
 };
 
+/**
+ *
+ * Appends Dataflow meshes
+ *
+ */
+USTRUCT(meta = (DataflowGeometryCollection))
+struct FDataflowMeshAppendDataflowNode : public FDataflowNode
+{
+	GENERATED_USTRUCT_BODY()
+	DATAFLOW_NODE_DEFINE_INTERNAL(FDataflowMeshAppendDataflowNode, "DataflowMeshAppend", "Mesh|Utilities", "")
+	DATAFLOW_NODE_RENDER_TYPE("SurfaceRender", FName("UDataflowMesh"), "Mesh" )
+
+public:
+
+	FDataflowMeshAppendDataflowNode(const UE::Dataflow::FNodeParameters& InParam, FGuid InGuid = FGuid::NewGuid())
+		: FDataflowNode(InParam, InGuid)
+	{
+		RegisterInputConnection(&Mesh);
+		RegisterOutputConnection(&Mesh, &Mesh);
+		RegisterInputConnection(&AppendMesh);
+	}
+
+private:
+
+	/** Mesh input/output */
+	UPROPERTY(meta = (DataflowInput, DataflowOutput, DataflowPassthrough = "Mesh", DataflowIntrinsic))
+	TObjectPtr<UDataflowMesh> Mesh;
+
+	/** Mesh to append  */
+	UPROPERTY(meta = (DataflowInput, DataflowIntrinsic))
+	TObjectPtr<const UDataflowMesh> AppendMesh;
+
+	virtual void Evaluate(UE::Dataflow::FContext& Context, const FDataflowOutput* Out) const override;
+
+};
+
+
+USTRUCT(meta = (DataflowGeometryCollection))
+struct FMakeDataflowMeshDataflowNode : public FDataflowNode
+{
+	GENERATED_USTRUCT_BODY()
+	DATAFLOW_NODE_DEFINE_INTERNAL(FMakeDataflowMeshDataflowNode, "MakeDataflowMesh", "Mesh|Utilities", "")
+	DATAFLOW_NODE_RENDER_TYPE("SurfaceRender", FName("UDataflowMesh"), "Mesh")
+
+public:
+
+	FMakeDataflowMeshDataflowNode(const UE::Dataflow::FNodeParameters& InParam, FGuid InGuid = FGuid::NewGuid())
+		: FDataflowNode(InParam, InGuid)
+	{
+		RegisterInputConnection(&InMesh);
+		RegisterInputConnection(&InMaterials);
+		RegisterOutputConnection(&Mesh);
+	}
+
+private:
+
+	/** DynamicMesh input */
+	UPROPERTY(meta = (DataflowInput, DataflowIntrinsic))
+	TObjectPtr<UDynamicMesh> InMesh;
+
+	/** Materials input */
+	UPROPERTY(meta = (DataflowInput, DataflowIntrinsic))
+	TArray<TObjectPtr<UMaterialInterface>> InMaterials;
+
+	/** DataflowMesh output */
+	UPROPERTY(meta = (DataflowOutput))
+	TObjectPtr<UDataflowMesh> Mesh;
+
+	virtual void Evaluate(UE::Dataflow::FContext& Context, const FDataflowOutput* Out) const override;
+};
+
+
+USTRUCT(meta = (DataflowGeometryCollection))
+struct FSplitDataflowMeshDataflowNode : public FDataflowNode
+{
+	GENERATED_USTRUCT_BODY()
+	DATAFLOW_NODE_DEFINE_INTERNAL(FSplitDataflowMeshDataflowNode, "SplitDataflowMesh", "Mesh|Utilities", "")
+	DATAFLOW_NODE_RENDER_TYPE("SurfaceRender", FName("FDynamicMesh3"), { "Mesh", "MaterialArray" })
+
+public:
+
+	FSplitDataflowMeshDataflowNode(const UE::Dataflow::FNodeParameters& InParam, FGuid InGuid = FGuid::NewGuid())
+		: FDataflowNode(InParam, InGuid)
+	{
+		RegisterInputConnection(&InMesh);
+		RegisterOutputConnection(&Mesh);
+		RegisterOutputConnection(&MaterialArray);
+	}
+
+private:
+
+	/** DataflowMesh input */
+	UPROPERTY(meta = (DataflowInput, DataflowIntrinsic))
+	TObjectPtr<UDataflowMesh> InMesh;
+
+	/** DyanmicMesh output */
+	UPROPERTY(meta = (DataflowOutput))
+	TObjectPtr<UDynamicMesh> Mesh;
+
+	/** Materials output */
+	UPROPERTY(meta = (DataflowOutput))
+	TArray<TObjectPtr<UMaterialInterface>> MaterialArray;
+
+	virtual void Evaluate(UE::Dataflow::FContext& Context, const FDataflowOutput* Out) const override;
+};
+
+
 
 UENUM(BlueprintType)
 enum class EMeshBooleanOperationEnum : uint8
diff --git a/Engine/Source/Runtime/Experimental/Dataflow/Engine/DataflowEngine.Build.cs b/Engine/Source/Runtime/Experimental/Dataflow/Engine/DataflowEngine.Build.cs
--- a/Engine/Source/Runtime/Experimental/Dataflow/Engine/DataflowEngine.Build.cs
+++ b/Engine/Source/Runtime/Experimental/Dataflow/Engine/DataflowEngine.Build.cs
@@ -16,7 +16,8 @@ namespace UnrealBuildTool.Rules
 					"Engine",
 					"DataflowCore",
 					"Chaos",
-					"ImageCore"
+					"ImageCore",
+					"GeometryFramework"		// UDynamicMesh
 				}
 			);
 		}
diff --git a/Engine/Source/Runtime/Experimental/Dataflow/Engine/Private/Dataflow/DataflowMesh.cpp b/Engine/Source/Runtime/Experimental/Dataflow/Engine/Private/Dataflow/DataflowMesh.cpp
new file mode 100644
--- /dev/null
+++ b/Engine/Source/Runtime/Experimental/Dataflow/Engine/Private/Dataflow/DataflowMesh.cpp
@@ -0,0 +1,9 @@
+// Copyright Epic Games, Inc. All Rights Reserved.
+
+#include "Dataflow/DataflowMesh.h"
+#include "Materials/MaterialInterface.h"
+
+UDataflowMesh::UDataflowMesh(const FObjectInitializer& ObjectInitializer)
+	: Super(ObjectInitializer)
+{
+}
diff --git a/Engine/Source/Runtime/Experimental/Dataflow/Engine/Public/Dataflow/DataflowMesh.h b/Engine/Source/Runtime/Experimental/Dataflow/Engine/Public/Dataflow/DataflowMesh.h
new file mode 100644
--- /dev/null
+++ b/Engine/Source/Runtime/Experimental/Dataflow/Engine/Public/Dataflow/DataflowMesh.h
@@ -0,0 +1,26 @@
+// Copyright Epic Games, Inc. All Rights Reserved.
+
+#pragma once
+
+#include "UObject/Object.h"
+#include "DataflowMesh.generated.h"
+
+class UDynamicMesh;
+class UMaterialInterface;
+
+/** 
+* Containter for storing mesh and material information in Dataflow graph
+*/
+UCLASS()
+class DATAFLOWENGINE_API UDataflowMesh : public UObject
+{
+	GENERATED_UCLASS_BODY()
+
+public:
+
+	UPROPERTY()
+	TObjectPtr<UDynamicMesh> DynamicMesh;
+
+	UPROPERTY()
+	TArray<TObjectPtr<UMaterialInterface>> Materials;
+};
