Dataflow : Add the ability for the toolkit to spawn tabs dynamically for SubGraphs

diff --git a/Engine/Plugins/Experimental/Dataflow/Source/DataflowEditor/Private/Dataflow/DataflowEditorSubGraphTabSummoner.cpp b/Engine/Plugins/Experimental/Dataflow/Source/DataflowEditor/Private/Dataflow/DataflowEditorSubGraphTabSummoner.cpp
new file mode 100644
--- /dev/null
+++ b/Engine/Plugins/Experimental/Dataflow/Source/DataflowEditor/Private/Dataflow/DataflowEditorSubGraphTabSummoner.cpp
@@ -0,0 +1,97 @@
+// Copyright Epic Games, Inc. All Rights Reserved.
+
+#include "Dataflow/DataflowEditorSubGraphTabSummoner.h"
+
+#include "Dataflow/DataflowGraphEditor.h"
+#include "Dataflow/DataflowEditorToolkit.h"
+#include "Widgets/Docking/SDockTab.h"
+
+#include "WorkflowOrientedApp/WorkflowTabManager.h"
+
+FDataflowEditorSubGraphTabSummoner::FDataflowEditorSubGraphTabSummoner(TSharedPtr<FDataflowEditorToolkit> InEditorToolkitPtr, FOnCreateGraphEditorWidget CreateGraphEditorWidgetCallback)
+	: FDocumentTabFactoryForObjects<UDataflowSubGraph>("DataflowEditor_SubGraphTab", InEditorToolkitPtr)
+	, EditorToolkitWeakPtr(InEditorToolkitPtr)
+	, OnCreateGraphEditorWidget(CreateGraphEditorWidgetCallback)
+{
+}
+
+TAttribute<FText> FDataflowEditorSubGraphTabSummoner::ConstructTabNameForObject(UDataflowSubGraph* DocumentID) const
+{
+	return TAttribute<FText>::Create(TAttribute<FText>::FGetter::CreateStatic(&FDataflowEditorSubGraphTabSummoner::GetGraphDisplayName, (const UEdGraph*)DocumentID));
+}
+
+FText FDataflowEditorSubGraphTabSummoner::GetGraphDisplayName(const UEdGraph* Graph)
+{
+	return FText::FromString(Graph->GetFName().ToString());
+}
+
+TSharedPtr<SDataflowGraphEditor> FDataflowEditorSubGraphTabSummoner::GetDataflowEditorFromTab(TSharedPtr<SDockTab> Tab) const
+{
+	TSharedRef<SGraphEditor> GraphEditor = StaticCastSharedRef<SGraphEditor>(Tab->GetContent());
+	return StaticCastSharedPtr<SDataflowGraphEditor>(GraphEditor.ToSharedPtr());
+}
+
+void FDataflowEditorSubGraphTabSummoner::OnTabActivated(TSharedPtr<SDockTab> Tab) const
+{
+	if (TSharedPtr<SDataflowGraphEditor> DataflowEditor = GetDataflowEditorFromTab(Tab))
+	{
+		if (TSharedPtr<FDataflowEditorToolkit> Toolkit = EditorToolkitWeakPtr.Pin())
+		{
+			Toolkit->SetSubGraphTabActiveState(DataflowEditor, true);
+		}
+	}
+}
+
+// Called when a tab created from this factory is brought to the foreground
+void FDataflowEditorSubGraphTabSummoner::OnTabForegrounded(TSharedPtr<SDockTab> Tab) const
+{
+	if (TSharedPtr<SDataflowGraphEditor> DataflowEditor = GetDataflowEditorFromTab(Tab))
+	{
+		if (TSharedPtr<FDataflowEditorToolkit> Toolkit = EditorToolkitWeakPtr.Pin())
+		{
+			Toolkit->SetSubGraphTabActiveState(DataflowEditor, true);
+		}
+	}
+}
+
+void FDataflowEditorSubGraphTabSummoner::OnTabBackgrounded(TSharedPtr<SDockTab> Tab) const
+{
+	if (TSharedPtr<SDataflowGraphEditor> DataflowEditor = GetDataflowEditorFromTab(Tab))
+	{
+		if (TSharedPtr<FDataflowEditorToolkit> Toolkit = EditorToolkitWeakPtr.Pin())
+		{
+			Toolkit->SetSubGraphTabActiveState(DataflowEditor, false);
+		}
+	}
+
+}
+
+void FDataflowEditorSubGraphTabSummoner::OnTabRefreshed(TSharedPtr<SDockTab> Tab) const
+{
+	if (TSharedPtr<SDataflowGraphEditor> DataflowEditor = GetDataflowEditorFromTab(Tab))
+	{
+		DataflowEditor->NotifyGraphChanged();
+	}
+}
+
+void FDataflowEditorSubGraphTabSummoner::SaveState(TSharedPtr<SDockTab> Tab, TSharedPtr<FTabPayload> Payload) const
+{
+	// NOT YET IMPLEMENTED
+}
+
+TSharedRef<SWidget> FDataflowEditorSubGraphTabSummoner::CreateTabBodyForObject(const FWorkflowTabSpawnInfo& Info, UDataflowSubGraph* DocumentID) const
+{
+	check(Info.TabInfo.IsValid());
+	return OnCreateGraphEditorWidget.Execute(Info.TabInfo.ToSharedRef(), DocumentID);
+}
+
+const FSlateBrush* FDataflowEditorSubGraphTabSummoner::GetTabIconForObject(const FWorkflowTabSpawnInfo& Info, UDataflowSubGraph* DocumentID) const
+{
+	const FSlateBrush* ReturnValue = FAppStyle::GetBrush(TEXT("GraphEditor.Function_16x"));
+	return ReturnValue;
+}
+
+TSharedRef<FGenericTabHistory> FDataflowEditorSubGraphTabSummoner::CreateTabHistoryNode(TSharedPtr<FTabPayload> Payload)
+{
+	return MakeShareable(new FGenericTabHistory(SharedThis(this), Payload));
+}
\ No newline at end of file
diff --git a/Engine/Plugins/Experimental/Dataflow/Source/DataflowEditor/Private/Dataflow/DataflowEditorSubGraphTabSummoner.h b/Engine/Plugins/Experimental/Dataflow/Source/DataflowEditor/Private/Dataflow/DataflowEditorSubGraphTabSummoner.h
new file mode 100644
--- /dev/null
+++ b/Engine/Plugins/Experimental/Dataflow/Source/DataflowEditor/Private/Dataflow/DataflowEditorSubGraphTabSummoner.h
@@ -0,0 +1,49 @@
+// Copyright Epic Games, Inc. All Rights Reserved.
+
+#pragma once
+
+#include "WorkflowOrientedApp/WorkflowUObjectDocuments.h"
+#include "Dataflow/DataflowSubGraph.h"
+
+class FDataflowEditorToolkit;
+class SDataflowGraphEditor;
+class SGraphEditor;
+/**
+* Summoner class for Dataflow SubGraph tabs
+*/
+struct FDataflowEditorSubGraphTabSummoner : public FDocumentTabFactoryForObjects<UDataflowSubGraph>
+{
+public:
+	DECLARE_DELEGATE_RetVal_TwoParams(TSharedRef<SGraphEditor>, FOnCreateGraphEditorWidget, TSharedRef<FTabInfo>, UDataflowSubGraph*);
+
+public:
+	FDataflowEditorSubGraphTabSummoner(TSharedPtr<FDataflowEditorToolkit> InEditorToolkitPtr, FOnCreateGraphEditorWidget CreateGraphEditorWidgetCallback);
+
+	virtual void OnTabActivated(TSharedPtr<SDockTab> Tab) const override;
+
+	virtual void OnTabForegrounded(TSharedPtr<SDockTab> Tab) const override;
+
+	virtual void OnTabBackgrounded(TSharedPtr<SDockTab> Tab) const override;
+
+	virtual void OnTabRefreshed(TSharedPtr<SDockTab> Tab) const override;
+
+	virtual void SaveState(TSharedPtr<SDockTab> Tab, TSharedPtr<FTabPayload> Payload) const override;
+
+protected:
+	static FText GetGraphDisplayName(const UEdGraph* Graph);
+
+	virtual TAttribute<FText> ConstructTabNameForObject(UDataflowSubGraph* DocumentID) const override;
+
+	virtual TSharedRef<SWidget> CreateTabBodyForObject(const FWorkflowTabSpawnInfo& Info, UDataflowSubGraph* DocumentID) const override;
+
+	virtual const FSlateBrush* GetTabIconForObject(const FWorkflowTabSpawnInfo& Info, UDataflowSubGraph* DocumentID) const override;
+
+	virtual TSharedRef<FGenericTabHistory> CreateTabHistoryNode(TSharedPtr<FTabPayload> Payload) override;
+
+	TSharedPtr<SDataflowGraphEditor> GetDataflowEditorFromTab(TSharedPtr<SDockTab> Tab) const;
+
+protected:
+	TWeakPtr<FDataflowEditorToolkit> EditorToolkitWeakPtr;
+	FOnCreateGraphEditorWidget OnCreateGraphEditorWidget;
+};
+
diff --git a/Engine/Plugins/Experimental/Dataflow/Source/DataflowEditor/Private/Dataflow/DataflowEditorToolkit.cpp b/Engine/Plugins/Experimental/Dataflow/Source/DataflowEditor/Private/Dataflow/DataflowEditorToolkit.cpp
--- a/Engine/Plugins/Experimental/Dataflow/Source/DataflowEditor/Private/Dataflow/DataflowEditorToolkit.cpp
+++ b/Engine/Plugins/Experimental/Dataflow/Source/DataflowEditor/Private/Dataflow/DataflowEditorToolkit.cpp
@@ -23,6 +23,7 @@
 #include "Dataflow/DataflowEditorModule.h"
 #include "Dataflow/DataflowEditorModeUILayer.h"
 #include "Dataflow/DataflowEditorUtil.h"
+#include "Dataflow/DataflowEditorSubGraphTabSummoner.h"
 #include "Dataflow/DataflowConstructionViewport.h"
 #include "Dataflow/DataflowEdNode.h"
 #include "Dataflow/DataflowGraphEditor.h"
@@ -63,6 +64,7 @@
 #include "Widgets/Docking/SDockTab.h"
 #include "GeometryCache.h"
 #include "ToolMenus.h"
+#include "WorkflowOrientedApp/WorkflowTabManager.h"
 
 #define LOCTEXT_NAMESPACE "DataflowEditorToolkit"
 
@@ -73,6 +75,7 @@ FAutoConsoleVariableRef CVARDataflowEnableSkeletonView(TEXT("p.Dataflow.Editor.E
 	TEXT("Deprecated Tool! Allows the Dataflow editor to create a skeleton view that reflects the hierarchy and selection state of the construction viewport.[def:false]"));
 
 const FName FDataflowEditorToolkit::GraphCanvasTabId(TEXT("DataflowEditor_GraphCanvas"));
+const FName FDataflowEditorToolkit::SubGraphCanvasTabId(TEXT("DataflowEditor_SubGraphTab"));
 const FName FDataflowEditorToolkit::NodeDetailsTabId(TEXT("DataflowEditor_NodeDetails"));
 const FName FDataflowEditorToolkit::PreviewSceneTabId(TEXT("DataflowEditor_PreviewScene"));
 const FName FDataflowEditorToolkit::SkeletonViewTabId(TEXT("DataflowEditor_SkeletonView"));
@@ -152,6 +155,7 @@ FDataflowEditorToolkit::FDataflowEditorToolkit(UAssetEditor* InOwningAssetEditor
 						FTabManager::NewStack()
 						->SetSizeCoefficient(0.1f)	// Relative height of (Dataflow Graph Editor, Outliner) vs (Tools Panel, Construction Viewport, Preview Viewport)
 						->AddTab(GraphCanvasTabId, ETabState::OpenedTab)
+						->AddTab(SubGraphCanvasTabId, ETabState::OpenedTab)
 						->SetExtensionId("MembersWidgetArea")
 						->SetHideTabWell(false)
 						->SetForegroundTab(MembersWidgetTabId)
@@ -161,6 +165,7 @@ FDataflowEditorToolkit::FDataflowEditorToolkit(UAssetEditor* InOwningAssetEditor
 						FTabManager::NewStack()
 						->SetSizeCoefficient(0.8f)	// Relative height of (Dataflow Graph Editor, Outliner) vs (Tools Panel, Construction Viewport, Preview Viewport)
 						->AddTab(GraphCanvasTabId, ETabState::OpenedTab)
+						->AddTab(SubGraphCanvasTabId, ETabState::OpenedTab)
 						->SetExtensionId("GraphEditorArea")
 						->SetHideTabWell(false)
 						->SetForegroundTab(GraphCanvasTabId)
@@ -239,6 +244,7 @@ FDataflowEditorToolkit::FDataflowEditorToolkit(UAssetEditor* InOwningAssetEditor
 						FTabManager::NewStack()
 						->SetSizeCoefficient(0.8f)	// Relative height of (Dataflow Graph Editor, Outliner) vs (Tools Panel, Construction Viewport, Preview Viewport)
 						->AddTab(GraphCanvasTabId, ETabState::OpenedTab)
+						->AddTab(SubGraphCanvasTabId, ETabState::OpenedTab)
 						->SetExtensionId("GraphEditorArea")
 						->SetHideTabWell(false)
 						->SetForegroundTab(GraphCanvasTabId)
@@ -314,7 +320,6 @@ FDataflowEditorToolkit::FDataflowEditorToolkit(UAssetEditor* InOwningAssetEditor
 			bDataflowEnableGraphEval = Var->GetBool();
 		});
 	}
-
 }
 
 FDataflowEditorToolkit::~FDataflowEditorToolkit()
@@ -1204,9 +1209,9 @@ void FDataflowEditorToolkit::EvaluateNode(const FDataflowNode* Node, const FData
 	}
 }
 
-TSharedRef<SDataflowGraphEditor> FDataflowEditorToolkit::CreateGraphEditorWidget(UDataflow* DataflowToEdit, TSharedPtr<IStructureDetailsView> InNodeDetailsEditor)
+TSharedRef<SDataflowGraphEditor> FDataflowEditorToolkit::CreateGraphEditorWidget(UEdGraph* GraphToEdit, TSharedPtr<IStructureDetailsView> InNodeDetailsEditor)
 {
-	ensure(DataflowToEdit);
+	ensure(GraphToEdit);
 	using namespace UE::Dataflow;
 
 	const FDataflowEditorCommands::FGraphEvaluationCallback Evaluate =
@@ -1234,8 +1239,11 @@ TSharedRef<SDataflowGraphEditor> FDataflowEditorToolkit::CreateGraphEditorWidget
 	InEvents.OnTextCommitted = FOnNodeTextCommitted::CreateSP(this, &FDataflowEditorToolkit::OnNodeTitleCommitted);
 	InEvents.OnNodeSingleClicked = SGraphEditor::FOnNodeSingleClicked::CreateSP(this, &FDataflowEditorToolkit::OnNodeSingleClicked);
 
-	TSharedRef<SDataflowGraphEditor> NewGraphEditor = SNew(SDataflowGraphEditor, DataflowToEdit)
-		.GraphToEdit(DataflowToEdit)
+	UDataflow* DataflowAsset = UDataflow::GetDataflowAssetFromEdGraph(GraphToEdit);
+	check(DataflowAsset);
+
+	TSharedRef<SDataflowGraphEditor> NewGraphEditor = SNew(SDataflowGraphEditor, DataflowAsset)
+		.GraphToEdit(GraphToEdit)
 		.GraphEvents(InEvents)
 		.DetailsView(InNodeDetailsEditor)
 		.EvaluateGraph(Evaluate)
@@ -1361,6 +1369,27 @@ TSharedRef<SDockTab> FDataflowEditorToolkit::SpawnTab_GraphCanvas(const FSpawnTa
 	return SpawnedTab;
 }
 
+TSharedRef<SDockTab> FDataflowEditorToolkit::SpawnTab_SubGraphTab(const FSpawnTabArgs& Args)
+{
+	check(Args.GetTabId() == SubGraphCanvasTabId);
+
+	TSharedPtr<SWidget> SubgraphEditor = SNullWidget::NullWidget;
+	if (const TObjectPtr<UDataflowBaseContent>& EditorContent = GetEditorContent())
+	{
+		if (UDataflow* DataflowAsset = EditorContent->GetDataflowAsset())
+		{
+			SubgraphEditor = CreateGraphEditorWidget(DataflowAsset, NodeDetailsEditor);
+		}
+	}
+	TSharedRef<SDockTab> SpawnedTab = SNew(SDockTab)
+		.Label(LOCTEXT("DataflowEditor_Subgraph_TabTitle", "Dataflow Graph"))
+		[
+			SubgraphEditor.ToSharedRef()
+		];
+
+	return SpawnedTab;
+}
+
 TSharedRef<SDockTab> FDataflowEditorToolkit::SpawnTab_NodeDetails(const FSpawnTabArgs& Args)
 {
 	check(Args.GetTabId() == NodeDetailsTabId);
@@ -1621,6 +1650,12 @@ void FDataflowEditorToolkit::RegisterTabSpawners(const TSharedRef<FTabManager>&
 		.SetGroup(EditorMenuCategory.ToSharedRef())
 		.SetIcon(FSlateIcon(FAppStyle::GetAppStyleSetName(), "GraphEditor.EventGraph_16x"));
 
+	InTabManager->RegisterTabSpawner(SubGraphCanvasTabId, FOnSpawnTab::CreateSP(this, &FDataflowEditorToolkit::SpawnTab_SubGraphTab))
+		.SetDisplayName(LOCTEXT("SubgraphTab", "Dataflow Subgraph"))
+		.SetGroup(EditorMenuCategory.ToSharedRef())
+		.SetIcon(FSlateIcon(FAppStyle::GetAppStyleSetName(), "GraphEditor.EventGraph_16x"))
+		.SetMenuType(ETabSpawnerMenuType::Hidden); // hide it from menus as those tabs can only be open by the toolkit itself 
+
 	InTabManager->RegisterTabSpawner(NodeDetailsTabId, FOnSpawnTab::CreateSP(this, &FDataflowEditorToolkit::SpawnTab_NodeDetails))
 		.SetDisplayName(LOCTEXT("NodeDetailsTab", "Node Details"))
 		.SetGroup(EditorMenuCategory.ToSharedRef())
@@ -1681,6 +1716,19 @@ void FDataflowEditorToolkit::RegisterTabSpawners(const TSharedRef<FTabManager>&
 	InTabManager->RegisterTabSpawner(MembersWidgetTabId, FOnSpawnTab::CreateSP(this, &FDataflowEditorToolkit::SpawnTab_MembersWidget))
 		.SetDisplayName(LOCTEXT("MembersWidgetTabDisplayName", "Dataflow Members"))
 		.SetGroup(AssetEditorTabsCategory.ToSharedRef());
+
+
+	// Create a document manager to be able to spawn subgraph tabs
+	DocumentManager = MakeShareable(new FDocumentTracker(SubGraphCanvasTabId));
+	DocumentManager->Initialize(SharedThis(this));
+	DocumentManager->SetTabManager(InTabManager);
+
+	TSharedPtr<FDataflowEditorSubGraphTabSummoner> SubGraphTabSummoner = 
+		MakeShareable( new FDataflowEditorSubGraphTabSummoner(
+			SharedThis(this),
+			FDataflowEditorSubGraphTabSummoner::FOnCreateGraphEditorWidget::CreateSP(this, &FDataflowEditorToolkit::CreateSubGraphEditorWidget)
+		));
+	DocumentManager->RegisterDocumentFactory(SubGraphTabSummoner);
 }
 
 void FDataflowEditorToolkit::UnregisterTabSpawners(const TSharedRef<class FTabManager>& InTabManager)
@@ -1741,6 +1789,47 @@ void FDataflowEditorToolkit::OnTabClosed(TSharedRef<SDockTab> Tab)
 	}
 }
 
+void FDataflowEditorToolkit::SetSubGraphTabActiveState(TSharedPtr<SDataflowGraphEditor> SubGraphEditor, bool bActive)
+{
+	if (bActive)
+	{
+		ActiveSubGraphEditorWeakPtr = SubGraphEditor;
+	}
+	else
+	{
+		// only reset to null if this was th previously active subgraph
+		if (ActiveSubGraphEditorWeakPtr == SubGraphEditor)
+		{
+			ActiveSubGraphEditorWeakPtr.Reset();
+		}
+	}
+}
+
+void FDataflowEditorToolkit::OpenSubGraphTab(FName SubGraphName)
+{
+	if (const TObjectPtr<UDataflowBaseContent>& EditorContent = GetEditorContent())
+	{
+		ensure(EditorContent);
+		if (UDataflow* DataflowAsset = EditorContent->GetDataflowAsset())
+		{
+			if (UDataflowSubGraph* SubGraph = DataflowAsset->FindSubGraphByName(SubGraphName))
+			{
+				TSharedRef<FTabPayload_UObject> Payload = FTabPayload_UObject::Make(SubGraph);
+				DocumentManager->OpenDocument(Payload, FDocumentTracker::EOpenDocumentCause::OpenNewDocument);
+			}
+		}
+	}
+}
+
+TSharedRef<SGraphEditor> FDataflowEditorToolkit::CreateSubGraphEditorWidget(TSharedRef<FTabInfo> InTabInfo, UDataflowSubGraph* InGraph)
+{
+	if (const TObjectPtr<UDataflowBaseContent>& EditorContent = GetEditorContent())
+	{
+		return CreateGraphEditorWidget(InGraph, NodeDetailsEditor);
+	}
+	return SNew(SGraphEditor);
+}
+
 FName FDataflowEditorToolkit::GetToolkitFName() const
 {
 	return FName("DataflowEditor");
diff --git a/Engine/Plugins/Experimental/Dataflow/Source/DataflowEditor/Public/Dataflow/DataflowEditorToolkit.h b/Engine/Plugins/Experimental/Dataflow/Source/DataflowEditor/Public/Dataflow/DataflowEditorToolkit.h
--- a/Engine/Plugins/Experimental/Dataflow/Source/DataflowEditor/Public/Dataflow/DataflowEditorToolkit.h
+++ b/Engine/Plugins/Experimental/Dataflow/Source/DataflowEditor/Public/Dataflow/DataflowEditorToolkit.h
@@ -13,12 +13,15 @@
 #include "Dataflow/DataflowConstructionViewport.h"
 #include "Dataflow/DataflowSimulationViewport.h"
 
+class FDocumentTracker;
 class FEditorViewportTabContent;
 class IDetailsView;
 class FTabManager;
+class FTabInfo;
 class IStructureDetailsView;
 class IToolkitHost;
 class UDataflow;
+class UDataflowSubGraph;
 class USkeletalMesh;
 class SDataflowGraphEditor;
 class FDataflowConstructionScene;
@@ -82,6 +85,10 @@ public:
 		return TEXT("FDataflowEditorToolkit");
 	}
 	// End of FSerializableObject interface
+
+	void OpenSubGraphTab(FName SubGraphName);
+	void SetSubGraphTabActiveState(TSharedPtr<SDataflowGraphEditor> SubGraphEditor, bool bActive);
+
 protected:
 
 	UDataflowEditor* DataflowEditor = nullptr;
@@ -106,6 +113,7 @@ private:
 	
 	// Spawning of all the additional tabs (viewport,details ones are coming from the base asset toolkit)
 	TSharedRef<SDockTab> SpawnTab_GraphCanvas(const FSpawnTabArgs& Args);
+	TSharedRef<SDockTab> SpawnTab_SubGraphTab(const FSpawnTabArgs& Args);
 	TSharedRef<SDockTab> SpawnTab_NodeDetails(const FSpawnTabArgs& Args);
 	TSharedRef<SDockTab> SpawnTab_SkeletonView(const FSpawnTabArgs& Args);
 	TSharedRef<SDockTab> SpawnTab_SelectionView(const FSpawnTabArgs& Args);
@@ -142,6 +150,7 @@ private:
 
 	// List of all the tab names ids that will be used to identify the editor widgets
 	static const FName GraphCanvasTabId;
+	static const FName SubGraphCanvasTabId;
 	static const FName NodeDetailsTabId;
 	static const FName SkeletonViewTabId;
 	static const FName SelectionViewTabId_1;
@@ -183,11 +192,12 @@ private:
 	TSharedRef<class IDetailCustomization> CustomizeAssetViewer() const;
 
 	// Utility factory functions to build the widgets
-	TSharedRef<SDataflowGraphEditor> CreateGraphEditorWidget(UDataflow* ObjectToEdit, TSharedPtr<IStructureDetailsView> PropertiesEditor);
+	TSharedRef<SDataflowGraphEditor> CreateGraphEditorWidget(UEdGraph* GraphToEdit, TSharedPtr<IStructureDetailsView> PropertiesEditor);
     TSharedPtr<IDetailsView> CreateAssetDetailsEditorWidget(const TArray<UObject*>& ObjectsToEdit);
 	TSharedPtr<SWidget> CreateSimulationVisualizationWidget();
     TSharedPtr<IStructureDetailsView> CreateNodeDetailsEditorWidget(UObject* ObjectToEdit);
 	TSharedPtr<SDataflowMembersWidget> CreateDataflowMembersWidget();
+	TSharedRef<SGraphEditor> CreateSubGraphEditorWidget(TSharedRef<FTabInfo> InTabInfo, UDataflowSubGraph* InGraph);
 
 	/** Create the simulation viewport client */
 	void CreateSimulationViewportClient();
@@ -245,4 +255,9 @@ private:
 
 	/** Delegate for updating the cached value of p.Dataflow.EnableGraphEval */
 	FDelegateHandle GraphEvalCVarChangedDelegateHandle;
+
+	TWeakPtr<SDataflowGraphEditor> ActiveSubGraphEditorWeakPtr;
+
+	/** Document tracker for dynamic tabs ( like subgraphs ) */
+	TSharedPtr<FDocumentTracker> DocumentManager;
 };
