Dataflow:
- Added sperator to Dataflow rt click menu on nodes

diff --git a/Engine/Plugins/Experimental/Dataflow/Source/DataflowEditor/Private/Dataflow/DataflowSchema.cpp b/Engine/Plugins/Experimental/Dataflow/Source/DataflowEditor/Private/Dataflow/DataflowSchema.cpp
--- a/Engine/Plugins/Experimental/Dataflow/Source/DataflowEditor/Private/Dataflow/DataflowSchema.cpp
+++ b/Engine/Plugins/Experimental/Dataflow/Source/DataflowEditor/Private/Dataflow/DataflowSchema.cpp
@@ -45,6 +45,7 @@ void UDataflowSchema::GetContextMenuActions(class UToolMenu* Menu, class UGraphN
 				Section.AddMenuEntry(FDataflowEditorCommands::Get().AddOptionPin);
 				Section.AddMenuEntry(FDataflowEditorCommands::Get().RemoveOptionPin);
 				Section.AddMenuEntry(FDataflowEditorCommands::Get().EvaluateNode);
+				Section.AddSeparator(TEXT("ActionsSeparator"));
 				Section.AddMenuEntry(FDataflowEditorCommands::Get().FreezeNodes, FText::FromString("Freeze"));
 				Section.AddMenuEntry(FDataflowEditorCommands::Get().UnfreezeNodes, FText::FromString("Unfreeze"));
 			}
