Dataflow:
- CollectionSpreadSheet 2.0
- Updated Transform, Faces, Vertices group display to hierarchial display, added color scheme by level
- Bug fixes

diff --git a/Engine/Plugins/Experimental/Dataflow/Source/DataflowEditor/Private/Dataflow/CollectionSpreadSheetWidget.cpp b/Engine/Plugins/Experimental/Dataflow/Source/DataflowEditor/Private/Dataflow/CollectionSpreadSheetWidget.cpp
--- a/Engine/Plugins/Experimental/Dataflow/Source/DataflowEditor/Private/Dataflow/CollectionSpreadSheetWidget.cpp
+++ b/Engine/Plugins/Experimental/Dataflow/Source/DataflowEditor/Private/Dataflow/CollectionSpreadSheetWidget.cpp
@@ -2,236 +2,21 @@
 
 #include "Dataflow/DataflowCollectionSpreadSheetWidget.h"
 #include "Widgets/Input/SButton.h"
-#include "Styling/StarshipCoreStyle.h"
+#include "Widgets/Colors/SColorBlock.h"
 #include "Widgets/Layout/SScrollBox.h"
+#include "Dataflow/DataflowSettings.h"
+#include "GeometryCollection/GeometryCollection.h"
+#include "Dataflow/DataflowCollectionSpreadSheetHelpers.h"
+#include "GeometryCollection/ManagedArrayCollection.h"
+#include "Widgets/Views/STableRow.h"
+#include "Widgets/Text/STextBlock.h"
+#include "Widgets/Input/SComboBox.h"
+#include "Widgets/Input/SCheckBox.h"
 
+#define LOCTEXT_NAMESPACE "CollectionSpreadSheet"
 
 const FName FCollectionSpreadSheetHeader::IndexColumnName = FName("Index");
 
-namespace CollectionSpreadSheetHelpers
-{
-	template<typename T>
-	FString AttributeValueToString(const T& Value)
-	{
-		return Value.ToString();
-	}
-
-	FString AttributeValueToString(float Value)
-	{
-		return FString::SanitizeFloat(Value, 2);
-	}
-
-	FString AttributeValueToString(int32 Value)
-	{
-		return FString::FromInt(Value);
-	}
-
-	FString AttributeValueToString(FString Value)
-	{
-		return Value;
-	}
-
-	FString AttributeValueToString(FLinearColor Value)
-	{
-		return FString::Printf(TEXT("(R=%.2f G=%.2f B=%.2f A=%.2f)"), Value.R, Value.G, Value.B, Value.A);
-	}
-
-	FString AttributeValueToString(FVector Value)
-	{
-		return FString::Printf(TEXT("(X=%.2f Y=%.2f Z=%.2f)"), Value.X, Value.Y, Value.Z);
-	}
-
-	FString AttributeValueToString(bool Value)
-	{
-		return Value ? FString("true") : FString("false");
-	}
-
-	FString AttributeValueToString(const FConstBitReference& Value)
-	{
-		return Value ? FString("true") : FString("false");
-	}
-
-	FString AttributeValueToString(TSet<int32> Value)
-	{
-		TArray<int32> Array = Value.Array();
-		FString Out;
-
-		for (int32 Idx = 0; Idx < Array.Num(); ++Idx)
-		{
-			Out += AttributeValueToString(Array[Idx]);
-
-			if (Idx != Array.Num() - 1)
-			{
-				Out += " ";
-			}
-		}
-
-		return Out;
-	}
-
-	template<typename T>
-	FString AttributeValueToString(const TArray<T>& Array)
-	{
-		FString Out;
-
-		for (int32 Idx = 0; Idx < Array.Num(); ++Idx)
-		{
-			Out += AttributeValueToString(Array[Idx]);
-
-			if (Idx != Array.Num() - 1)
-			{
-				Out += "; ";
-			}
-		}
-
-		return Out;
-	}
-
-	FString AttributeValueToString(FTransform3f Value)
-	{
-		const FVector3f Translation = Value.GetTranslation();
-		const FVector3f Rotation = Value.GetRotation().Euler();
-		const FVector3f Scale = Value.GetScale3D();
-
-		return FString::Printf(TEXT("T:(%s) R:(%s) S:(%s)"), *Translation.ToString(), *Rotation.ToString(), *Scale.ToString());
-	}
-
-	FString AttributeValueToString(FTransform Value)
-	{
-		const FVector Translation = Value.GetTranslation();
-		const FVector Rotation = Value.GetRotation().Euler();
-		const FVector Scale = Value.GetScale3D();
-
-		return FString::Printf(TEXT("T:(%s) R:(%s) S:(%s)"), *Translation.ToString(), *Rotation.ToString(), *Scale.ToString());
-	}
-
-	FString AttributeValueToString(FBox Value)
-	{
-		return FString::Printf(TEXT("Min:(%s) Max:(%s)"), *Value.Min.ToString(), *Value.Max.ToString());
-	}
-
-	FString AttributeValueToString(FIntVector Value)
-	{
-		return FString::Printf(TEXT("%d %d %d"), Value.X, Value.Y, Value.Z);
-	}
-
-	FString AttributeValueToString(FIntVector4 Value)
-	{
-		return FString::Printf(TEXT("%d %d %d %d"), Value[0], Value[1], Value[2], Value[3]);
-	}
-
-	FString AttributeValueToString(FGuid Value)
-	{
-		return Value.ToString();
-	}
-
-	template<typename T>
-	FString AttributeValueToString(const FManagedArrayCollection& InCollection, const FName& InAttributeName, const FName& InGroupName, int32 InIdxColumn)
-	{
-		const TManagedArray<T>* const Array = InCollection.FindAttributeTyped<T>(InAttributeName, InGroupName);
-		if (Array == nullptr)
-		{
-			return FString("<Unknown Attribute>");
-		}
-
-		if (InIdxColumn < 0 || Array->Num() <= InIdxColumn)
-		{
-			return FString("<Index out of bounds>");
-		}
-
-		return AttributeValueToString((*Array)[InIdxColumn]);
-	}
-
-	FString AttributeValueToString(const FManagedArrayCollection& InCollection, const FName& InAttributeName, const FName& InGroupName, int32 InIdxColumn)
-	{
-		const FManagedArrayCollection::EArrayType ArrayType = InCollection.GetAttributeType(InAttributeName, InGroupName);
-
-		FString ValueAsString;
-
-		switch (ArrayType)
-		{
-		case FManagedArrayCollection::EArrayType::FFloatType:
-			ValueAsString = AttributeValueToString<float>(InCollection, InAttributeName, InGroupName, InIdxColumn);
-			break;
-
-		case FManagedArrayCollection::EArrayType::FInt32Type:
-			ValueAsString = AttributeValueToString<int32>(InCollection, InAttributeName, InGroupName, InIdxColumn);
-			break;
-
-		case FManagedArrayCollection::EArrayType::FBoolType:
-			ValueAsString = AttributeValueToString<bool>(InCollection, InAttributeName, InGroupName, InIdxColumn);
-			break;
-
-		case FManagedArrayCollection::EArrayType::FStringType:
-			ValueAsString = AttributeValueToString<FString>(InCollection, InAttributeName, InGroupName, InIdxColumn);
-			break;
-
-		case FManagedArrayCollection::EArrayType::FLinearColorType:
-			ValueAsString = AttributeValueToString<FLinearColor>(InCollection, InAttributeName, InGroupName, InIdxColumn);
-			break;
-
-		case FManagedArrayCollection::EArrayType::FVectorType:
-			ValueAsString = AttributeValueToString<FVector3f>(InCollection, InAttributeName, InGroupName, InIdxColumn);
-			break;
-
-		case FManagedArrayCollection::EArrayType::FVector2DType:
-			ValueAsString = AttributeValueToString<FVector2f>(InCollection, InAttributeName, InGroupName, InIdxColumn);
-			break;
-
-		case FManagedArrayCollection::EArrayType::FVector3dType:
-			ValueAsString = AttributeValueToString<FVector2f>(InCollection, InAttributeName, InGroupName, InIdxColumn);
-			break;
-
-		case FManagedArrayCollection::EArrayType::FIntVectorType:
-			ValueAsString = AttributeValueToString<FIntVector>(InCollection, InAttributeName, InGroupName, InIdxColumn);
-			break;
-
-		case FManagedArrayCollection::EArrayType::FTransformType:
-			ValueAsString = AttributeValueToString<FTransform>(InCollection, InAttributeName, InGroupName, InIdxColumn);
-			break;
-
-		case FManagedArrayCollection::EArrayType::FVector2DArrayType:
-			ValueAsString = AttributeValueToString<TArray<FVector2f>>(InCollection, InAttributeName, InGroupName, InIdxColumn);
-			break;
-
-		case FManagedArrayCollection::EArrayType::FIntArrayType:
-			ValueAsString = AttributeValueToString<TSet<int32>>(InCollection, InAttributeName, InGroupName, InIdxColumn);
-			break;
-
-		case FManagedArrayCollection::EArrayType::FInt32ArrayType:
-			ValueAsString = AttributeValueToString<TArray<int32>>(InCollection, InAttributeName, InGroupName, InIdxColumn);
-			break;
-
-		case FManagedArrayCollection::EArrayType::FFloatArrayType:
-			ValueAsString = AttributeValueToString<TArray<float>>(InCollection, InAttributeName, InGroupName, InIdxColumn);
-			break;
-
-		case FManagedArrayCollection::EArrayType::FBoxType:
-			ValueAsString = AttributeValueToString<FBox>(InCollection, InAttributeName, InGroupName, InIdxColumn);
-			break;
-
-		case FManagedArrayCollection::EArrayType::FTransform3fType:
-			ValueAsString = AttributeValueToString<FTransform3f>(InCollection, InAttributeName, InGroupName, InIdxColumn);
-			break;
-
-		case FManagedArrayCollection::EArrayType::FIntVector4Type:
-			ValueAsString = AttributeValueToString<FIntVector4>(InCollection, InAttributeName, InGroupName, InIdxColumn);
-			break;
-
-		case FManagedArrayCollection::EArrayType::FGuidType:
-			ValueAsString = AttributeValueToString<FGuid>(InCollection, InAttributeName, InGroupName, InIdxColumn);
-			break;
-
-		default:
-			//ensure(false);
-			ValueAsString = "<Unknown Data Type>";
-		}
-
-		return ValueAsString;
-	}
-}
-
-
 void SCollectionSpreadSheetRow::Construct(const FArguments& InArgs, TSharedRef<STableViewBase> OwnerTableView, const TSharedPtr<const FCollectionSpreadSheetHeader>& InHeader, const TSharedPtr<const FCollectionSpreadSheetItem>& InItem)
 {
 	Header = InHeader;
@@ -239,8 +24,8 @@ void SCollectionSpreadSheetRow::Construct(const FArguments& InArgs, TSharedRef<S
 
 	SMultiColumnTableRow<TSharedPtr<const FCollectionSpreadSheetItem>>::Construct(
 		FSuperRowType::FArguments()
-		.Style(&FStarshipCoreStyle::GetCoreStyle().GetWidgetStyle<FTableRowStyle>("TableView.AlternatingRow")),
-		OwnerTableView);
+		.Style(&FStarshipCoreStyle::GetCoreStyle().GetWidgetStyle<FTableRowStyle>("TableView.AlternatingRow"))
+		, OwnerTableView);
 }
 
 
@@ -273,10 +58,6 @@ void SCollectionSpreadSheet::Construct(const FArguments& InArgs)
 {
 	SelectedOutput = InArgs._SelectedOutput;
 
-	SpreadSheetVerticalScrollBar = SNew(SScrollBar)
-		.Orientation(Orient_Vertical)
-		.Thickness(FVector2D(12.0f, 12.0f));
-
 	HeaderRowWidget = SNew(SHeaderRow)
 		.Visibility(EVisibility::Visible);
 
@@ -296,32 +77,14 @@ void SCollectionSpreadSheet::Construct(const FArguments& InArgs)
 			+ SOverlay::Slot()
 			[
 				SAssignNew(ListView, SListView<TSharedPtr<const FCollectionSpreadSheetItem>>)
-//				SAssignNew(ListView, STreeView<TSharedPtr<const FCollectionSpreadSheetItem>>)
 				.SelectionMode(ESelectionMode::Multi)
 				.ListItemsSource(&ListItems)
-//				.TreeItemsSource(&ListItems)
 				.OnGenerateRow(this, &SCollectionSpreadSheet::GenerateRow)
 				.HeaderRow(HeaderRowWidget)
-				.ExternalScrollbar(SpreadSheetVerticalScrollBar)
+				.ExternalScrollbar(InArgs._ExternalVerticalScrollBar)
 			]
 		]
 	];
-
-	AttrTypeWidthMap.Add("Transform", 600);
-	AttrTypeWidthMap.Add("Transform3f", 600);
-	AttrTypeWidthMap.Add("String", 200);
-	AttrTypeWidthMap.Add("LinearColor", 250);
-	AttrTypeWidthMap.Add("int32", 100);
-	AttrTypeWidthMap.Add("IntArray", 200);
-	AttrTypeWidthMap.Add("Vector", 300);
-	AttrTypeWidthMap.Add("Vector2D", 160);
-	AttrTypeWidthMap.Add("Float", 150);
-	AttrTypeWidthMap.Add("IntVector", 150);
-	AttrTypeWidthMap.Add("Bool", 75);
-	AttrTypeWidthMap.Add("Box", 550);
-	AttrTypeWidthMap.Add("MeshSection", 0);
-	AttrTypeWidthMap.Add("UInt8", 0);
-	AttrTypeWidthMap.Add("Guid", 350);
 }
 
 
@@ -339,6 +102,10 @@ const FName& SCollectionSpreadSheet::GetSelectedOutput() const
 	return SelectedOutput;
 }
 
+const FString SCollectionSpreadSheet::GetSelectedOutputStr() const
+{
+	return SelectedOutput.ToString();
+}
 
 void SCollectionSpreadSheet::SetSelectedGroup(const FName& InSelectedGroup)
 {
@@ -354,20 +121,6 @@ const FName& SCollectionSpreadSheet::GetSelectedGroup() const
 	return SelectedGroup;
 }
 
-namespace {
-	inline FName GetArrayTypeString(FManagedArrayCollection::EArrayType ArrayType)
-	{
-		switch (ArrayType)
-		{
-#define MANAGED_ARRAY_TYPE(a,A)	case EManagedArrayType::F##A##Type:\
-			return FName(#A);
-#include "GeometryCollection/ManagedArrayTypeValues.inl"
-#undef MANAGED_ARRAY_TYPE
-		}
-		return FName();
-	}
-}
-
 void SCollectionSpreadSheet::RegenerateHeader()
 {
 	HeaderRowWidget->ClearColumns();
@@ -383,7 +136,7 @@ void SCollectionSpreadSheet::RegenerateHeader()
 		for (FName Attr : CollectionInfoMap[SelectedOutput.ToString()].Collection.AttributeNames(SelectedGroup))
 		{
 			Header->ColumnNames.Add(Attr);
-			AttrTypes.Add(GetArrayTypeString(CollectionInfoMap[SelectedOutput.ToString()].Collection.GetAttributeType(Attr, SelectedGroup)).ToString());
+			AttrTypes.Add(UE::Dataflow::CollectionSpreadSheetHelpers::GetArrayTypeString(CollectionInfoMap[SelectedOutput.ToString()].Collection.GetAttributeType(Attr, SelectedGroup)).ToString());
 		}
 
 	}
@@ -416,9 +169,9 @@ void SCollectionSpreadSheet::RegenerateHeader()
 		{
 			int32 ColumnNameStrLen = ColumnNameStr.Len() * 9;
 			int32 AttrTypeWidth = 100;
-			if (AttrTypeWidthMap.Contains(AttrTypeStr))
+			if (UE::Dataflow::CollectionSpreadSheetHelpers::AttrTypeWidthMap.Contains(AttrTypeStr))
 			{
-				AttrTypeWidth = AttrTypeWidthMap[AttrTypeStr];
+				AttrTypeWidth = UE::Dataflow::CollectionSpreadSheetHelpers::AttrTypeWidthMap[AttrTypeStr];
 			}
 			ColumnWidth = ColumnNameStrLen > AttrTypeWidth ? ColumnNameStrLen : AttrTypeWidth;
 		}
@@ -436,7 +189,6 @@ void SCollectionSpreadSheet::RegenerateHeader()
 	}
 }
 
-
 void SCollectionSpreadSheet::RepopulateListView()
 {
 	ListItems.Empty();
@@ -461,7 +213,7 @@ void SCollectionSpreadSheet::RepopulateListView()
 				}
 				else
 				{
-					NewItem->Values[IdxColumn] = CollectionSpreadSheetHelpers::AttributeValueToString(CollectionInfoMap[SelectedOutput.ToString()].Collection, ColumnName, SelectedGroup, IdxElem);
+					NewItem->Values[IdxColumn] = UE::Dataflow::CollectionSpreadSheetHelpers::AttributeValueToString(CollectionInfoMap[SelectedOutput.ToString()].Collection, ColumnName, SelectedGroup, IdxElem);
 				}
 			}
 
@@ -505,6 +257,45 @@ void SCollectionSpreadSheetWidget::NodeOutputsComboBoxSelectionChanged(FName InS
 			}
 		}
 	}
+
+	if (TransformOutliner)
+	{
+		const FString OutputStr = InSelectedOutput.ToString();
+
+		if (CollectionTable->GetCollectionInfoMap().Contains(OutputStr))
+		{
+			TSharedPtr<FManagedArrayCollection> CollectionPtr = MakeShared<FManagedArrayCollection>(CollectionTable->GetCollectionInfoMap()[OutputStr].Collection);
+			TransformOutliner->SetCollection(CollectionPtr, InSelectedOutput);
+
+			TransformOutliner->RegenerateHeader();
+		}
+	}
+
+	if (VerticesOutliner)
+	{
+		const FString OutputStr = InSelectedOutput.ToString();
+
+		if (CollectionTable->GetCollectionInfoMap().Contains(OutputStr))
+		{
+			TSharedPtr<FManagedArrayCollection> CollectionPtr = MakeShared<FManagedArrayCollection>(CollectionTable->GetCollectionInfoMap()[OutputStr].Collection);
+			VerticesOutliner->SetCollection(CollectionPtr);
+
+			VerticesOutliner->RegenerateHeader();
+		}
+	}
+
+	if (FacesOutliner)
+	{
+		const FString OutputStr = InSelectedOutput.ToString();
+
+		if (CollectionTable->GetCollectionInfoMap().Contains(OutputStr))
+		{
+			TSharedPtr<FManagedArrayCollection> CollectionPtr = MakeShared<FManagedArrayCollection>(CollectionTable->GetCollectionInfoMap()[OutputStr].Collection);
+			FacesOutliner->SetCollection(CollectionPtr);
+
+			FacesOutliner->RegenerateHeader();
+		}
+	}
 }
 
 
@@ -526,6 +317,33 @@ void SCollectionSpreadSheetWidget::CollectionGroupsComboBoxSelectionChanged(FNam
 		}
 	}
 
+	if (TransformOutliner)
+	{
+		const FString GroupStr = InSelectedGroup.ToString();
+		if (GroupStr == "Transform")
+		{
+			TransformOutliner->RegenerateHeader();
+		}
+	}
+
+	if (VerticesOutliner)
+	{
+		const FString GroupStr = InSelectedGroup.ToString();
+		if (GroupStr == "Vertices")
+		{
+			VerticesOutliner->RegenerateHeader();
+		}
+	}
+
+	if (FacesOutliner)
+	{
+		const FString GroupStr = InSelectedGroup.ToString();
+		if (GroupStr == "Faces")
+		{
+			FacesOutliner->RegenerateHeader();
+		}
+	}
+
 	SetStatusText();
 }
 
@@ -541,12 +359,57 @@ FText SCollectionSpreadSheetWidget::GetNoGroupText()
 	return FText::FromString("No Group(s)");
 }
 
+const FSlateBrush* SCollectionSpreadSheetWidget::GetPinButtonImage() const
+{
+	if (bIsPinnedDown)
+	{
+		return FAppStyle::Get().GetBrush("Icons.Pinned");
+	}
+	else
+	{
+		return FAppStyle::Get().GetBrush("Icons.Unpinned");
+	}
+}
+
+const FSlateBrush* SCollectionSpreadSheetWidget::GetLockButtonImage() const
+{
+	if (bIsRefreshLocked)
+	{
+		return FAppStyle::Get().GetBrush("Icons.Lock");
+	}
+	else
+	{
+		return FAppStyle::Get().GetBrush("Icons.Unlock");
+	}
+}
 
 void SCollectionSpreadSheetWidget::Construct(const FArguments& InArgs)
 {
-	TSharedPtr<SScrollBar> SpreadSheetHorizontalScrollBar = SNew(SScrollBar)
-		.Orientation(Orient_Horizontal)
-		.Thickness(FVector2D(12.0f, 12.0f));
+	constexpr float CScrollBarWidth = 12.f;
+
+	SAssignNew(SpreadSheetHorizontalScrollBar, SScrollBar)
+	.Orientation(Orient_Horizontal)
+	.Thickness(FVector2D(CScrollBarWidth, CScrollBarWidth));
+
+	SAssignNew(CollectionSpreadSheetExternalVerticalScrollBar, SScrollBar)
+	.Orientation(Orient_Vertical)
+	.Thickness(FVector2D(CScrollBarWidth, CScrollBarWidth))
+	.Visibility(this, &SCollectionSpreadSheetWidget::GetCollectionSpreadSheetVisibility);
+
+	SAssignNew(TransformOutlinerExternalVerticalScrollBar, SScrollBar)
+	.Orientation(Orient_Vertical)
+	.Thickness(FVector2D(CScrollBarWidth, CScrollBarWidth))
+	.Visibility(this, &SCollectionSpreadSheetWidget::GetTransformOutlinerVisibility);
+
+	SAssignNew(VerticesOutlinerExternalVerticalScrollBar, SScrollBar)
+	.Orientation(Orient_Vertical)
+	.Thickness(FVector2D(CScrollBarWidth, CScrollBarWidth))
+	.Visibility(this, &SCollectionSpreadSheetWidget::GetVerticesOutlinerVisibility);
+
+	SAssignNew(FacesOutlinerExternalVerticalScrollBar, SScrollBar)
+	.Orientation(Orient_Vertical)
+	.Thickness(FVector2D(CScrollBarWidth, CScrollBarWidth))
+	.Visibility(this, &SCollectionSpreadSheetWidget::GetFacesOutlinerVisibility);
 
 	ChildSlot
 	[
@@ -577,76 +440,65 @@ void SCollectionSpreadSheetWidget::Construct(const FArguments& InArgs)
 			.AutoWidth()
 			.Padding(10.0f, 10.0f, 4.0f, 5.0f)
 			[
-				SNew(SButton)
+				SNew(SCheckBox)
+				.Style(FAppStyle::Get(), "ToggleButtonCheckbox")
 				.ToolTipText(FText::FromString("The button pins down the panel. When it pinned down it doesn't react to node selection change."))
-				.ButtonStyle(FAppStyle::Get(), "Button")
-				.HAlign(HAlign_Center)
-				.VAlign(VAlign_Center)
-				.Content()
-				[
-					SNew(STextBlock)
-					.Text_Lambda([this]
+				.IsChecked_Lambda([this]()-> ECheckBoxState
+				{
+					if (bIsPinnedDown)
 					{
-						if (!NodeNameTextBlock->GetText().IsEmpty())
-						{
-							return FText::FromString(bIsPinnedDown ? " X " : " O ");
-						}
-						else
-						{
-							return FText::FromString(" O ");
-						}
-					})
-				]
-				.OnClicked_Lambda([this]
+						return ECheckBoxState::Checked;
+					}
+					return ECheckBoxState::Unchecked;
+				})
+				.OnCheckStateChanged_Lambda([this](ECheckBoxState State)
 				{
 					if (!NodeNameTextBlock->GetText().IsEmpty())
 					{
 						bIsPinnedDown = !bIsPinnedDown;
-		
 						OnPinnedDownChangedDelegate.Broadcast(bIsPinnedDown);
 					}
-
-					return FReply::Handled();
 				})
+				.Padding(2.0f)
+				.HAlign(HAlign_Center)
+				[
+					SNew(SImage)
+					.ColorAndOpacity(FSlateColor::UseForeground())
+					.Image(this, &SCollectionSpreadSheetWidget::GetPinButtonImage)
+				]
 			]
 
 			+ SHorizontalBox::Slot()
 			.AutoWidth()
 			.Padding(4.0f, 10.0f, 10.0f, 5.0f)
 			[
-				SNew(SButton)
+				SNew(SCheckBox)
+				.Style(FAppStyle::Get(), "ToggleButtonCheckbox")
 				.ToolTipText(FText::FromString("The button locks the refresh of the values in the panel."))
-				.ButtonStyle(FAppStyle::Get(), "Button")
-				.HAlign(HAlign_Center)
-				.VAlign(VAlign_Center)
-				.Content()
-				[
-					SNew(STextBlock)
-					.Text_Lambda([this]
+				.IsChecked_Lambda([this]()-> ECheckBoxState
 					{
-						if (!NodeNameTextBlock->GetText().IsEmpty())
-						{
-							return FText::FromString(bIsRefreshLocked ? " L " : " U ");
-						}
-						else
+						if (bIsRefreshLocked)
 						{
-							return FText::FromString(" U ");
+							return ECheckBoxState::Checked;
 						}
+						return ECheckBoxState::Unchecked;
 					})
-				]
-				.OnClicked_Lambda([this]
+				.OnCheckStateChanged_Lambda([this](ECheckBoxState State)
 				{
 					if (!NodeNameTextBlock->GetText().IsEmpty())
 					{
 						bIsRefreshLocked = !bIsRefreshLocked;
-		
 						OnRefreshLockedChangedDelegate.Broadcast(bIsRefreshLocked);
 					}
-		
-					return FReply::Handled();
 				})
+				.Padding(2.0f)
+				.HAlign(HAlign_Center)
+				[
+					SNew(SImage)
+					.ColorAndOpacity(FSlateColor::UseForeground())
+					.Image(this, &SCollectionSpreadSheetWidget::GetLockButtonImage)
+				]
 			]
-
 		]
 
 		+ SVerticalBox::Slot()
@@ -727,16 +579,64 @@ void SCollectionSpreadSheetWidget::Construct(const FArguments& InArgs)
 				.Orientation(Orient_Horizontal)
 				.ExternalScrollbar(SpreadSheetHorizontalScrollBar)
 
+				// All the groups not using a TreeView can use SCollectionSpreadSheet
 				+ SScrollBox::Slot()
 				[
 					SAssignNew(CollectionTable, SCollectionSpreadSheet)
+					.ExternalVerticalScrollBar(CollectionSpreadSheetExternalVerticalScrollBar)
+					.Visibility(this, &SCollectionSpreadSheetWidget::GetCollectionSpreadSheetVisibility)
+				]
+
+				// Displaying Transform group with hierarchy
+				+ SScrollBox::Slot()
+				.FillSize(1)
+				[
+					SAssignNew(TransformOutliner, STransformOutliner)
+					.ExternalVerticalScrollBar(TransformOutlinerExternalVerticalScrollBar)
+					.Visibility(this, &SCollectionSpreadSheetWidget::GetTransformOutlinerVisibility)
+				]
+
+				// Displaying Vertices group with hierarchy
+				+ SScrollBox::Slot()
+				.FillSize(1)
+				[
+					SAssignNew(VerticesOutliner, SVerticesOutliner)
+					.ExternalVerticalScrollBar(VerticesOutlinerExternalVerticalScrollBar)
+					.Visibility(this, &SCollectionSpreadSheetWidget::GetVerticesOutlinerVisibility)
+				]
+
+				// Displaying Faces group with hierarchy
+				+ SScrollBox::Slot()
+				.FillSize(1)
+				[
+					SAssignNew(FacesOutliner, SFacesOutliner)
+					.ExternalVerticalScrollBar(FacesOutlinerExternalVerticalScrollBar)
+					.Visibility(this, &SCollectionSpreadSheetWidget::GetFacesOutlinerVisibility)
 				]
 			]
 
 			+ SHorizontalBox::Slot()
 			.AutoWidth()
 			[
-				CollectionTable->SpreadSheetVerticalScrollBar.ToSharedRef()
+				CollectionSpreadSheetExternalVerticalScrollBar.ToSharedRef()
+			]
+
+			+ SHorizontalBox::Slot()
+			.AutoWidth()
+			[
+				TransformOutlinerExternalVerticalScrollBar.ToSharedRef()
+			]
+
+			+ SHorizontalBox::Slot()
+			.AutoWidth()
+			[
+				VerticesOutlinerExternalVerticalScrollBar.ToSharedRef()
+			]
+
+			+ SHorizontalBox::Slot()
+			.AutoWidth()
+			[
+				FacesOutlinerExternalVerticalScrollBar.ToSharedRef()
 			]
 		]
 
@@ -761,7 +661,6 @@ void SCollectionSpreadSheetWidget::Construct(const FArguments& InArgs)
 	];
 }
 
-
 void SCollectionSpreadSheetWidget::SetData(const FString& InNodeName)
 {
 	NodeName = InNodeName;
@@ -822,7 +721,8 @@ void SCollectionSpreadSheetWidget::SetStatusText()
 {
 	if (!NodeName.IsEmpty())
 	{
-		FString Str = FString::Printf(TEXT("Group has %d elements"), CollectionTable->GetNumItems());
+		FString NumberStr = FString::FormatAsNumber(CollectionTable->GetNumItems());
+		FString Str = "Group has " + NumberStr + " elements";
 		StatusTextBlock->SetText(FText::FromString(Str));
 	}
 	else
@@ -845,3 +745,5 @@ void SCollectionSpreadSheetWidget::UpdateCollectionGroups(const FName& InOutputN
 	}
 }
 
+
+#undef LOCTEXT_NAMESPACE // "CollectionSpreadSheet"
diff --git a/Engine/Plugins/Experimental/Dataflow/Source/DataflowEditor/Private/Dataflow/DataflowCollectionSpreadSheetHelpers.cpp b/Engine/Plugins/Experimental/Dataflow/Source/DataflowEditor/Private/Dataflow/DataflowCollectionSpreadSheetHelpers.cpp
new file mode 100644
--- /dev/null
+++ b/Engine/Plugins/Experimental/Dataflow/Source/DataflowEditor/Private/Dataflow/DataflowCollectionSpreadSheetHelpers.cpp
@@ -0,0 +1,542 @@
+// Copyright Epic Games, Inc. All Rights Reserved.
+
+#include "Dataflow/DataflowCollectionSpreadSheetHelpers.h"
+#include "GeometryCollection/GeometryCollection.h"
+#include "Dataflow/DataflowSettings.h"
+#include "Components/HorizontalBox.h"
+#include "Widgets/Text/STextBlock.h"
+#include "Styling/SlateTypes.h"
+#include "Widgets/Colors/SColorBlock.h"
+#include "Layout/Margin.h"
+
+namespace UE::Dataflow::CollectionSpreadSheetHelpers
+{
+	template<typename T>
+	FString AttributeValueToString(const T& Value)
+	{
+		return Value.ToString();
+	}
+
+	template<typename T>
+	FString AttributeValueToString(const TArray<T>& Array)
+	{
+		FString Out;
+
+		for (int32 Idx = 0; Idx < Array.Num(); ++Idx)
+		{
+			Out += AttributeValueToString(Array[Idx]);
+
+			if (Idx != Array.Num() - 1)
+			{
+				Out += "; ";
+			}
+		}
+
+		return Out;
+	}
+
+	template<typename T>
+	FString AttributeValueToString(const FManagedArrayCollection& InCollection, const FName& InAttributeName, const FName& InGroupName, int32 InIdxColumn)
+	{
+		const TManagedArray<T>* const Array = InCollection.FindAttributeTyped<T>(InAttributeName, InGroupName);
+		if (Array == nullptr)
+		{
+			return FString("<Unknown Attribute>");
+		}
+
+		if (InIdxColumn < 0 || Array->Num() <= InIdxColumn)
+		{
+			return FString("<Index out of bounds>");
+		}
+
+		return AttributeValueToString((*Array)[InIdxColumn]);
+	}
+
+	FString AttributeValueToString(float Value)
+	{
+		return FString::SanitizeFloat(Value, 2);
+	}
+
+	FString AttributeValueToString(int32 Value)
+	{
+		return FString::FromInt(Value);
+	}
+
+	FString AttributeValueToString(FString Value)
+	{
+		return Value;
+	}
+
+	FString AttributeValueToString(FLinearColor Value)
+	{
+		return FString::Printf(TEXT("(R=%.2f G=%.2f B=%.2f A=%.2f)"), Value.R, Value.G, Value.B, Value.A);
+	}
+
+	FString AttributeValueToString(FVector Value)
+	{
+		return FString::Printf(TEXT("(X=%2.2f Y=%2.2f Z=%2.2f)"), Value.X, Value.Y, Value.Z);
+	}
+
+	FString AttributeValueToString(bool Value)
+	{
+		return Value ? FString("true") : FString("false");
+	}
+
+	FString AttributeValueToString(const FConstBitReference& Value)
+	{
+		return Value ? FString("true") : FString("false");
+	}
+
+	FString AttributeValueToString(TSet<int32> Value)
+	{
+		TArray<int32> Array = Value.Array();
+		FString Out;
+
+		for (int32 Idx = 0; Idx < Array.Num(); ++Idx)
+		{
+			Out += AttributeValueToString(Array[Idx]);
+
+			if (Idx != Array.Num() - 1)
+			{
+				Out += " ";
+			}
+		}
+
+		return Out;
+	}
+
+	FString AttributeValueToString(FTransform3f Value)
+	{
+		const FVector3f Translation = Value.GetTranslation();
+		const FVector3f Rotation = Value.GetRotation().Euler();
+		const FVector3f Scale = Value.GetScale3D();
+
+		return FString::Printf(TEXT("T:(%s) R:(%s) S:(%s)"), *Translation.ToString(), *Rotation.ToString(), *Scale.ToString());
+	}
+
+	FString AttributeValueToString(FTransform Value)
+	{
+		const FVector Translation = Value.GetTranslation();
+		const FVector Rotation = Value.GetRotation().Euler();
+		const FVector Scale = Value.GetScale3D();
+
+		return FString::Printf(TEXT("T:(%s) R:(%s) S:(%s)"), *Translation.ToString(), *Rotation.ToString(), *Scale.ToString());
+	}
+
+	FString AttributeValueToString(FBox Value)
+	{
+		FVector Center, Extents;
+		Value.GetCenterAndExtents(Center, Extents);
+		return FString::Printf(TEXT("Center:(%s) Extents:(%s)"), *Center.ToString(), *Extents.ToString());
+	}
+
+	FString AttributeValueToString(FIntVector Value)
+	{
+		return FString::Printf(TEXT("%d %d %d"), Value.X, Value.Y, Value.Z);
+	}
+
+	FString AttributeValueToString(FIntVector4 Value)
+	{
+		return FString::Printf(TEXT("%d %d %d %d"), Value[0], Value[1], Value[2], Value[3]);
+	}
+
+	FString AttributeValueToString(FGuid Value)
+	{
+		return Value.ToString();
+	}
+
+	FString AttributeValueToString(const FManagedArrayCollection& InCollection, const FName& InAttributeName, const FName& InGroupName, int32 InIdxColumn)
+	{
+		const FManagedArrayCollection::EArrayType ArrayType = InCollection.GetAttributeType(InAttributeName, InGroupName);
+
+		FString ValueAsString;
+
+		switch (ArrayType)
+		{
+		case FManagedArrayCollection::EArrayType::FFloatType:
+			ValueAsString = AttributeValueToString<float>(InCollection, InAttributeName, InGroupName, InIdxColumn);
+			break;
+
+		case FManagedArrayCollection::EArrayType::FInt32Type:
+			ValueAsString = AttributeValueToString<int32>(InCollection, InAttributeName, InGroupName, InIdxColumn);
+			break;
+
+		case FManagedArrayCollection::EArrayType::FBoolType:
+			ValueAsString = AttributeValueToString<bool>(InCollection, InAttributeName, InGroupName, InIdxColumn);
+			break;
+
+		case FManagedArrayCollection::EArrayType::FStringType:
+			ValueAsString = AttributeValueToString<FString>(InCollection, InAttributeName, InGroupName, InIdxColumn);
+			break;
+
+		case FManagedArrayCollection::EArrayType::FLinearColorType:
+			ValueAsString = AttributeValueToString<FLinearColor>(InCollection, InAttributeName, InGroupName, InIdxColumn);
+			break;
+
+		case FManagedArrayCollection::EArrayType::FVectorType:
+			ValueAsString = AttributeValueToString<FVector3f>(InCollection, InAttributeName, InGroupName, InIdxColumn);
+			break;
+
+		case FManagedArrayCollection::EArrayType::FVector2DType:
+			ValueAsString = AttributeValueToString<FVector2f>(InCollection, InAttributeName, InGroupName, InIdxColumn);
+			break;
+
+		case FManagedArrayCollection::EArrayType::FVector3dType:
+			ValueAsString = AttributeValueToString<FVector2f>(InCollection, InAttributeName, InGroupName, InIdxColumn);
+			break;
+
+		case FManagedArrayCollection::EArrayType::FIntVectorType:
+			ValueAsString = AttributeValueToString<FIntVector>(InCollection, InAttributeName, InGroupName, InIdxColumn);
+			break;
+
+		case FManagedArrayCollection::EArrayType::FTransformType:
+			ValueAsString = AttributeValueToString<FTransform>(InCollection, InAttributeName, InGroupName, InIdxColumn);
+			break;
+
+		case FManagedArrayCollection::EArrayType::FVector2DArrayType:
+			ValueAsString = AttributeValueToString<TArray<FVector2f>>(InCollection, InAttributeName, InGroupName, InIdxColumn);
+			break;
+
+		case FManagedArrayCollection::EArrayType::FIntArrayType:
+			ValueAsString = AttributeValueToString<TSet<int32>>(InCollection, InAttributeName, InGroupName, InIdxColumn);
+			break;
+
+		case FManagedArrayCollection::EArrayType::FInt32ArrayType:
+			ValueAsString = AttributeValueToString<TArray<int32>>(InCollection, InAttributeName, InGroupName, InIdxColumn);
+			break;
+
+		case FManagedArrayCollection::EArrayType::FFloatArrayType:
+			ValueAsString = AttributeValueToString<TArray<float>>(InCollection, InAttributeName, InGroupName, InIdxColumn);
+			break;
+
+		case FManagedArrayCollection::EArrayType::FBoxType:
+			ValueAsString = AttributeValueToString<FBox>(InCollection, InAttributeName, InGroupName, InIdxColumn);
+			break;
+
+		case FManagedArrayCollection::EArrayType::FTransform3fType:
+			ValueAsString = AttributeValueToString<FTransform3f>(InCollection, InAttributeName, InGroupName, InIdxColumn);
+			break;
+
+		case FManagedArrayCollection::EArrayType::FIntVector4Type:
+			ValueAsString = AttributeValueToString<FIntVector4>(InCollection, InAttributeName, InGroupName, InIdxColumn);
+			break;
+
+		case FManagedArrayCollection::EArrayType::FGuidType:
+			ValueAsString = AttributeValueToString<FGuid>(InCollection, InAttributeName, InGroupName, InIdxColumn);
+			break;
+
+		default:
+			//ensure(false);
+			ValueAsString = "<Unknown Data Type>";
+		}
+
+		return ValueAsString;
+	}
+
+	FColor GetColorPerDepth(uint32 Depth)
+	{
+		const UDataflowSettings* DataflowSettings = GetDefault<UDataflowSettings>();
+
+		const uint32 ColorPerDepthCount = DataflowSettings->TransformLevelColors.LevelColors.Num();
+		return DataflowSettings->TransformLevelColors.LevelColors[Depth % ColorPerDepthCount].ToFColor(true);
+	}
+
+	FSlateColor UpdateItemColorFromCollection(const TSharedPtr<FManagedArrayCollection> InCollection, const FName InGroup, const int32 InItemIndex)
+	{
+		constexpr FLinearColor InvalidColor(0.1f, 0.1f, 0.1f);
+
+		int32 BoneIndex = InItemIndex;
+		if (InGroup == FGeometryCollection::VerticesGroup)
+		{
+			if (InCollection->HasAttribute("BoneMap", FGeometryCollection::VerticesGroup))
+			{
+				const TManagedArray<int32>& BoneMap = InCollection->GetAttribute<int32>("BoneMap", FGeometryCollection::VerticesGroup);
+				BoneIndex = BoneMap[InItemIndex];
+			}
+			else
+			{
+				return InvalidColor;
+			}
+		}
+		else if (InGroup == FGeometryCollection::FacesGroup)
+		{
+			if (InCollection->HasAttribute("Indices", FGeometryCollection::FacesGroup) &&
+				InCollection->HasAttribute("BoneMap", FGeometryCollection::VerticesGroup))
+			{
+				const TManagedArray<FIntVector>& Indices = InCollection->GetAttribute<FIntVector>("Indices", FGeometryCollection::FacesGroup);
+				const int32 VertexIndex = Indices[InItemIndex][0];
+
+				const TManagedArray<int32>& BoneMap = InCollection->GetAttribute<int32>("BoneMap", FGeometryCollection::VerticesGroup);
+				BoneIndex = BoneMap[VertexIndex];
+			}
+			else
+			{
+				return InvalidColor;
+			}
+		}
+
+		if (InCollection->HasAttribute("Level", FGeometryCollection::TransformGroup))
+		{
+			const TManagedArray<int32>& Level = InCollection->GetAttribute<int32>("Level", FGeometryCollection::TransformGroup);
+			return GetColorPerDepth((uint32)Level[BoneIndex]);
+		}
+		else if (InCollection->HasAttribute("SimulationType", FGeometryCollection::TransformGroup))
+		{
+			const TManagedArray<int32>& SimulationType = InCollection->GetAttribute<int32>("SimulationType", FGeometryCollection::TransformGroup);
+			switch (SimulationType[BoneIndex])
+			{
+			case FGeometryCollection::ESimulationTypes::FST_None:
+				return FLinearColor::Green;
+
+			case FGeometryCollection::ESimulationTypes::FST_Rigid:
+			{
+				bool IsVisible = true;
+				if (InCollection->HasAttribute("Visible", FGeometryCollection::TransformGroup))
+				{
+					IsVisible = InCollection->GetAttribute<bool>("Visible", FGeometryCollection::TransformGroup)[BoneIndex];
+				}
+
+				return IsVisible ? FSlateColor::UseForeground() : InvalidColor;
+			}
+			case FGeometryCollection::ESimulationTypes::FST_Clustered:
+				return FColor::Cyan;
+
+			default:
+				ensureMsgf(false, TEXT("Invalid Geometry Collection simulation type encountered."));
+				return InvalidColor;
+			}
+		}
+
+		return InvalidColor;
+	}
+
+	TSharedRef<SWidget> MakeColumnWidget(const TSharedPtr<FManagedArrayCollection> InCollection,
+		const FName InGroup,
+		const FName InAttr,
+		const int32 InItemIndex,
+		const FSlateColor InItemColor)
+	{
+		FName AttrType = UE::Dataflow::CollectionSpreadSheetHelpers::GetArrayTypeString(InCollection->GetAttributeType(InAttr, InGroup));
+
+		if (InAttr == "Index")
+		{
+			return SNew(SHorizontalBox)
+				+ SHorizontalBox::Slot()
+				.AutoWidth()
+				[
+					SNew(STextBlock)
+						.Text(FText::AsNumber(InItemIndex))
+						.ColorAndOpacity(InItemColor)
+				];
+		}
+		else if (InAttr == "SimulationType")
+		{
+			const TArray<FString> SimTypes = { "None", "Rigid", "Clusterd" };	
+			FString AttrValueStr = SimTypes[InCollection->GetAttribute<int32>("SimulationType", InGroup)[InItemIndex]];
+
+			return SNew(SHorizontalBox)
+				+ SHorizontalBox::Slot()
+				.AutoWidth()
+				[
+					SNew(STextBlock)
+					.Text(FText::FromString(AttrValueStr))
+					.ColorAndOpacity(InItemColor)
+					.OverflowPolicy(ETextOverflowPolicy::Ellipsis)
+				];
+		}
+
+		if (AttrType == "LinearColor")
+		{
+			FLinearColor AttrValue = InCollection->GetAttribute<FLinearColor>(InAttr, InGroup)[InItemIndex];
+
+			return SNew(SHorizontalBox)
+				+ SHorizontalBox::Slot()
+				.AutoWidth()
+				[
+					SNew(SColorBlock)
+						.Color(AttrValue)
+						.Size(FVector2D(48, 16))
+						.CornerRadius(2.f)
+				];
+		}
+		else if (AttrType == "Vector2D")
+		{
+			FVector2f AttrValue = InCollection->GetAttribute<FVector2f>(InAttr, InGroup)[InItemIndex];
+
+			return SNew(SHorizontalBox)
+				+ SHorizontalBox::Slot()
+				.Padding(FMargin(1.f, 0.f))
+				.MinWidth(60)
+				.MaxWidth(60)
+				[
+					SNew(STextBlock)
+					.Text(FText::AsNumber(AttrValue.X))
+					.ColorAndOpacity(InItemColor)
+					.Justification(ETextJustify::Right)
+				]
+
+				+ SHorizontalBox::Slot()
+				.Padding(FMargin(1.f, 0.f))
+				.MinWidth(60)
+				.MaxWidth(60)
+				[
+					SNew(STextBlock)
+					.Text(FText::AsNumber(AttrValue.Y))
+					.ColorAndOpacity(InItemColor)
+					.Justification(ETextJustify::Right)
+				];
+		}
+		else if (AttrType == "Vector")
+		{
+			FVector3f AttrValue = InCollection->GetAttribute<FVector3f>(InAttr, InGroup)[InItemIndex];
+
+			return SNew(SHorizontalBox)
+				+ SHorizontalBox::Slot()
+				.Padding(FMargin(1.f, 0.f))
+				.MinWidth(60)
+				.MaxWidth(60)
+				[
+					SNew(STextBlock)
+						.Text(FText::AsNumber(AttrValue.X))
+						.ColorAndOpacity(InItemColor)
+						.Justification(ETextJustify::Right)
+				]
+
+				+ SHorizontalBox::Slot()
+				.Padding(FMargin(1.f, 0.f))
+				.MinWidth(60)
+				.MaxWidth(60)
+				[
+					SNew(STextBlock)
+						.Text(FText::AsNumber(AttrValue.Y))
+						.ColorAndOpacity(InItemColor)
+						.Justification(ETextJustify::Right)
+				]
+
+				+ SHorizontalBox::Slot()
+				.Padding(FMargin(1.f, 0.f))
+				.MinWidth(60)
+				.MaxWidth(60)
+				[
+					SNew(STextBlock)
+						.Text(FText::AsNumber(AttrValue.Z))
+						.ColorAndOpacity(InItemColor)
+						.Justification(ETextJustify::Right)
+				];
+		}
+		else if (AttrType == "IntVector")
+		{
+			FIntVector AttrValue = InCollection->GetAttribute<FIntVector>(InAttr, InGroup)[InItemIndex];
+
+			return SNew(SHorizontalBox)
+				+ SHorizontalBox::Slot()
+				.Padding(FMargin(1.f, 0.f))
+				.MinWidth(60)
+				.MaxWidth(60)
+				[
+					SNew(STextBlock)
+					.Text(FText::AsNumber(AttrValue.X))
+					.ColorAndOpacity(InItemColor)
+					.Justification(ETextJustify::Right)
+				]
+
+				+ SHorizontalBox::Slot()
+				.Padding(FMargin(1.f, 0.f))
+				.MinWidth(60)
+				.MaxWidth(60)
+				[
+					SNew(STextBlock)
+					.Text(FText::AsNumber(AttrValue.Y))
+					.ColorAndOpacity(InItemColor)
+					.Justification(ETextJustify::Right)
+				]
+
+				+ SHorizontalBox::Slot()
+				.Padding(FMargin(1.f, 0.f))
+				.MinWidth(60)
+				.MaxWidth(60)
+				[
+					SNew(STextBlock)
+					.Text(FText::AsNumber(AttrValue.Z))
+					.ColorAndOpacity(InItemColor)
+					.Justification(ETextJustify::Right)
+				];
+		}
+		else if (AttrType == "Transform3f")
+		{
+			constexpr float TextFieldWidth = 50.f;
+			constexpr int32 FractionalDigits = 2;
+
+			struct FTextDisplayInfo
+			{
+				FText Text;
+				ETextJustify::Type Justify;
+				float Width = 0.f;
+			};
+
+			FTransform3f AttrValue = InCollection->GetAttribute<FTransform3f>(InAttr, InGroup)[InItemIndex];
+
+			const FVector3f Translation = AttrValue.GetTranslation();
+			const FVector3f Rotation = AttrValue.GetRotation().Euler();
+			const FVector3f Scale = AttrValue.GetScale3D();
+
+			FNumberFormattingOptions FormattingOptions;
+			FormattingOptions.SetMinimumFractionalDigits(FractionalDigits);
+			FormattingOptions.SetMaximumFractionalDigits(FractionalDigits);
+
+			const FTextDisplayInfo TextDisplayArr[] =
+			{
+				{ FText::FromString(TEXT("T:[")), ETextJustify::Right, 20 },
+				{ FText::AsNumber(Translation.X, &FormattingOptions), ETextJustify::Right, TextFieldWidth },
+				{ FText::AsNumber(Translation.Y, &FormattingOptions), ETextJustify::Right, TextFieldWidth },
+				{ FText::AsNumber(Translation.Z, &FormattingOptions), ETextJustify::Right, TextFieldWidth },
+				{ FText::FromString(TEXT("]")), ETextJustify::Left, 15 },
+				{ FText::FromString(TEXT("R:[")), ETextJustify::Right, 20 },
+				{ FText::AsNumber(Rotation.X, &FormattingOptions), ETextJustify::Right, TextFieldWidth },
+				{ FText::AsNumber(Rotation.Y, &FormattingOptions), ETextJustify::Right, TextFieldWidth },
+				{ FText::AsNumber(Rotation.Z, &FormattingOptions), ETextJustify::Right, TextFieldWidth },
+				{ FText::FromString(TEXT("]")), ETextJustify::Left, 15 },
+				{ FText::FromString(TEXT("S:[")), ETextJustify::Right, 20 },
+				{ FText::AsNumber(Scale.X, &FormattingOptions), ETextJustify::Right, TextFieldWidth },
+				{ FText::AsNumber(Scale.Y, &FormattingOptions), ETextJustify::Right, TextFieldWidth },
+				{ FText::AsNumber(Scale.Z, &FormattingOptions), ETextJustify::Right, TextFieldWidth },
+				{ FText::FromString(TEXT("]")), ETextJustify::Left, 5 }
+			};
+
+			constexpr int32 TextArrSize = sizeof(TextDisplayArr) / sizeof(FTextDisplayInfo);
+
+			TSharedRef<SHorizontalBox> Box = SNew(SHorizontalBox);
+
+			for (int32 Idx = 0; Idx < TextArrSize; ++Idx)
+			{
+				Box->AddSlot()
+					.Padding(FMargin(1.f, 0.f))
+					.MinWidth(TextDisplayArr[Idx].Width)
+					.MaxWidth(TextDisplayArr[Idx].Width)
+					[
+						SNew(STextBlock)
+						.Text(TextDisplayArr[Idx].Text)
+						.ColorAndOpacity(InItemColor)
+						.Justification(TextDisplayArr[Idx].Justify)
+					];
+			}
+
+			return Box;
+		}
+		else
+		{
+			FString AttrValueStr = UE::Dataflow::CollectionSpreadSheetHelpers::AttributeValueToString(*InCollection.Get(), InAttr, InGroup, InItemIndex);
+
+			return SNew(SHorizontalBox)
+				+ SHorizontalBox::Slot()
+				.AutoWidth()
+				[
+					SNew(STextBlock)
+						.Text(FText::FromString(AttrValueStr))
+						.ColorAndOpacity(InItemColor)
+						.OverflowPolicy(ETextOverflowPolicy::Ellipsis)
+				];
+		}
+	}
+
+}
\ No newline at end of file
diff --git a/Engine/Plugins/Experimental/Dataflow/Source/DataflowEditor/Private/Dataflow/DataflowFacesOutlinerWidget.cpp b/Engine/Plugins/Experimental/Dataflow/Source/DataflowEditor/Private/Dataflow/DataflowFacesOutlinerWidget.cpp
new file mode 100644
--- /dev/null
+++ b/Engine/Plugins/Experimental/Dataflow/Source/DataflowEditor/Private/Dataflow/DataflowFacesOutlinerWidget.cpp
@@ -0,0 +1,462 @@
+// Copyright Epic Games, Inc. All Rights Reserved.
+
+#include "Dataflow/DataflowFacesOutlinerWidget.h"
+#include "Widgets/Input/SButton.h"
+#include "Widgets/Colors/SColorBlock.h"
+#include "Widgets/Layout/SScrollBox.h"
+#include "Dataflow/DataflowSettings.h"
+#include "GeometryCollection/GeometryCollection.h"
+#include "Dataflow/DataflowCollectionSpreadSheetHelpers.h"
+#include "GeometryCollection/ManagedArrayCollection.h"
+#include "Widgets/Views/STableRow.h"
+#include "Widgets/Text/STextBlock.h"
+#include "Widgets/Input/SComboBox.h"
+
+#define LOCTEXT_NAMESPACE "FacesOutliner"
+
+// --- FGeometryCollectionFacesTreeItem ---
+
+void FGeometryCollectionFacesTreeItem::GenerateContextMenu(UToolMenu* Menu, SFacesOutliner& Outliner)
+{
+}
+
+// --- FGeometryCollectionFacesTreeItemGeometry ---
+
+TSharedRef<ITableRow> FGeometryCollectionFacesTreeItemGeometry::MakeTreeRowWidget(const TSharedRef<STableViewBase>& InOwnerTable, bool bNoExtraColumn)
+{
+	FString NumVertsString = "";
+	FString NumFacesString = "";
+	FString TransformIdxStr = "";
+	FString GeometryIdxStr = "";
+	FString BoneNameStr = "";
+
+	if (Collection.IsValid() &&
+		Collection->HasAttribute("TransformIndex", FGeometryCollection::GeometryGroup) &&
+		Collection->HasAttribute("TransformToGeometryIndex", FGeometryCollection::TransformGroup) &&
+		Collection->HasAttribute("VertexCount", FGeometryCollection::GeometryGroup) &&
+		Collection->HasAttribute("FaceCount", FGeometryCollection::GeometryGroup) &&
+		Collection->HasAttribute("BoneName", FGeometryCollection::TransformGroup))
+	{
+		const TManagedArray<int32>& TransformIndices = Collection->GetAttribute<int32>("TransformIndex", FGeometryCollection::GeometryGroup);
+		const TManagedArray<int32>& TransformToGeometryIndices = Collection->GetAttribute<int32>("TransformToGeometryIndex", FGeometryCollection::TransformGroup);
+		const TManagedArray<int32>& VertexCounts = Collection->GetAttribute<int32>("VertexCount", FGeometryCollection::GeometryGroup);
+		const TManagedArray<int32>& FaceCounts = Collection->GetAttribute<int32>("FaceCount", FGeometryCollection::GeometryGroup);
+		const TManagedArray<FString>& BoneNames = Collection->GetAttribute<FString>("BoneName", FGeometryCollection::TransformGroup);
+
+		const int32 TransformIdx = TransformIndices[GeometryIdx];
+		TransformIdxStr = FString::FormatAsNumber(TransformIdx);
+
+		GeometryIdxStr = FString::FormatAsNumber(GeometryIdx);
+
+		const int32 NumVerts = VertexCounts[GeometryIdx];
+		NumVertsString = FString::FormatAsNumber(NumVerts);
+
+		const int32 NumFaces = FaceCounts[GeometryIdx];
+		NumFacesString = FString::FormatAsNumber(NumFaces);
+
+		BoneNameStr = BoneNames[TransformIdx];
+	}
+
+	return SNew(STableRow<FGeometryCollectionFacesTreeItemPtr>, InOwnerTable)
+		.Content()
+		[
+			SNew(STextBlock).Text(FText::FromString("Geometry Idx: " + GeometryIdxStr + "     BoneName: " + BoneNameStr + "     Transform Idx: " + TransformIdxStr + "     Faces: " + NumFacesString + "     Verts: " + NumVertsString))
+		];
+}
+
+void FGeometryCollectionFacesTreeItemGeometry::GetChildren(FGeometryCollectionFacesTreeItemList& OutChildren)
+{
+	OutChildren = ChildItems;
+}
+
+FGeometryCollectionFacesTreeItemPtr FGeometryCollectionFacesTreeItemGeometry::GetItemFromFaceIndex(int32 InFaceIndex) const
+{
+	return ItemsByFaceIndex.FindRef(InFaceIndex);
+}
+
+void FGeometryCollectionFacesTreeItemGeometry::GetChildrenForFace(FGeometryCollectionFacesTreeItemFace& FaceItem, FGeometryCollectionFacesTreeItemList& OutChildren) const
+{
+	if (!Collection.IsValid())
+	{
+		return;
+	}
+}
+
+bool FGeometryCollectionFacesTreeItemGeometry::HasChildrenForFace(const FGeometryCollectionFacesTreeItemFace& FaceItem) const
+{
+	if (!Collection.IsValid())
+	{
+		return false;
+	}
+
+	return false;
+}
+
+void FGeometryCollectionFacesTreeItemGeometry::ExpandAll()
+{
+	TreeView->SetItemExpansion(AsShared(), true);
+
+	for (const auto& Elem : ItemsByFaceIndex)
+	{
+		TreeView->SetItemExpansion(Elem.Value, true);
+	}
+}
+
+void FGeometryCollectionFacesTreeItemGeometry::RegenerateChildren()
+{
+	if (Collection.IsValid())
+	{
+		ItemsByFaceIndex.Empty();
+		ChildItems.Empty();
+
+		if (Collection.IsValid() &&
+			Collection->HasAttribute("FaceStart", FGeometryCollection::GeometryGroup) &&
+			Collection->HasAttribute("FaceCount", FGeometryCollection::GeometryGroup))
+		{
+			const TManagedArray<int32>& FaceStarts = Collection->GetAttribute<int32>("FaceStart", FGeometryCollection::GeometryGroup);
+			const TManagedArray<int32>& FaceCounts = Collection->GetAttribute<int32>("FaceCount", FGeometryCollection::GeometryGroup);
+
+			const int32 FaceStart = FaceStarts[GeometryIdx];
+			const int32 FaceCount = FaceCounts[GeometryIdx];
+
+			RootIndex = FGeometryCollection::Invalid;
+
+			// Add a sub item to the outliner tree for each of the bones/chunks in this GeometryCollection
+			for (int32 Index = 0; Index < FaceCount; ++Index)
+			{
+				const int32 FaceIndex = FaceStart + Index;
+
+				TSharedRef<FGeometryCollectionFacesTreeItemFace> NewItem = MakeShared<FGeometryCollectionFacesTreeItemFace>(FaceIndex, this);
+
+				ChildItems.Add(NewItem);
+				ItemsByFaceIndex.Add(FaceIndex, NewItem);
+			}
+		}
+	}
+}
+
+void FGeometryCollectionFacesTreeItemGeometry::RequestTreeRefresh()
+{
+	if (TreeView.IsValid())
+	{
+		TreeView->RequestTreeRefresh();
+	}
+}
+
+bool FGeometryCollectionFacesTreeItemGeometry::FilterBoneIndex(int32 BoneIndex) const
+{
+	return true;
+}
+
+bool FGeometryCollectionFacesTreeItemGeometry::IsValid() const
+{
+	if (bInvalidated)
+	{
+		return false;
+	}
+
+	if (Collection)
+	{
+		return true;
+	}
+
+	return false;
+}
+
+// --- FGeometryCollectionFacesTreeItemGeometry ---
+
+
+TSharedPtr<FManagedArrayCollection> FGeometryCollectionFacesTreeItemFace::GetCollection() const
+{
+	ensure(ParentCollectionItem);
+	return ParentCollectionItem->GetCollection();
+}
+
+void FGeometryCollectionFacesTreeItemFace::UpdateItemColorFromCollection()
+{
+	ItemColor = UE::Dataflow::CollectionSpreadSheetHelpers::UpdateItemColorFromCollection(GetCollection(), FGeometryCollection::FacesGroup, GetFaceIndex());
+}
+
+TSharedRef<ITableRow> FGeometryCollectionFacesTreeItemFace::MakeTreeRowWidget(const TSharedRef<STableViewBase>& InOwnerTable, bool bIsPinned)
+{
+	UpdateItemColorFromCollection();
+
+	TSharedPtr<FManagedArrayCollection> Collection = GetCollection();
+
+	return SNew(SFacesOutlinerRow, InOwnerTable, SharedThis(this));
+}
+
+bool FGeometryCollectionFacesTreeItemFace::IsValidFace() const
+{
+	TSharedPtr<FManagedArrayCollection> Collection = GetCollection();
+
+	if (ParentCollectionItem && ParentCollectionItem->IsValid())
+	{
+		return FaceIndex >= 0 && FaceIndex < Collection->NumElements(FGeometryCollection::FacesGroup);
+	}
+
+	return false;
+}
+
+TSharedRef<SWidget> FGeometryCollectionFacesTreeItemFace::MakeIndexColumnWidget() const
+{
+	return SNew(SHorizontalBox)
+		+ SHorizontalBox::Slot()
+		.AutoWidth()
+		[
+			SNew(STextBlock)
+			.Text(FText::AsNumber(FaceIndex))
+			.ColorAndOpacity(ItemColor)
+		];
+}
+
+TSharedRef<SWidget> FGeometryCollectionFacesTreeItemFace::MakeEmptyColumnWidget() const
+{
+	return SNew(SHorizontalBox)
+		+ SHorizontalBox::Slot()
+		.Padding(12.f, 0.f)
+		[
+			SNew(STextBlock)
+			.ColorAndOpacity(ItemColor)
+		];
+}
+
+void FGeometryCollectionFacesTreeItemFace::GetChildren(FGeometryCollectionFacesTreeItemList& OutChildren)
+{
+	ParentCollectionItem->GetChildrenForFace(*this, OutChildren);
+}
+
+bool FGeometryCollectionFacesTreeItemFace::HasChildren() const
+{
+	return ParentCollectionItem->HasChildrenForFace(*this);
+}
+
+// --- SFacesOutlinerRow ---
+
+TSharedRef<SWidget> SFacesOutlinerRow::GenerateWidgetForColumn(const FName& ColumnName)
+{
+	TSharedPtr<FManagedArrayCollection> Collection = Item->GetCollection();
+	FName AttrType = UE::Dataflow::CollectionSpreadSheetHelpers::GetArrayTypeString(Collection->GetAttributeType(ColumnName, FGeometryCollection::FacesGroup));
+
+	// This can happen because sometimes slate retains old items until the next tick, and keeps calling callbacks on them until then
+	if (!Item->IsValidFace())
+	{
+		return Item->MakeEmptyColumnWidget();
+	}
+
+	if (ColumnName == "Index")
+	{
+		const TSharedPtr<SWidget> NameWidget = Item->MakeIndexColumnWidget();
+		return SNew(SHorizontalBox)
+			+ SHorizontalBox::Slot()
+			.HAlign(HAlign_Left)
+			.VAlign(VAlign_Center)
+			.AutoWidth()
+			[
+				SNew(SExpanderArrow, SharedThis(this))
+					.ShouldDrawWires(true)
+			]
+			+ SHorizontalBox::Slot()
+			[
+				NameWidget.ToSharedRef()
+			];
+	}
+
+	return UE::Dataflow::CollectionSpreadSheetHelpers::MakeColumnWidget(Collection,
+		FGeometryCollection::FacesGroup,
+		ColumnName,
+		Item->GetFaceIndex(),
+		Item->GetItemColor());
+}
+
+// --- SFacesOutliner ---
+
+void SFacesOutliner::Construct(const FArguments& InArgs)
+{
+	bPerformingSelection = false;
+
+	HeaderRowWidget =
+		SNew(SHeaderRow)
+		.Visibility(EVisibility::Visible);
+
+	RegenerateHeader();
+
+	ChildSlot
+	[
+		SNew(SVerticalBox)
+		+ SVerticalBox::Slot()
+		.Padding(FMargin(0.0f, 3.f))
+		[
+			SNew(SOverlay)
+			+ SOverlay::Slot()
+			[
+				SAssignNew(TreeView, STreeView<FGeometryCollectionFacesTreeItemPtr>)
+				.TreeItemsSource(reinterpret_cast<FGeometryCollectionFacesTreeItemList*>(&RootNodes))
+				.OnSelectionChanged(this, &SFacesOutliner::OnSelectionChanged)
+				.OnGenerateRow(this, &SFacesOutliner::MakeTreeRowWidget)
+				.OnGetChildren(this, &SFacesOutliner::OnGetChildren)
+				.OnContextMenuOpening(this, &SFacesOutliner::OnOpenContextMenu)
+				.AllowInvisibleItemSelection(true)
+				.ShouldStackHierarchyHeaders(true)
+				.OnGeneratePinnedRow(this, &SFacesOutliner::OnGeneratePinnedRowWidget, true)
+				.HighlightParentNodesForSelection(true)
+				.OnSetExpansionRecursive(this, &SFacesOutliner::ExpandRecursive)
+				.HeaderRow(HeaderRowWidget)
+				.ExternalScrollbar(InArgs._ExternalVerticalScrollBar)
+			]
+		]
+	];
+}
+
+void SFacesOutliner::RegenerateHeader()
+{
+	constexpr int32 CustomColumnWidth = 120;
+	constexpr int32 CWidthMult = 9;
+
+	HeaderRowWidget->ClearColumns();
+
+	if (!Collection.IsValid() || Collection->NumElements(FGeometryCollection::FacesGroup) == 0)
+	{
+		return;
+	}
+
+	TArray<UE::Dataflow::CollectionSpreadSheetHelpers::FAttrInfo> AttrInfo;
+
+	// Add "Index" manually
+	AttrInfo.Add({ "Index", "int32" });
+	// Add all the other attrs
+	for (FName Attr : Collection->AttributeNames(FGeometryCollection::FacesGroup))
+	{
+		AttrInfo.Add({ Attr, UE::Dataflow::CollectionSpreadSheetHelpers::GetArrayTypeString(Collection->GetAttributeType(Attr, FGeometryCollection::FacesGroup)).ToString()});
+	}
+
+	for (int32 Idx = 0; Idx < AttrInfo.Num(); ++Idx)
+	{
+		int32 ColumnWidth = CustomColumnWidth;
+		const FString ColumnNameStr = AttrInfo[Idx].Name.ToString();
+
+		if (ColumnNameStr == "Index")
+		{
+			ColumnWidth = CustomColumnWidth;
+
+			HeaderRowWidget->AddColumn(
+				SHeaderRow::Column(AttrInfo[Idx].Name)
+				.DefaultLabel(FText::FromName(AttrInfo[Idx].Name))
+				.ManualWidth(ColumnWidth)
+				.HAlignCell(HAlign_Left)
+				.HAlignHeader(HAlign_Center)
+				.VAlignCell(VAlign_Center));
+		}
+		else
+		{
+			int32 ColumnNameStrLen = ColumnNameStr.Len() * CWidthMult;
+
+			if (UE::Dataflow::CollectionSpreadSheetHelpers::AttrTypeWidthMap.Contains(AttrInfo[Idx].Type))
+			{
+				ColumnWidth = UE::Dataflow::CollectionSpreadSheetHelpers::AttrTypeWidthMap[AttrInfo[Idx].Type];
+			}
+
+			ColumnWidth = ColumnNameStrLen > ColumnWidth ? ColumnNameStrLen : ColumnWidth;
+
+			HeaderRowWidget->AddColumn(
+				SHeaderRow::Column(AttrInfo[Idx].Name)
+				.DefaultLabel(FText::FromName(AttrInfo[Idx].Name))
+				.ManualWidth(ColumnWidth)
+				.HAlignCell(HAlign_Center)
+				.HAlignHeader(HAlign_Center)
+				.VAlignCell(VAlign_Center));
+		}
+	}
+}
+
+void SFacesOutliner::RegenerateItems()
+{
+	TreeView->RebuildList();
+}
+
+TSharedRef<ITableRow> SFacesOutliner::MakeTreeRowWidget(FGeometryCollectionFacesTreeItemPtr InItem, const TSharedRef<STableViewBase>& InOwnerTable)
+{
+	return InItem->MakeTreeRowWidget(InOwnerTable);
+}
+
+TSharedRef<ITableRow> SFacesOutliner::OnGeneratePinnedRowWidget(FGeometryCollectionFacesTreeItemPtr InItem, const TSharedRef<STableViewBase>& InOwnerTable, bool bPinned)
+{
+	return InItem->MakeTreeRowWidget(InOwnerTable, true);
+}
+
+void SFacesOutliner::OnGetChildren(FGeometryCollectionFacesTreeItemPtr InItem, TArray<FGeometryCollectionFacesTreeItemPtr>& OutChildren)
+{
+	InItem->GetChildren(OutChildren);
+}
+
+TSharedPtr<SWidget> SFacesOutliner::OnOpenContextMenu()
+{
+	return TSharedPtr<SWidget>();
+}
+
+void SFacesOutliner::UpdateGeometryCollection()
+{
+	TreeView->RequestTreeRefresh();
+	ExpandAll();
+}
+
+void SFacesOutliner::SetCollection(const TSharedPtr<FManagedArrayCollection>& InCollection)
+{
+	// Clear the cached Tree ItemSelection without affecting the SelectedBones as 
+	// we want to refresh the tree selection using selected bones
+	TGuardValue<bool> ExternalSelectionGuard(bPerformingSelection, true);
+	TreeView->ClearSelection();
+
+	// explicitly mark the root nodes as invalid before emptying, so we know we can safely ignore them in case slate still triggers callbacks for them (they will not be deleted until the tree view refresh, on tick)
+	for (TSharedPtr<FGeometryCollectionFacesTreeItemGeometry>& RootNode : RootNodes)
+	{
+		if (RootNode)
+		{
+			RootNode->Invalidate();
+		}
+	}
+	RootNodes.Empty();
+
+	const int32 NumGeoms = InCollection->NumElements(FGeometryCollection::GeometryGroup);
+	
+	for (int32 Idx = 0; Idx < NumGeoms; ++Idx)
+	{
+		RootNodes.Add(MakeShared<FGeometryCollectionFacesTreeItemGeometry>(InCollection, Idx, TreeView));
+	}
+
+	TreeView->RequestTreeRefresh();
+	ExpandAll();
+
+	Collection = InCollection;
+}
+
+void SFacesOutliner::ExpandAll()
+{
+	for (TSharedPtr<FGeometryCollectionFacesTreeItemGeometry> ItemPtr : RootNodes)
+	{
+		ItemPtr->ExpandAll();
+	}
+}
+
+void SFacesOutliner::ExpandRecursive(TSharedPtr<FGeometryCollectionFacesTreeItem> ItemPtr, bool bInExpansionState) const
+{
+	TreeView->SetItemExpansion(ItemPtr, bInExpansionState);
+
+	FGeometryCollectionFacesTreeItemList ItemChildren;
+	ItemPtr->GetChildren(ItemChildren);
+	for (auto& Child : ItemChildren)
+	{
+		ExpandRecursive(Child, bInExpansionState);
+	}
+}
+
+int32 SFacesOutliner::GetBoneSelectionCount() const
+{
+	return TreeView->GetSelectedItems().Num();
+}
+
+void SFacesOutliner::OnSelectionChanged(FGeometryCollectionFacesTreeItemPtr Item, ESelectInfo::Type SelectInfo)
+{
+}
+
+#undef LOCTEXT_NAMESPACE // "FacesOutliner"
diff --git a/Engine/Plugins/Experimental/Dataflow/Source/DataflowEditor/Private/Dataflow/DataflowTransformOutlinerWidget.cpp b/Engine/Plugins/Experimental/Dataflow/Source/DataflowEditor/Private/Dataflow/DataflowTransformOutlinerWidget.cpp
new file mode 100644
--- /dev/null
+++ b/Engine/Plugins/Experimental/Dataflow/Source/DataflowEditor/Private/Dataflow/DataflowTransformOutlinerWidget.cpp
@@ -0,0 +1,630 @@
+// Copyright Epic Games, Inc. All Rights Reserved.
+
+#include "Dataflow/DataflowTransformOutlinerWidget.h"
+#include "Widgets/Input/SButton.h"
+#include "Widgets/Colors/SColorBlock.h"
+#include "Widgets/Layout/SScrollBox.h"
+#include "Dataflow/DataflowSettings.h"
+#include "GeometryCollection/GeometryCollection.h"
+#include "Dataflow/DataflowCollectionSpreadSheetHelpers.h"
+#include "ToolMenus.h"
+
+#define LOCTEXT_NAMESPACE "TransformOutliner"
+
+// --- FGeometryCollectionTreeItem ---
+
+void FGeometryCollectionTreeItem::GenerateContextMenu(UToolMenu* Menu, STransformOutliner& Outliner)
+{
+	auto SharedOutliner = StaticCastSharedRef<STransformOutliner>(Outliner.AsShared());
+
+	FToolMenuSection& Section = Menu->AddSection("Section");
+	Section.AddMenuEntry("ExpandAll", LOCTEXT("ExpandAll", "Expand All Levels"), FText(), FSlateIcon(FAppStyle::GetAppStyleSetName(), "Icons.Visible"), FUIAction(FExecuteAction::CreateRaw(&Outliner, &STransformOutliner::ContextMenuExpandAll)));
+	Section.AddMenuEntry("PreviousLevel", LOCTEXT("PreviousLevel", "Decrement Display Level"), FText(), FSlateIcon(FAppStyle::GetAppStyleSetName(), "Icons.Minus"), FUIAction(FExecuteAction::CreateRaw(&Outliner, &STransformOutliner::ContextMenuPreviousLevel)));
+	Section.AddMenuEntry("NextLevel", LOCTEXT("NextLevel", "Increment Display Level"), FText(), FSlateIcon(FAppStyle::GetAppStyleSetName(), "Icons.Plus"), FUIAction(FExecuteAction::CreateRaw(&Outliner, &STransformOutliner::ContextMenuNextLevel)));
+}
+
+// for now use the transform group as this simplify the attribute copy 
+static const FName DataCollectionGroup = FGeometryCollection::TransformGroup;
+
+template <typename T>
+static T GetAttributeValue(const TManagedArrayAccessor<T>& Attribute, int32 Index, T Default)
+{
+	return (Attribute.IsValid()) ? Attribute.Get()[Index] : Default;
+}
+
+// --- FGeometryCollectionTreeItemCollection ---
+
+TSharedRef<ITableRow> FGeometryCollectionTreeItemCollection::MakeTreeRowWidget(const TSharedRef<STableViewBase>& InOwnerTable, bool bNoExtraColumn)
+{
+	FString Output = OutputName.ToString();
+	FString NumVertsString = "";
+	FString NumFacesString = "";
+
+	if (Collection.IsValid())
+	{
+		const int32 NumVerts = Collection->NumElements(FGeometryCollection::VerticesGroup);
+		NumVertsString = FString::FormatAsNumber(NumVerts);
+		const int32 NumFaces = Collection->NumElements(FGeometryCollection::FacesGroup);
+		NumFacesString = FString::FormatAsNumber(NumFaces);
+	}
+
+	return SNew(STableRow<FGeometryCollectionTreeItemPtr>, InOwnerTable)
+		.Content()
+		[
+			SNew(STextBlock).Text(FText::FromString(Output + "      Faces: " + NumFacesString + "    Verts: " + NumVertsString))
+		];
+}
+
+void FGeometryCollectionTreeItemCollection::GetChildren(FGeometryCollectionTreeItemList& OutChildren)
+{
+	OutChildren = ChildItems;
+}
+
+FGeometryCollectionTreeItemPtr FGeometryCollectionTreeItemCollection::GetItemFromBoneIndex(int32 BoneIndex) const
+{
+	return ItemsByBoneIndex.FindRef(BoneIndex);
+}
+
+void FGeometryCollectionTreeItemCollection::GetChildrenForBone(FGeometryCollectionTreeItemBone& BoneItem, FGeometryCollectionTreeItemList& OutChildren) const
+{
+	if (!Collection.IsValid())
+	{
+		return;
+	}
+
+	if (Collection->HasAttribute("Children", FGeometryCollection::TransformGroup))
+	{
+		const int32 BoneIndex = BoneItem.GetBoneIndex();
+		if (ensure(BoneIndex >= 0 && BoneIndex < Collection->NumElements(FGeometryCollection::TransformGroup)))
+		{
+			const TManagedArray<TSet<int32>>& Children = Collection->GetAttribute<TSet<int32>>("Children", FGeometryCollection::TransformGroup);
+			for (const int32 ChildIndex : Children[BoneIndex])
+			{
+				FGeometryCollectionTreeItemPtr ChildPtr = ItemsByBoneIndex.FindRef(ChildIndex);
+				if (ChildPtr.IsValid())
+				{
+					OutChildren.Add(ChildPtr);
+				}
+			}
+		}
+	}
+}
+
+bool FGeometryCollectionTreeItemCollection::HasChildrenForBone(const FGeometryCollectionTreeItemBone& BoneItem) const
+{
+	if (!Collection.IsValid())
+	{
+		return false;
+	}
+
+	if (Collection->HasAttribute("Children", FGeometryCollection::TransformGroup))
+	{
+		const int32 BoneIndex = BoneItem.GetBoneIndex();
+		if (ensure(BoneIndex >= 0 && BoneIndex < Collection->NumElements(FGeometryCollection::TransformGroup)))
+		{
+			const TManagedArray<TSet<int32>>& Children = Collection->GetAttribute<TSet<int32>>("Children", FGeometryCollection::TransformGroup);
+			return Children[(BoneIndex)].Num() > 0;
+		}
+	}
+
+	return false;
+}
+
+void FGeometryCollectionTreeItemCollection::ExpandAll()
+{
+	TreeView->SetItemExpansion(AsShared(), true);
+
+	for (const auto& Elem : ItemsByBoneIndex)
+	{
+		TreeView->SetItemExpansion(Elem.Value, true);
+	}
+}
+
+void FGeometryCollectionTreeItemCollection::ExpandByLevel(const int32 InLevel)
+{
+	if (Collection->HasAttribute("Level", DataCollectionGroup))
+	{
+		const TManagedArray<int32>& Levels = Collection->GetAttribute<int32>("Level", DataCollectionGroup);
+
+		TreeView->SetItemExpansion(AsShared(), true);
+
+		for (const auto& Elem : ItemsByBoneIndex)
+		{
+			const int32 BoneIndex = Elem.Key;
+
+			if (Levels[BoneIndex] < InLevel)
+			{
+				TreeView->SetItemExpansion(Elem.Value, true);
+			}
+			else
+			{
+				TreeView->SetItemExpansion(Elem.Value, false);
+			}
+		}
+	}
+}
+
+void FGeometryCollectionTreeItemCollection::RegenerateChildren()
+{
+	if (Collection.IsValid())
+	{
+		ItemsByBoneIndex.Empty();
+		ChildItems.Empty();
+
+		if (Collection)
+		{
+			if (Collection->HasAttribute("Parent", FGeometryCollection::TransformGroup))
+			{
+				const int32 NumElements = Collection->NumElements(FGeometryCollection::TransformGroup);
+				const TManagedArray<int32>& Parents = Collection->GetAttribute<int32>("Parent", FGeometryCollection::TransformGroup);
+
+				RootIndex = FGeometryCollection::Invalid;
+
+				// Add a sub item to the outliner tree for each of the bones/chunks in this GeometryCollection
+					for (int32 BoneIndex = 0; BoneIndex < NumElements; BoneIndex++)
+					{
+						if (FilterBoneIndex(BoneIndex))
+						{
+							TSharedRef<FGeometryCollectionTreeItemBone> NewItem = MakeShared<FGeometryCollectionTreeItemBone>(BoneIndex, this);
+							if (Parents[BoneIndex] == RootIndex)
+							{
+								// The actual children directly beneath this node are the ones without a parent.  The rest are children of children
+								ChildItems.Add(NewItem);
+							}
+
+							ItemsByBoneIndex.Add(BoneIndex, NewItem);
+						}
+					}
+			}
+		}
+	}
+}
+
+void FGeometryCollectionTreeItemCollection::RequestTreeRefresh()
+{
+	if (TreeView.IsValid())
+	{
+		TreeView->RequestTreeRefresh();
+	}
+}
+
+bool FGeometryCollectionTreeItemCollection::FilterBoneIndex(int32 BoneIndex) const
+{
+	if (Collection.IsValid())
+	{
+		if (Collection->HasAttribute("SimulationType", FGeometryCollection::TransformGroup) &&
+			Collection->HasAttribute("Children", FGeometryCollection::TransformGroup))
+		{
+			const TManagedArray<int32>& SimTypes = Collection->GetAttribute<int32>("SimulationType", FGeometryCollection::TransformGroup);
+			const TManagedArray<TSet<int32>>& Children = Collection->GetAttribute<TSet<int32>>("Children", FGeometryCollection::TransformGroup);
+
+			bool bHasChildren = Children[BoneIndex].Num() > 0;
+
+			if (SimTypes[BoneIndex] != FGeometryCollection::ESimulationTypes::FST_Clustered)
+			{
+				// We only display cluster nodes deeper than the view level.
+				const int32 FractureLevel = -1; // Make all the levels visible
+
+				if (FractureLevel >= 0 && Collection->HasAttribute("Level", FTransformCollection::TransformGroup))
+				{
+					const TManagedArray<int32>& Level = Collection->GetAttribute<int32>("Level", FTransformCollection::TransformGroup);
+					int32 BoneLevel = Level[BoneIndex];
+					// bone is not at the right level itself and doesn't have child(ren) at the right level
+					if (BoneLevel != FractureLevel && (!bHasChildren || BoneLevel + 1 != FractureLevel))
+					{
+						return false;
+					}
+				}
+			}
+		}
+	}
+
+	return true;
+}
+
+bool FGeometryCollectionTreeItemCollection::IsValid() const
+{
+	if (bInvalidated)
+	{
+		return false;
+	}
+
+	if (Collection)
+	{
+		return true;
+	}
+
+	return false;
+}
+
+// --- FGeometryCollectionTreeItemBone ---
+
+
+TSharedPtr<FManagedArrayCollection> FGeometryCollectionTreeItemBone::GetCollection() const
+{
+	ensure(ParentCollectionItem);
+	return ParentCollectionItem->GetCollection();
+}
+
+void FGeometryCollectionTreeItemBone::UpdateItemColorFromCollection()
+{
+	ItemColor = UE::Dataflow::CollectionSpreadSheetHelpers::UpdateItemColorFromCollection(GetCollection(), FGeometryCollection::TransformGroup, GetBoneIndex());
+}
+
+TSharedRef<ITableRow> FGeometryCollectionTreeItemBone::MakeTreeRowWidget(const TSharedRef<STableViewBase>& InOwnerTable, bool bIsPinned)
+{
+	UpdateItemColorFromCollection();
+
+	TSharedPtr<FManagedArrayCollection> Collection = GetCollection();
+
+	// This gets displayed in the pinned rows
+	if (bIsPinned)
+	{
+		return SNew(STableRow<FGeometryCollectionTreeItemPtr>, InOwnerTable)
+			.Content()
+			[
+				SNew(SHorizontalBox)
+					+ SHorizontalBox::Slot()
+					.Padding(2.0f, 4.0f)
+					.AutoWidth()
+					[
+						SNew(STextBlock)
+						.Text(FText::AsNumber(BoneIndex))
+						.ColorAndOpacity(ItemColor)
+					]
+					+ SHorizontalBox::Slot()
+					.Padding(2.0f, 4.0f)
+					.AutoWidth()
+					[
+						SNew(STextBlock)
+						.Text(FText::FromString(Collection->GetAttribute<FString>("BoneName", DataCollectionGroup)[BoneIndex]))
+						.ColorAndOpacity(ItemColor)
+					]
+			];
+	}
+
+	return SNew(STransformOutlinerRow, InOwnerTable, SharedThis(this));
+}
+
+bool FGeometryCollectionTreeItemBone::IsValidBone() const
+{
+	TSharedPtr<FManagedArrayCollection> Collection = GetCollection();
+
+	if (ParentCollectionItem && ParentCollectionItem->IsValid())
+	{
+		return BoneIndex >= 0 && BoneIndex < Collection->NumElements(DataCollectionGroup);
+	}
+	return false;
+}
+
+TSharedRef<SWidget> FGeometryCollectionTreeItemBone::MakeIndexColumnWidget() const
+{
+	return SNew(SHorizontalBox)
+		+ SHorizontalBox::Slot()
+		.AutoWidth()
+		[
+			SNew(STextBlock)
+			.Text(FText::AsNumber(BoneIndex))
+			.ColorAndOpacity(ItemColor)
+		];
+}
+
+TSharedRef<SWidget> FGeometryCollectionTreeItemBone::MakeEmptyColumnWidget() const
+{
+	return SNew(SHorizontalBox)
+		+ SHorizontalBox::Slot()
+		.Padding(12.f, 0.f)
+		[
+			SNew(STextBlock)
+			.ColorAndOpacity(ItemColor)
+		];
+}
+
+void FGeometryCollectionTreeItemBone::GetChildren(FGeometryCollectionTreeItemList& OutChildren)
+{
+	ParentCollectionItem->GetChildrenForBone(*this, OutChildren);
+}
+
+bool FGeometryCollectionTreeItemBone::HasChildren() const
+{
+	return ParentCollectionItem->HasChildrenForBone(*this);
+}
+
+// --- STransformOutlinerRow ---
+
+TSharedRef<SWidget> STransformOutlinerRow::GenerateWidgetForColumn(const FName& ColumnName)
+{
+	TSharedPtr<FManagedArrayCollection> Collection = Item->GetCollection();
+	FName AttrType = UE::Dataflow::CollectionSpreadSheetHelpers::GetArrayTypeString(Collection->GetAttributeType(ColumnName, DataCollectionGroup));
+
+	// This can happen because sometimes slate retains old items until the next tick, and keeps calling callbacks on them until then
+	if (!Item->IsValidBone())
+	{
+		return Item->MakeEmptyColumnWidget();
+	}
+
+	if (ColumnName == "Index")
+	{
+		const TSharedPtr<SWidget> NameWidget = Item->MakeIndexColumnWidget();
+		return SNew(SHorizontalBox)
+			+ SHorizontalBox::Slot()
+			.HAlign(HAlign_Left)
+			.VAlign(VAlign_Center)
+			.AutoWidth()
+			[
+				SNew(SExpanderArrow, SharedThis(this))
+					.ShouldDrawWires(true)
+			]
+			+ SHorizontalBox::Slot()
+			[
+				NameWidget.ToSharedRef()
+			];
+	}
+
+	return UE::Dataflow::CollectionSpreadSheetHelpers::MakeColumnWidget(Collection,
+		FGeometryCollection::TransformGroup,
+		ColumnName,
+		Item->GetBoneIndex(),
+		Item->GetItemColor());
+}
+
+// --- STransformOutliner ---
+
+void STransformOutliner::Construct(const FArguments& InArgs)
+{
+	bPerformingSelection = false;
+
+	HeaderRowWidget =
+		SNew(SHeaderRow)
+		.Visibility(EVisibility::Visible);
+
+	RegenerateHeader();
+
+	ChildSlot
+	[
+		SNew(SVerticalBox)
+		+ SVerticalBox::Slot()
+		.Padding(FMargin(0.0f, 3.f))
+		[
+			SNew(SOverlay)
+			+ SOverlay::Slot()
+			[
+				SAssignNew(TreeView, STreeView<FGeometryCollectionTreeItemPtr>)
+				.TreeItemsSource(reinterpret_cast<FGeometryCollectionTreeItemList*>(&RootNodes))
+				.OnSelectionChanged(this, &STransformOutliner::OnSelectionChanged)
+				.OnGenerateRow(this, &STransformOutliner::MakeTreeRowWidget)
+				.OnGetChildren(this, &STransformOutliner::OnGetChildren)
+				.OnContextMenuOpening(this, &STransformOutliner::OnOpenContextMenu)
+				.AllowInvisibleItemSelection(true)
+				.ShouldStackHierarchyHeaders(true)
+				.OnGeneratePinnedRow(this, &STransformOutliner::OnGeneratePinnedRowWidget, true)
+				.HighlightParentNodesForSelection(true)
+				.OnSetExpansionRecursive(this, &STransformOutliner::ExpandRecursive)
+				.HeaderRow(HeaderRowWidget)
+				.ExternalScrollbar(InArgs._ExternalVerticalScrollBar)
+			]
+		]
+	];
+}
+
+void STransformOutliner::RegenerateHeader()
+{
+	constexpr int32 CustomColumnWidth = 120;
+	constexpr int32 CWidthMult = 9;
+
+	HeaderRowWidget->ClearColumns();
+
+	if (!Collection.IsValid() || Collection->NumElements(DataCollectionGroup) == 0)
+	{
+		return;
+	}
+
+	TArray<UE::Dataflow::CollectionSpreadSheetHelpers::FAttrInfo> AttrInfo;
+
+	// Add "Index" manually
+	AttrInfo.Add({ "Index", "int32" });
+	// Add all the other attrs
+	for (FName Attr : Collection->AttributeNames(DataCollectionGroup))
+	{
+		AttrInfo.Add({ Attr, UE::Dataflow::CollectionSpreadSheetHelpers::GetArrayTypeString(Collection->GetAttributeType(Attr, DataCollectionGroup)).ToString()});
+	}
+
+	for (int32 Idx = 0; Idx < AttrInfo.Num(); ++Idx)
+	{
+		int32 ColumnWidth = CustomColumnWidth;
+		const FString ColumnNameStr = AttrInfo[Idx].Name.ToString();
+
+		if (ColumnNameStr == "Index")
+		{
+			ColumnWidth = CustomColumnWidth;
+
+			HeaderRowWidget->AddColumn(
+				SHeaderRow::Column(AttrInfo[Idx].Name)
+				.DefaultLabel(FText::FromName(AttrInfo[Idx].Name))
+				.ManualWidth(ColumnWidth)
+				.HAlignCell(HAlign_Left)
+				.HAlignHeader(HAlign_Center)
+				.VAlignCell(VAlign_Center));
+		}
+		else
+		{
+			int32 ColumnNameStrLen = ColumnNameStr.Len() * CWidthMult;
+
+			if (UE::Dataflow::CollectionSpreadSheetHelpers::AttrTypeWidthMap.Contains(AttrInfo[Idx].Type))
+			{
+				ColumnWidth = UE::Dataflow::CollectionSpreadSheetHelpers::AttrTypeWidthMap[AttrInfo[Idx].Type];
+			}
+
+			ColumnWidth = ColumnNameStrLen > ColumnWidth ? ColumnNameStrLen : ColumnWidth;
+
+			HeaderRowWidget->AddColumn(
+				SHeaderRow::Column(AttrInfo[Idx].Name)
+				.DefaultLabel(FText::FromName(AttrInfo[Idx].Name))
+				.ManualWidth(ColumnWidth)
+				.HAlignCell(HAlign_Center)
+				.HAlignHeader(HAlign_Center)
+				.VAlignCell(VAlign_Center));
+		}
+	}
+}
+
+void STransformOutliner::RegenerateItems()
+{
+	TreeView->RebuildList();
+}
+
+TSharedRef<ITableRow> STransformOutliner::MakeTreeRowWidget(FGeometryCollectionTreeItemPtr InItem, const TSharedRef<STableViewBase>& InOwnerTable)
+{
+	return InItem->MakeTreeRowWidget(InOwnerTable);
+}
+
+TSharedRef<ITableRow> STransformOutliner::OnGeneratePinnedRowWidget(FGeometryCollectionTreeItemPtr InItem, const TSharedRef<STableViewBase>& InOwnerTable, bool bPinned)
+{
+	return InItem->MakeTreeRowWidget(InOwnerTable, true);
+}
+
+void STransformOutliner::OnGetChildren(FGeometryCollectionTreeItemPtr InItem, TArray<FGeometryCollectionTreeItemPtr>& OutChildren)
+{
+	InItem->GetChildren(OutChildren);
+}
+
+TSharedPtr<SWidget> STransformOutliner::OnOpenContextMenu()
+{
+	FGeometryCollectionTreeItemList SelectedItems;
+	TreeView->GetSelectedItems(SelectedItems);
+
+	if (SelectedItems.Num())
+	{
+		UToolMenus* ToolMenus = UToolMenus::Get();
+		static const FName MenuName = "STransformOutliner.TransformOutlinerContextMenu";
+		if (!ToolMenus->IsMenuRegistered(MenuName))
+		{
+			ToolMenus->RegisterMenu(MenuName);
+		}
+
+		// Build up the menu for a selection
+		FToolMenuContext Context;
+		UToolMenu* Menu = ToolMenus->GenerateMenu(MenuName, Context);
+		SelectedItems[0]->GenerateContextMenu(Menu, *this);
+		return ToolMenus->GenerateWidget(Menu);
+	}
+
+	return TSharedPtr<SWidget>();
+}
+
+void STransformOutliner::UpdateGeometryCollection()
+{
+	TreeView->RequestTreeRefresh();
+	ExpandAll();
+}
+
+void STransformOutliner::SetCollection(const TSharedPtr<FManagedArrayCollection>& InCollection, const FName InOutputName)
+{
+	// Clear the cached Tree ItemSelection without affecting the SelectedBones as 
+	// we want to refresh the tree selection using selected bones
+	TGuardValue<bool> ExternalSelectionGuard(bPerformingSelection, true);
+	TreeView->ClearSelection();
+
+	// explicitly mark the root nodes as invalid before emptying, so we know we can safely ignore them in case slate still triggers callbacks for them (they will not be deleted until the tree view refresh, on tick)
+	for (TSharedPtr<FGeometryCollectionTreeItemCollection>& RootNode : RootNodes)
+	{
+		if (RootNode)
+		{
+			RootNode->Invalidate();
+		}
+	}
+	RootNodes.Empty();
+
+	RootNodes.Add(MakeShared<FGeometryCollectionTreeItemCollection>(InCollection, InOutputName, TreeView));
+
+	TreeView->RequestTreeRefresh();
+	ExpandAll();
+
+	Collection = InCollection;
+}
+
+void STransformOutliner::ExpandAll()
+{
+	for (TSharedPtr<FGeometryCollectionTreeItemCollection> ItemPtr : RootNodes)
+	{
+		ItemPtr->ExpandAll();
+	}
+}
+
+void STransformOutliner::ExpandByLevel(const int32 InLevel)
+{
+	for (TSharedPtr<FGeometryCollectionTreeItemCollection> ItemPtr : RootNodes)
+	{
+		ItemPtr->ExpandByLevel(InLevel);
+	}
+}
+
+void STransformOutliner::ExpandRecursive(TSharedPtr<FGeometryCollectionTreeItem> ItemPtr, bool bInExpansionState) const
+{
+	TreeView->SetItemExpansion(ItemPtr, bInExpansionState);
+
+	FGeometryCollectionTreeItemList ItemChildren;
+	ItemPtr->GetChildren(ItemChildren);
+	for (auto& Child : ItemChildren)
+	{
+		ExpandRecursive(Child, bInExpansionState);
+	}
+}
+
+int32 STransformOutliner::GetBoneSelectionCount() const
+{
+	return TreeView->GetSelectedItems().Num();
+}
+
+void STransformOutliner::OnSelectionChanged(FGeometryCollectionTreeItemPtr Item, ESelectInfo::Type SelectInfo)
+{
+}
+
+void STransformOutliner::ContextMenuExpandAll()
+{
+	DisplayLevel = -1;
+
+	ExpandAll();
+}
+
+void STransformOutliner::ContextMenuPreviousLevel()
+{
+	if (DisplayLevel == -1)
+	{
+		return;
+	}
+
+	DisplayLevel--;
+
+	if (DisplayLevel == -1)
+	{
+		ExpandAll();
+	}
+	else
+	{
+		ExpandByLevel(DisplayLevel);
+	}
+}
+
+void STransformOutliner::ContextMenuNextLevel()
+{
+	if (Collection->HasAttribute("Level", DataCollectionGroup))
+	{
+		const TManagedArray<int32>& Levels = Collection->GetAttribute<int32>("Level", DataCollectionGroup);
+		int32 MaxLevel = -1000;
+		for (int32 Idx = 0; Idx < Levels.Num(); ++Idx)
+		{
+			if (Levels[Idx] > MaxLevel)
+			{
+				MaxLevel = Levels[Idx];
+			}
+		}
+
+		if (DisplayLevel < MaxLevel)
+		{
+			ExpandByLevel(++DisplayLevel);
+		}
+	}
+}
+
+#undef LOCTEXT_NAMESPACE // "CollectionSpreadSheet"
diff --git a/Engine/Plugins/Experimental/Dataflow/Source/DataflowEditor/Private/Dataflow/DataflowVerticesOutlinerWidget.cpp b/Engine/Plugins/Experimental/Dataflow/Source/DataflowEditor/Private/Dataflow/DataflowVerticesOutlinerWidget.cpp
new file mode 100644
--- /dev/null
+++ b/Engine/Plugins/Experimental/Dataflow/Source/DataflowEditor/Private/Dataflow/DataflowVerticesOutlinerWidget.cpp
@@ -0,0 +1,458 @@
+// Copyright Epic Games, Inc. All Rights Reserved.
+
+#include "Dataflow/DataflowVerticesOutlinerWidget.h"
+#include "Widgets/Input/SButton.h"
+#include "Widgets/Colors/SColorBlock.h"
+#include "Widgets/Layout/SScrollBox.h"
+#include "Dataflow/DataflowSettings.h"
+#include "GeometryCollection/GeometryCollection.h"
+#include "Dataflow/DataflowCollectionSpreadSheetHelpers.h"
+
+#define LOCTEXT_NAMESPACE "VerticesOutliner"
+
+// --- FGeometryCollectionVerticesTreeItem ---
+
+void FGeometryCollectionVerticesTreeItem::GenerateContextMenu(UToolMenu* Menu, SVerticesOutliner& Outliner)
+{
+}
+
+// --- FGeometryCollectionVerticesTreeItemGeometry ---
+
+TSharedRef<ITableRow> FGeometryCollectionVerticesTreeItemGeometry::MakeTreeRowWidget(const TSharedRef<STableViewBase>& InOwnerTable, bool bNoExtraColumn)
+{
+	FString NumVertsString = "";
+	FString NumFacesString = "";
+	FString TransformIdxStr = "";
+	FString GeometryIdxStr = "";
+	FString BoneNameStr = "";
+
+	if (Collection.IsValid() &&
+		Collection->HasAttribute("TransformIndex", FGeometryCollection::GeometryGroup) &&
+		Collection->HasAttribute("TransformToGeometryIndex", FGeometryCollection::TransformGroup) &&
+		Collection->HasAttribute("VertexCount", FGeometryCollection::GeometryGroup) &&
+		Collection->HasAttribute("FaceCount", FGeometryCollection::GeometryGroup) &&
+		Collection->HasAttribute("BoneName", FGeometryCollection::TransformGroup))
+	{
+		const TManagedArray<int32>& TransformIndices = Collection->GetAttribute<int32>("TransformIndex", FGeometryCollection::GeometryGroup);
+		const TManagedArray<int32>& TransformToGeometryIndices = Collection->GetAttribute<int32>("TransformToGeometryIndex", FGeometryCollection::TransformGroup);
+		const TManagedArray<int32>& VertexCounts = Collection->GetAttribute<int32>("VertexCount", FGeometryCollection::GeometryGroup);
+		const TManagedArray<int32>& FaceCounts = Collection->GetAttribute<int32>("FaceCount", FGeometryCollection::GeometryGroup);
+		const TManagedArray<FString>& BoneNames = Collection->GetAttribute<FString>("BoneName", FGeometryCollection::TransformGroup);
+
+		const int32 TransformIdx = TransformIndices[GeometryIdx];
+		TransformIdxStr = FString::FormatAsNumber(TransformIdx);
+
+		GeometryIdxStr = FString::FormatAsNumber(GeometryIdx);
+
+		const int32 NumVerts = VertexCounts[GeometryIdx];
+		NumVertsString = FString::FormatAsNumber(NumVerts);
+
+		const int32 NumFaces = FaceCounts[GeometryIdx];
+		NumFacesString = FString::FormatAsNumber(NumFaces);
+
+		BoneNameStr = BoneNames[TransformIdx];
+	}
+
+	return SNew(STableRow<FGeometryCollectionVerticesTreeItemPtr>, InOwnerTable)
+		.Content()
+		[
+			SNew(STextBlock).Text(FText::FromString("Geometry Idx: " + GeometryIdxStr + "     BoneName: " + BoneNameStr + "     Transform Idx: " + TransformIdxStr + "     Faces: " + NumFacesString + "     Verts: " + NumVertsString))
+		];
+}
+
+void FGeometryCollectionVerticesTreeItemGeometry::GetChildren(FGeometryCollectionVerticesTreeItemList& OutChildren)
+{
+	OutChildren = ChildItems;
+}
+
+FGeometryCollectionVerticesTreeItemPtr FGeometryCollectionVerticesTreeItemGeometry::GetItemFromVertexIndex(int32 InVertexIndex) const
+{
+	return ItemsByVertexIndex.FindRef(InVertexIndex);
+}
+
+void FGeometryCollectionVerticesTreeItemGeometry::GetChildrenForVertex(FGeometryCollectionVerticesTreeItemVertex& VertexItem, FGeometryCollectionVerticesTreeItemList& OutChildren) const
+{
+	if (!Collection.IsValid())
+	{
+		return;
+	}
+}
+
+bool FGeometryCollectionVerticesTreeItemGeometry::HasChildrenForVertex(const FGeometryCollectionVerticesTreeItemVertex& VertexItem) const
+{
+	if (!Collection.IsValid())
+	{
+		return false;
+	}
+
+	return false;
+}
+
+void FGeometryCollectionVerticesTreeItemGeometry::ExpandAll()
+{
+	TreeView->SetItemExpansion(AsShared(), true);
+
+	for (const auto& Elem : ItemsByVertexIndex)
+	{
+		TreeView->SetItemExpansion(Elem.Value, true);
+	}
+}
+
+void FGeometryCollectionVerticesTreeItemGeometry::RegenerateChildren()
+{
+	if (Collection.IsValid())
+	{
+		ItemsByVertexIndex.Empty();
+		ChildItems.Empty();
+
+		if (Collection &&
+			Collection->HasAttribute("VertexStart", FGeometryCollection::GeometryGroup) &&
+			Collection->HasAttribute("VertexCount", FGeometryCollection::GeometryGroup))
+		{
+			const TManagedArray<int32>& VertexStarts = Collection->GetAttribute<int32>("VertexStart", FGeometryCollection::GeometryGroup);
+			const TManagedArray<int32>& VertexCounts = Collection->GetAttribute<int32>("VertexCount", FGeometryCollection::GeometryGroup);
+
+			const int32 VertexStart = VertexStarts[GeometryIdx];
+			const int32 VertexCount = VertexCounts[GeometryIdx];
+
+			RootIndex = FGeometryCollection::Invalid;
+
+			// Add a sub item to the outliner tree for each of the bones/chunks in this GeometryCollection
+			for (int32 Index = 0; Index < VertexCount; ++Index)
+			{
+				const int32 VertexIndex = VertexStart + Index;
+
+				TSharedRef<FGeometryCollectionVerticesTreeItemVertex> NewItem = MakeShared<FGeometryCollectionVerticesTreeItemVertex>(VertexIndex, this);
+
+				ChildItems.Add(NewItem);
+				ItemsByVertexIndex.Add(VertexIndex, NewItem);
+			}
+		}
+	}
+}
+
+void FGeometryCollectionVerticesTreeItemGeometry::RequestTreeRefresh()
+{
+	if (TreeView.IsValid())
+	{
+		TreeView->RequestTreeRefresh();
+	}
+}
+
+bool FGeometryCollectionVerticesTreeItemGeometry::FilterBoneIndex(int32 BoneIndex) const
+{
+	return true;
+}
+
+bool FGeometryCollectionVerticesTreeItemGeometry::IsValid() const
+{
+	if (bInvalidated)
+	{
+		return false;
+	}
+
+	if (Collection)
+	{
+		return true;
+	}
+
+	return false;
+}
+
+// --- FGeometryCollectionVerticesTreeItemGeometry ---
+
+
+TSharedPtr<FManagedArrayCollection> FGeometryCollectionVerticesTreeItemVertex::GetCollection() const
+{
+	ensure(ParentCollectionItem);
+	return ParentCollectionItem->GetCollection();
+}
+
+void FGeometryCollectionVerticesTreeItemVertex::UpdateItemColorFromCollection()
+{
+	ItemColor = UE::Dataflow::CollectionSpreadSheetHelpers::UpdateItemColorFromCollection(GetCollection(), FGeometryCollection::VerticesGroup, GetVertexIndex());
+}
+
+TSharedRef<ITableRow> FGeometryCollectionVerticesTreeItemVertex::MakeTreeRowWidget(const TSharedRef<STableViewBase>& InOwnerTable, bool bIsPinned)
+{
+	UpdateItemColorFromCollection();
+
+	TSharedPtr<FManagedArrayCollection> Collection = GetCollection();
+
+	return SNew(SVerticesOutlinerRow, InOwnerTable, SharedThis(this));
+}
+
+bool FGeometryCollectionVerticesTreeItemVertex::IsValidVertex() const
+{
+	TSharedPtr<FManagedArrayCollection> Collection = GetCollection();
+
+	if (ParentCollectionItem && ParentCollectionItem->IsValid())
+	{
+		return VertexIndex >= 0 && VertexIndex < Collection->NumElements(FGeometryCollection::VerticesGroup);
+	}
+
+	return false;
+}
+
+TSharedRef<SWidget> FGeometryCollectionVerticesTreeItemVertex::MakeIndexColumnWidget() const
+{
+	return SNew(SHorizontalBox)
+		+ SHorizontalBox::Slot()
+		.AutoWidth()
+		[
+			SNew(STextBlock)
+			.Text(FText::AsNumber(VertexIndex))
+			.ColorAndOpacity(ItemColor)
+		];
+}
+
+TSharedRef<SWidget> FGeometryCollectionVerticesTreeItemVertex::MakeEmptyColumnWidget() const
+{
+	return SNew(SHorizontalBox)
+		+ SHorizontalBox::Slot()
+		.Padding(12.f, 0.f)
+		[
+			SNew(STextBlock)
+			.ColorAndOpacity(ItemColor)
+		];
+}
+
+void FGeometryCollectionVerticesTreeItemVertex::GetChildren(FGeometryCollectionVerticesTreeItemList& OutChildren)
+{
+	ParentCollectionItem->GetChildrenForVertex(*this, OutChildren);
+}
+
+bool FGeometryCollectionVerticesTreeItemVertex::HasChildren() const
+{
+	return ParentCollectionItem->HasChildrenForVertex(*this);
+}
+
+// --- SVerticesOutlinerRow ---
+
+TSharedRef<SWidget> SVerticesOutlinerRow::GenerateWidgetForColumn(const FName& ColumnName)
+{
+	TSharedPtr<FManagedArrayCollection> Collection = Item->GetCollection();
+	FName AttrType = UE::Dataflow::CollectionSpreadSheetHelpers::GetArrayTypeString(Collection->GetAttributeType(ColumnName, FGeometryCollection::VerticesGroup));
+
+	// This can happen because sometimes slate retains old items until the next tick, and keeps calling callbacks on them until then
+	if (!Item->IsValidVertex())
+	{
+		return Item->MakeEmptyColumnWidget();
+	}
+
+	if (ColumnName == "Index")
+	{
+		const TSharedPtr<SWidget> NameWidget = Item->MakeIndexColumnWidget();
+		return SNew(SHorizontalBox)
+			+ SHorizontalBox::Slot()
+			.HAlign(HAlign_Left)
+			.VAlign(VAlign_Center)
+			.AutoWidth()
+			[
+				SNew(SExpanderArrow, SharedThis(this))
+				.ShouldDrawWires(true)
+			]
+			+ SHorizontalBox::Slot()
+			[
+				NameWidget.ToSharedRef()
+			];
+	}
+
+	return UE::Dataflow::CollectionSpreadSheetHelpers::MakeColumnWidget(Collection, 
+		FGeometryCollection::VerticesGroup, 
+		ColumnName, 
+		Item->GetVertexIndex(), 
+		Item->GetItemColor());
+}
+
+// --- SVerticesOutliner ---
+
+void SVerticesOutliner::Construct(const FArguments& InArgs)
+{
+	bPerformingSelection = false;
+
+	HeaderRowWidget =
+		SNew(SHeaderRow)
+		.Visibility(EVisibility::Visible);
+
+	RegenerateHeader();
+
+	ChildSlot
+	[
+		SNew(SVerticalBox)
+		+ SVerticalBox::Slot()
+		.Padding(FMargin(0.0f, 3.f))
+		[
+			SNew(SOverlay)
+			+ SOverlay::Slot()
+			[
+				SAssignNew(TreeView, STreeView<FGeometryCollectionVerticesTreeItemPtr>)
+				.TreeItemsSource(reinterpret_cast<FGeometryCollectionVerticesTreeItemList*>(&RootNodes))
+				.OnSelectionChanged(this, &SVerticesOutliner::OnSelectionChanged)
+				.OnGenerateRow(this, &SVerticesOutliner::MakeTreeRowWidget)
+				.OnGetChildren(this, &SVerticesOutliner::OnGetChildren)
+				.OnContextMenuOpening(this, &SVerticesOutliner::OnOpenContextMenu)
+				.AllowInvisibleItemSelection(true)
+				.ShouldStackHierarchyHeaders(true)
+				.OnGeneratePinnedRow(this, &SVerticesOutliner::OnGeneratePinnedRowWidget, true)
+				.HighlightParentNodesForSelection(true)
+				.OnSetExpansionRecursive(this, &SVerticesOutliner::ExpandRecursive)
+				.HeaderRow(HeaderRowWidget)
+				.ExternalScrollbar(InArgs._ExternalVerticalScrollBar)
+			]
+		]
+	];
+}
+
+void SVerticesOutliner::RegenerateHeader()
+{
+	constexpr int32 CustomColumnWidth = 120;
+	constexpr int32 CWidthMult = 9;
+
+	HeaderRowWidget->ClearColumns();
+
+	if (!Collection.IsValid() || Collection->NumElements(FGeometryCollection::VerticesGroup) == 0)
+	{
+		return;
+	}
+
+	TArray<UE::Dataflow::CollectionSpreadSheetHelpers::FAttrInfo> AttrInfo;
+
+	// Add "Index" manually
+	AttrInfo.Add({ "Index", "int32" });
+	// Add all the other attrs
+	for (FName Attr : Collection->AttributeNames(FGeometryCollection::VerticesGroup))
+	{
+		AttrInfo.Add({ Attr, UE::Dataflow::CollectionSpreadSheetHelpers::GetArrayTypeString(Collection->GetAttributeType(Attr, FGeometryCollection::VerticesGroup)).ToString()});
+	}
+
+	for (int32 Idx = 0; Idx < AttrInfo.Num(); ++Idx)
+	{
+		int32 ColumnWidth = CustomColumnWidth;
+		const FString ColumnNameStr = AttrInfo[Idx].Name.ToString();
+
+		if (ColumnNameStr == "Index")
+		{
+			ColumnWidth = CustomColumnWidth;
+
+			HeaderRowWidget->AddColumn(
+				SHeaderRow::Column(AttrInfo[Idx].Name)
+				.DefaultLabel(FText::FromName(AttrInfo[Idx].Name))
+				.ManualWidth(ColumnWidth)
+				.HAlignCell(HAlign_Left)
+				.HAlignHeader(HAlign_Center)
+				.VAlignCell(VAlign_Center));
+		}
+		else
+		{
+			int32 ColumnNameStrLen = ColumnNameStr.Len() * CWidthMult;
+
+			if (UE::Dataflow::CollectionSpreadSheetHelpers::AttrTypeWidthMap.Contains(AttrInfo[Idx].Type))
+			{
+				ColumnWidth = UE::Dataflow::CollectionSpreadSheetHelpers::AttrTypeWidthMap[AttrInfo[Idx].Type];
+			}
+
+			ColumnWidth = ColumnNameStrLen > ColumnWidth ? ColumnNameStrLen : ColumnWidth;
+
+			HeaderRowWidget->AddColumn(
+				SHeaderRow::Column(AttrInfo[Idx].Name)
+				.DefaultLabel(FText::FromName(AttrInfo[Idx].Name))
+				.ManualWidth(ColumnWidth)
+				.HAlignCell(HAlign_Center)
+				.HAlignHeader(HAlign_Center)
+				.VAlignCell(VAlign_Center));
+		}
+	}
+}
+
+void SVerticesOutliner::RegenerateItems()
+{
+	TreeView->RebuildList();
+}
+
+TSharedRef<ITableRow> SVerticesOutliner::MakeTreeRowWidget(FGeometryCollectionVerticesTreeItemPtr InItem, const TSharedRef<STableViewBase>& InOwnerTable)
+{
+	return InItem->MakeTreeRowWidget(InOwnerTable);
+}
+
+TSharedRef<ITableRow> SVerticesOutliner::OnGeneratePinnedRowWidget(FGeometryCollectionVerticesTreeItemPtr InItem, const TSharedRef<STableViewBase>& InOwnerTable, bool bPinned)
+{
+	return InItem->MakeTreeRowWidget(InOwnerTable, true);
+}
+
+void SVerticesOutliner::OnGetChildren(FGeometryCollectionVerticesTreeItemPtr InItem, TArray<FGeometryCollectionVerticesTreeItemPtr>& OutChildren)
+{
+	InItem->GetChildren(OutChildren);
+}
+
+TSharedPtr<SWidget> SVerticesOutliner::OnOpenContextMenu()
+{
+	return TSharedPtr<SWidget>();
+}
+
+void SVerticesOutliner::UpdateGeometryCollection()
+{
+	TreeView->RequestTreeRefresh();
+	ExpandAll();
+}
+
+void SVerticesOutliner::SetCollection(const TSharedPtr<FManagedArrayCollection>& InCollection)
+{
+	// Clear the cached Tree ItemSelection without affecting the SelectedBones as 
+	// we want to refresh the tree selection using selected bones
+	TGuardValue<bool> ExternalSelectionGuard(bPerformingSelection, true);
+	TreeView->ClearSelection();
+
+	// explicitly mark the root nodes as invalid before emptying, so we know we can safely ignore them in case slate still triggers callbacks for them (they will not be deleted until the tree view refresh, on tick)
+	for (TSharedPtr<FGeometryCollectionVerticesTreeItemGeometry>& RootNode : RootNodes)
+	{
+		if (RootNode)
+		{
+			RootNode->Invalidate();
+		}
+	}
+	RootNodes.Empty();
+
+	const int32 NumGeoms = InCollection->NumElements(FGeometryCollection::GeometryGroup);
+	
+	for (int32 Idx = 0; Idx < NumGeoms; ++Idx)
+	{
+		RootNodes.Add(MakeShared<FGeometryCollectionVerticesTreeItemGeometry>(InCollection, Idx, TreeView));
+	}
+
+	TreeView->RequestTreeRefresh();
+	ExpandAll();
+
+	Collection = InCollection;
+}
+
+void SVerticesOutliner::ExpandAll()
+{
+	for (TSharedPtr<FGeometryCollectionVerticesTreeItemGeometry> ItemPtr : RootNodes)
+	{
+		ItemPtr->ExpandAll();
+	}
+}
+
+void SVerticesOutliner::ExpandRecursive(TSharedPtr<FGeometryCollectionVerticesTreeItem> ItemPtr, bool bInExpansionState) const
+{
+	TreeView->SetItemExpansion(ItemPtr, bInExpansionState);
+
+	FGeometryCollectionVerticesTreeItemList ItemChildren;
+	ItemPtr->GetChildren(ItemChildren);
+	for (auto& Child : ItemChildren)
+	{
+		ExpandRecursive(Child, bInExpansionState);
+	}
+}
+
+int32 SVerticesOutliner::GetBoneSelectionCount() const
+{
+	return TreeView->GetSelectedItems().Num();
+}
+
+void SVerticesOutliner::OnSelectionChanged(FGeometryCollectionVerticesTreeItemPtr Item, ESelectInfo::Type SelectInfo)
+{
+}
+
+#undef LOCTEXT_NAMESPACE // "VerticesOutliner"
diff --git a/Engine/Plugins/Experimental/Dataflow/Source/DataflowEditor/Private/Dataflow/SelectionViewWidget.cpp b/Engine/Plugins/Experimental/Dataflow/Source/DataflowEditor/Private/Dataflow/SelectionViewWidget.cpp
--- a/Engine/Plugins/Experimental/Dataflow/Source/DataflowEditor/Private/Dataflow/SelectionViewWidget.cpp
+++ b/Engine/Plugins/Experimental/Dataflow/Source/DataflowEditor/Private/Dataflow/SelectionViewWidget.cpp
@@ -3,11 +3,10 @@
 #include "Dataflow/SelectionViewWidget.h"
 #include "Widgets/Input/SButton.h"
 #include "Styling/StarshipCoreStyle.h"
-
+#include "Widgets/Input/SCheckBox.h"
 
 #define LOCTEXT_NAMESPACE "SelectionViewWidget"
 
-
 const FName FSelectionViewHeader::IndexColumnNameTransform = FName("Transform Index");
 const FName FSelectionViewHeader::IndexColumnNameFace = FName("Face Index");
 const FName FSelectionViewHeader::IndexColumnNameVertex = FName("Vertex Index");
@@ -233,6 +232,29 @@ void SSelectionViewWidget::NodeOutputsComboBoxSelectionChanged(FName InSelectedO
 	}
 }
 
+const FSlateBrush* SSelectionViewWidget::GetPinButtonImage() const
+{
+	if (bIsPinnedDown)
+	{
+		return FAppStyle::Get().GetBrush("Icons.Pinned");
+	}
+	else
+	{
+		return FAppStyle::Get().GetBrush("Icons.Unpinned");
+	}
+}
+
+const FSlateBrush* SSelectionViewWidget::GetLockButtonImage() const
+{
+	if (bIsRefreshLocked)
+	{
+		return FAppStyle::Get().GetBrush("Icons.Lock");
+	}
+	else
+	{
+		return FAppStyle::Get().GetBrush("Icons.Unlock");
+	}
+}
 
 void SSelectionViewWidget::Construct(const FArguments& InArgs)
 {
@@ -265,76 +287,65 @@ void SSelectionViewWidget::Construct(const FArguments& InArgs)
 			.AutoWidth()
 			.Padding(10.0f, 10.0f, 4.0f, 5.0f)
 			[
-				SNew(SButton)
+				SNew(SCheckBox)
+				.Style(FAppStyle::Get(), "ToggleButtonCheckbox")
 				.ToolTipText(LOCTEXT("PinDownButtonToolTip", "The button pins down the panel. When it pinned down it doesn't react to node selection change."))
-				.ButtonStyle(FAppStyle::Get(), "Button")
-				.HAlign(HAlign_Center)
-				.VAlign(VAlign_Center)
-				.Content()
-				[
-					SNew(STextBlock)
-					.Text_Lambda([this]
+				.IsChecked_Lambda([this]()-> ECheckBoxState
+				{
+					if (bIsPinnedDown)
 					{
-						if (!NodeNameTextBlock->GetText().IsEmpty())
-						{
-							return FText::FromString(bIsPinnedDown ? " X " : " O ");
-						}
-						else
-						{
-							return FText::FromString(" O ");
-						}
-					})
-				]
-				.OnClicked_Lambda([this]
+						return ECheckBoxState::Checked;
+					}
+					return ECheckBoxState::Unchecked;
+				})
+				.OnCheckStateChanged_Lambda([this](ECheckBoxState State)
 				{
 					if (!NodeNameTextBlock->GetText().IsEmpty())
 					{
 						bIsPinnedDown = !bIsPinnedDown;
-
 						OnPinnedDownChangedDelegate.Broadcast(bIsPinnedDown);
 					}
-
-					return FReply::Handled();					
 				})
+				.Padding(2.0f)
+				.HAlign(HAlign_Center)
+				[
+					SNew(SImage)
+					.ColorAndOpacity(FSlateColor::UseForeground())
+					.Image(this, &SSelectionViewWidget::GetPinButtonImage)
+				]
 			]
 
 			+ SHorizontalBox::Slot()
 			.AutoWidth()
 			.Padding(4.0f, 10.0f, 10.0f, 5.0f)
 			[
-				SNew(SButton)
+				SNew(SCheckBox)
+				.Style(FAppStyle::Get(), "ToggleButtonCheckbox")
 				.ToolTipText(LOCTEXT("LockRefreshButtonToolTip", "The button locks the refresh of the values in the panel."))
-				.ButtonStyle(FAppStyle::Get(), "Button")
-				.HAlign(HAlign_Center)
-				.VAlign(VAlign_Center)
-				.Content()
-				[
-					SNew(STextBlock)
-					.Text_Lambda([this]
+				.IsChecked_Lambda([this]()-> ECheckBoxState
+				{
+					if (bIsRefreshLocked)
 					{
-						if (!NodeNameTextBlock->GetText().IsEmpty())
-						{
-							return FText::FromString(bIsRefreshLocked ? " L " : " U ");
-						}
-						else
-						{
-							return FText::FromString(" U ");
-						}
-					})
-				]
-				.OnClicked_Lambda([this]
+						return ECheckBoxState::Checked;
+					}
+					return ECheckBoxState::Unchecked;
+				})
+				.OnCheckStateChanged_Lambda([this](ECheckBoxState State)
 				{
 					if (!NodeNameTextBlock->GetText().IsEmpty())
 					{
 						bIsRefreshLocked = !bIsRefreshLocked;
-
 						OnRefreshLockedChangedDelegate.Broadcast(bIsRefreshLocked);
 					}
-
-					return FReply::Handled();
 				})
+				.Padding(2.0f)
+				.HAlign(HAlign_Center)
+				[
+					SNew(SImage)
+					.ColorAndOpacity(FSlateColor::UseForeground())
+					.Image(this, &SSelectionViewWidget::GetLockButtonImage)
+				]
 			]
-
 		]
 
 		+ SVerticalBox::Slot()
diff --git a/Engine/Plugins/Experimental/Dataflow/Source/DataflowEditor/Public/Dataflow/DataflowCollectionSpreadSheetHelpers.h b/Engine/Plugins/Experimental/Dataflow/Source/DataflowEditor/Public/Dataflow/DataflowCollectionSpreadSheetHelpers.h
new file mode 100644
--- /dev/null
+++ b/Engine/Plugins/Experimental/Dataflow/Source/DataflowEditor/Public/Dataflow/DataflowCollectionSpreadSheetHelpers.h
@@ -0,0 +1,90 @@
+// Copyright Epic Games, Inc. All Rights Reserved.
+
+#pragma once
+
+#include "GeometryCollection/ManagedArrayCollection.h"
+#include "Styling/SlateColor.h"
+
+class SWidget;
+
+namespace UE::Dataflow::CollectionSpreadSheetHelpers
+{
+	struct FAttrInfo
+	{
+		FName Name;
+		FString Type;
+	};
+
+	static TMap<FString, int32> AttrTypeWidthMap =
+	{
+		{ TEXT("Transform"),	600 },
+		{ TEXT("Transform3f"),	600 },
+		{ TEXT("String"),		200 },
+		{ TEXT("LinearColor"),	80 },
+		{ TEXT("int32"),		100 },
+		{ TEXT("IntArray"),		200 },
+		{ TEXT("Vector"),		250 },
+		{ TEXT("Vector2D"),		160 },
+		{ TEXT("Float"),		150 },
+		{ TEXT("IntVector"),	220 },
+		{ TEXT("Bool"),			75 },
+		{ TEXT("Box"),			550 },
+		{ TEXT("MeshSection"),	100 },
+		{ TEXT("UInt8"),		100 },
+		{ TEXT("Guid"),			350 },
+	};
+
+	FString AttributeValueToString(float Value);
+
+	FString AttributeValueToString(int32 Value);
+
+	FString AttributeValueToString(FString Value);
+
+	FString AttributeValueToString(FLinearColor Value);
+
+	FString AttributeValueToString(FVector Value);
+
+	FString AttributeValueToString(bool Value);
+
+	FString AttributeValueToString(const FConstBitReference& Value);
+
+	FString AttributeValueToString(TSet<int32> Value);
+
+	FString AttributeValueToString(FTransform3f Value);
+
+	FString AttributeValueToString(FTransform Value);
+
+	FString AttributeValueToString(FBox Value);
+
+	FString AttributeValueToString(FIntVector Value);
+
+	FString AttributeValueToString(FIntVector4 Value);
+
+	FString AttributeValueToString(FGuid Value);
+
+	FString AttributeValueToString(const FManagedArrayCollection& InCollection, const FName& InAttributeName, const FName& InGroupName, int32 InIdxColumn);
+
+	inline FName GetArrayTypeString(FManagedArrayCollection::EArrayType ArrayType)
+	{
+		switch (ArrayType)
+		{
+#define MANAGED_ARRAY_TYPE(a,A)	case EManagedArrayType::F##A##Type:\
+		return FName(#A);
+#include "GeometryCollection/ManagedArrayTypeValues.inl"
+#undef MANAGED_ARRAY_TYPE
+		}
+		return FName();
+	}
+
+	FColor GetColorPerDepth(uint32 Depth);
+
+	FSlateColor UpdateItemColorFromCollection(const TSharedPtr<FManagedArrayCollection> InCollection, const FName InGroup, const int32 InItemIndex);
+
+	TSharedRef<SWidget> MakeColumnWidget(const TSharedPtr<FManagedArrayCollection> InCollection,
+		const FName InGroup,
+		const FName InAttr,
+		const int32 InItemIndex,
+		const FSlateColor InItemColor);
+}
+
+
diff --git a/Engine/Plugins/Experimental/Dataflow/Source/DataflowEditor/Public/Dataflow/DataflowCollectionSpreadSheetWidget.h b/Engine/Plugins/Experimental/Dataflow/Source/DataflowEditor/Public/Dataflow/DataflowCollectionSpreadSheetWidget.h
--- a/Engine/Plugins/Experimental/Dataflow/Source/DataflowEditor/Public/Dataflow/DataflowCollectionSpreadSheetWidget.h
+++ b/Engine/Plugins/Experimental/Dataflow/Source/DataflowEditor/Public/Dataflow/DataflowCollectionSpreadSheetWidget.h
@@ -5,10 +5,14 @@
 #include "CoreMinimal.h"
 #include "Dataflow/DataflowSelection.h"
 #include "DetailLayoutBuilder.h"
-#include "Widgets/Views/STableRow.h"
-#include "Widgets/Text/STextBlock.h"
-#include "Widgets/Input/SComboBox.h"
-#include "GeometryCollection/ManagedArrayCollection.h"
+#include "Delegates/Delegate.h"
+#include "UObject/NameTypes.h"
+#include "Styling/StarshipCoreStyle.h"
+#include "Dataflow/DataflowTransformOutlinerWidget.h"
+#include "Dataflow/DataflowVerticesOutlinerWidget.h"
+#include "Dataflow/DataflowFacesOutlinerWidget.h"
+
+struct FManagedArrayCollection;
 
 /**
 * Struct to hold OutputType/Selection data for the outputs
@@ -73,6 +77,7 @@ class SCollectionSpreadSheet : public SCompoundWidget
 public:
 	SLATE_BEGIN_ARGS(SCollectionSpreadSheet) {}
 		SLATE_ARGUMENT(FName, SelectedOutput)
+		SLATE_ARGUMENT(TSharedPtr<SScrollBar>, ExternalVerticalScrollBar)
 	SLATE_END_ARGS()
 
 	void Construct(const FArguments& InArgs);
@@ -80,6 +85,7 @@ public:
 	TMap<FString, SCollectionInfo>& GetCollectionInfoMap() { return CollectionInfoMap; }
 
 	const FName& GetSelectedOutput() const;
+	const FString GetSelectedOutputStr() const;
 	void SetSelectedOutput(const FName& InSelectedOutput);
 	
 	const FName& GetSelectedGroup() const;
@@ -90,17 +96,12 @@ public:
 
 	TSharedRef<ITableRow> GenerateRow(TSharedPtr<const FCollectionSpreadSheetItem> InItem, const TSharedRef<STableViewBase>& OwnerTable);
 
-	TSharedPtr<SScrollBar> SpreadSheetVerticalScrollBar;
-
-	TMap<FString, int32> AttrTypeWidthMap;
-
 private:
 	FName SelectedOutput = FName(TEXT(""));
 	FName SelectedGroup = FName(TEXT(""));
-	TMap<FString, SCollectionInfo> CollectionInfoMap;
+	TMap<FString, SCollectionInfo> CollectionInfoMap; // Map<Name of Output, FManagedArrayCollection>
 
 	TSharedPtr<SListView<TSharedPtr<const FCollectionSpreadSheetItem>>> ListView;
-//	TSharedPtr<STreeView<TSharedPtr<const FCollectionSpreadSheetItem>>> ListView;
 	TArray<TSharedPtr<const FCollectionSpreadSheetItem>> ListItems;
 
 	TSharedPtr<FCollectionSpreadSheetHeader> Header;
@@ -111,6 +112,10 @@ private:
 	void RepopulateListView();
 };
 
+class STransformOutliner;
+class SVerticesOutliner;
+class SFacesOutliner;
+
 /**
 * Widget for the CollectionSpreadSheet panel
 */
@@ -135,6 +140,9 @@ public:
 
 	void UpdateCollectionGroups(const FName& InOutputName);
 
+	const FSlateBrush* GetPinButtonImage() const;
+	const FSlateBrush* GetLockButtonImage() const;
+
 	/** Gets a multicast delegate which is called whenever the PinnedDown button clicked */
 	DECLARE_MULTICAST_DELEGATE_OneParam(FOnPinnedDownChanged, const bool);
 	FOnPinnedDownChanged& GetOnPinnedDownChangedDelegate() { return OnPinnedDownChangedDelegate; }
@@ -143,6 +151,28 @@ public:
 	DECLARE_MULTICAST_DELEGATE_OneParam(FOnRefreshLockedChanged, const bool);
 	FOnRefreshLockedChanged& GetOnRefreshLockedChangedDelegate() { return OnRefreshLockedChangedDelegate; }
 
+	TSet<FName> NonListViewGroups = { FName(TEXT("Transform")), FName(TEXT("Vertices")), FName(TEXT("Faces")) };
+
+	EVisibility GetCollectionSpreadSheetVisibility() const
+	{
+		return NonListViewGroups.Contains(CollectionTable->GetSelectedGroup()) ? EVisibility::Collapsed : EVisibility::Visible;
+	}
+
+	EVisibility GetTransformOutlinerVisibility() const
+	{
+		return CollectionTable->GetSelectedGroup() == FName(TEXT("Transform")) ? EVisibility::Visible : EVisibility::Collapsed;
+	}
+
+	EVisibility GetVerticesOutlinerVisibility() const
+	{
+		return CollectionTable->GetSelectedGroup() == FName(TEXT("Vertices")) ? EVisibility::Visible : EVisibility::Collapsed;
+	}
+
+	EVisibility GetFacesOutlinerVisibility() const
+	{
+		return CollectionTable->GetSelectedGroup() == FName(TEXT("Faces")) ? EVisibility::Visible : EVisibility::Collapsed;
+	}
+
 	// UI callbacks
 	void NodeOutputsComboBoxSelectionChanged(FName InSelectedOutput, ESelectInfo::Type InSelectInfo);
 	void CollectionGroupsComboBoxSelectionChanged(FName InSelectedOutput, ESelectInfo::Type InSelectInfo);
@@ -156,6 +186,16 @@ private:
 	TSharedPtr<SComboBox<FName>> CollectionGroupsComboBox;
 	TSharedPtr<STextBlock> CollectionGroupsComboBoxLabel;
 	TSharedPtr<SCollectionSpreadSheet> CollectionTable;
+
+	// Widget to display Transform hierarchy
+	TSharedPtr<STransformOutliner> TransformOutliner;
+
+	// Widget to display Vertices hierarchy
+	TSharedPtr<SVerticesOutliner> VerticesOutliner;
+
+	// Widget to display Faces hierarchy
+	TSharedPtr<SFacesOutliner> FacesOutliner;
+
 	TSharedPtr<STextBlock> StatusTextBlock;
 
 	FString NodeName;
@@ -163,9 +203,19 @@ private:
 	TArray<FName> CollectionGroups;
 	bool bIsPinnedDown = false;
 	bool bIsRefreshLocked = false;
+	// ScrollBar to scroll ListView/TreeView view area
+	TSharedPtr<SScrollBar> SpreadSheetHorizontalScrollBar;
+	// External scrollbar to scroll CollectionSpreadSheet vertically, only visible when CollectionSpreadSheet is visible
+	TSharedPtr<SScrollBar> CollectionSpreadSheetExternalVerticalScrollBar;
+	// External scrollbar to scroll TransformOutliner vertically, only visible when TransformOutliner is visible
+	TSharedPtr<SScrollBar> TransformOutlinerExternalVerticalScrollBar;
+	// External scrollbar to scroll VerticesOutliner vertically, only visible when VerticesOutliner is visible
+	TSharedPtr<SScrollBar> VerticesOutlinerExternalVerticalScrollBar;
+	// External scrollbar to scroll FacesOutliner vertically, only visible when FacesOutliner is visible
+	TSharedPtr<SScrollBar> FacesOutlinerExternalVerticalScrollBar;
 
 protected:
 	FOnPinnedDownChanged OnPinnedDownChangedDelegate;
 	FOnRefreshLockedChanged OnRefreshLockedChangedDelegate;
+};
 
-};
\ No newline at end of file
diff --git a/Engine/Plugins/Experimental/Dataflow/Source/DataflowEditor/Public/Dataflow/DataflowFacesOutlinerWidget.h b/Engine/Plugins/Experimental/Dataflow/Source/DataflowEditor/Public/Dataflow/DataflowFacesOutlinerWidget.h
new file mode 100644
--- /dev/null
+++ b/Engine/Plugins/Experimental/Dataflow/Source/DataflowEditor/Public/Dataflow/DataflowFacesOutlinerWidget.h
@@ -0,0 +1,173 @@
+// Copyright Epic Games, Inc. All Rights Reserved.
+
+#pragma once
+
+#include "CoreMinimal.h"
+#include "Dataflow/DataflowSelection.h"
+#include "DetailLayoutBuilder.h"
+#include "Delegates/Delegate.h"
+#include "UObject/NameTypes.h"
+#include "Styling/StarshipCoreStyle.h"
+#include "Widgets/Views/STreeView.h"
+
+class FGeometryCollection;
+class FGeometryCollectionFacesTreeItem;
+class FGeometryCollectionFacesTreeItemFace;
+class SFacesOutliner;
+struct FManagedArrayCollection;
+
+typedef TArray<TSharedPtr<FGeometryCollectionFacesTreeItem>> FGeometryCollectionFacesTreeItemList;
+typedef TSharedPtr<FGeometryCollectionFacesTreeItem> FGeometryCollectionFacesTreeItemPtr;
+typedef TSharedPtr<class FGeometryCollectionFacesTreeItemFace> FGeometryCollectionFacesTreeItemFacePtr;
+
+class FGeometryCollectionFacesTreeItem : public TSharedFromThis<FGeometryCollectionFacesTreeItem>
+{
+public:
+	virtual ~FGeometryCollectionFacesTreeItem() {}
+	virtual TSharedRef<ITableRow> MakeTreeRowWidget(const TSharedRef<STableViewBase>& InOwnerTable, bool bIsPinned = false) = 0;
+	virtual void GetChildren(FGeometryCollectionFacesTreeItemList& OutChildren) = 0;
+
+	virtual int32 GetFaceIndex() const { return INDEX_NONE; }
+
+	void GenerateContextMenu(UToolMenu* Menu, SFacesOutliner& Outliner);
+};
+
+class FGeometryCollectionFacesTreeItemGeometry : public FGeometryCollectionFacesTreeItem
+{
+public:
+	FGeometryCollectionFacesTreeItemGeometry(TSharedPtr<FManagedArrayCollection> InCollection, int32 InGeoemtryIdx, TSharedPtr<STreeView<FGeometryCollectionFacesTreeItemPtr>> InTreeView)
+		: Collection(InCollection)
+		, GeometryIdx(InGeoemtryIdx)
+		, TreeView(InTreeView)
+	{
+		RegenerateChildren();
+	}
+
+	virtual ~FGeometryCollectionFacesTreeItemGeometry() {}
+
+	// FGeometryCollectionFacesTreeItem interface
+	virtual TSharedRef<ITableRow> MakeTreeRowWidget(const TSharedRef<STableViewBase>& InOwnerTable, bool bIsPinned = false);
+	virtual void GetChildren(FGeometryCollectionFacesTreeItemList& OutChildren) override;
+
+	FGeometryCollectionFacesTreeItemPtr GetItemFromFaceIndex(int32 InFaceIndex) const;
+
+	void GetChildrenForFace(FGeometryCollectionFacesTreeItemFace& FaceItem, FGeometryCollectionFacesTreeItemList& OutChildren) const;
+	bool HasChildrenForFace(const FGeometryCollectionFacesTreeItemFace& FaceItem) const;
+
+	void ExpandAll();
+	void RegenerateChildren();
+	void RequestTreeRefresh();
+
+	TSharedPtr<FManagedArrayCollection> GetCollection() { return Collection; }
+
+	bool IsValid() const;
+
+	// Mark item as unused/invalid; helpful because slate defers destroying tree items and can still run callbacks on them until tick
+	void Invalidate()
+	{
+		bInvalidated = true;
+	}
+
+private:
+	bool FilterBoneIndex(int32 BoneIndex) const;
+
+	TSharedPtr<FManagedArrayCollection> Collection;
+	int32 GeometryIdx;
+
+	TSharedPtr<STreeView<FGeometryCollectionFacesTreeItemPtr>> TreeView;
+
+	// The direct children under this component
+	TArray<FGeometryCollectionFacesTreeItemPtr> ChildItems;
+	TMap<int32, FGeometryCollectionFacesTreeItemPtr> ItemsByFaceIndex;
+	int32 RootIndex;
+
+	// track whether the item has been explicitly invalidated
+	bool bInvalidated = false;
+};
+
+class FGeometryCollectionFacesTreeItemFace : public FGeometryCollectionFacesTreeItem
+{
+public:
+	FGeometryCollectionFacesTreeItemFace(const int32 InFaceIndex, FGeometryCollectionFacesTreeItemGeometry* InParentCollectionItem)
+		: FaceIndex(InFaceIndex)
+		, ParentCollectionItem(InParentCollectionItem)
+		, ItemColor(FSlateColor::UseForeground())
+	{}
+
+	// FGeometryCollectionFacesTreeItem interface
+	TSharedRef<ITableRow> MakeTreeRowWidget(const TSharedRef<STableViewBase>& InOwnerTable, bool bIsPinned = false);
+	TSharedRef<SWidget> MakeIndexColumnWidget() const;
+	TSharedRef<SWidget> MakeEmptyColumnWidget() const;
+
+	virtual void GetChildren(FGeometryCollectionFacesTreeItemList& OutChildren) override;
+	bool IsValidFace() const;
+	virtual int32 GetFaceIndex() const override { return FaceIndex; }
+	FSlateColor GetItemColor() const { return ItemColor; }
+	bool HasChildren() const;
+	TSharedPtr<FManagedArrayCollection> GetCollection() const;
+
+protected:
+
+private:
+	void UpdateItemColorFromCollection();
+
+	const int32 FaceIndex;
+	FGeometryCollectionFacesTreeItemGeometry* ParentCollectionItem;
+	FSlateColor ItemColor;
+};
+
+class SFacesOutlinerRow : public SMultiColumnTableRow<FGeometryCollectionFacesTreeItemFacePtr>
+{
+protected:
+	FGeometryCollectionFacesTreeItemFacePtr Item;
+
+public:
+	SLATE_BEGIN_ARGS(SFacesOutlinerRow) {}
+	SLATE_END_ARGS()
+
+	virtual TSharedRef<SWidget> GenerateWidgetForColumn(const FName& ColumnName) override;
+
+	void Construct(const FArguments& InArgs, TSharedRef<STableViewBase> OwnerTableView, FGeometryCollectionFacesTreeItemFacePtr InItemToEdit)
+	{
+		Item = InItemToEdit;
+		SMultiColumnTableRow<FGeometryCollectionFacesTreeItemFacePtr>::Construct(
+			FSuperRowType::FArguments()
+			.Style(&FStarshipCoreStyle::GetCoreStyle().GetWidgetStyle<FTableRowStyle>("TableView.AlternatingRow"))
+			, OwnerTableView);
+	}
+};
+
+class SFacesOutliner : public SCompoundWidget
+{
+	SLATE_BEGIN_ARGS(SFacesOutliner) {}
+		SLATE_ARGUMENT(TSharedPtr<SScrollBar>, ExternalVerticalScrollBar)
+	SLATE_END_ARGS()
+
+public:
+	void Construct(const FArguments& InArgs);
+
+	void RegenerateItems();
+	void RegenerateHeader();
+
+	TSharedRef<ITableRow> MakeTreeRowWidget(FGeometryCollectionFacesTreeItemPtr InInfo, const TSharedRef<STableViewBase>& OwnerTable);
+	TSharedRef<ITableRow> OnGeneratePinnedRowWidget(FGeometryCollectionFacesTreeItemPtr InItem, const TSharedRef<STableViewBase>& InOwnerTable, bool bPinned);
+	void OnGetChildren(TSharedPtr<FGeometryCollectionFacesTreeItem> InInfo, TArray< TSharedPtr<FGeometryCollectionFacesTreeItem> >& OutChildren);
+	TSharedPtr<SWidget> OnOpenContextMenu();
+
+	void UpdateGeometryCollection();
+	void SetCollection(const TSharedPtr<FManagedArrayCollection>& InCollection);
+	int32 GetBoneSelectionCount() const;
+
+	void ExpandAll();
+	void ExpandRecursive(TSharedPtr<FGeometryCollectionFacesTreeItem> TreeItem, bool bInExpansionState) const;
+
+private:
+	void OnSelectionChanged(FGeometryCollectionFacesTreeItemPtr Item, ESelectInfo::Type SelectInfo);
+
+	TSharedPtr<STreeView<FGeometryCollectionFacesTreeItemPtr>> TreeView;
+	TSharedPtr<SHeaderRow> HeaderRowWidget;
+	TArray<TSharedPtr<FGeometryCollectionFacesTreeItemGeometry>> RootNodes;
+	bool bPerformingSelection;
+	TSharedPtr<FManagedArrayCollection> Collection;
+};
+
diff --git a/Engine/Plugins/Experimental/Dataflow/Source/DataflowEditor/Public/Dataflow/DataflowTransformOutlinerWidget.h b/Engine/Plugins/Experimental/Dataflow/Source/DataflowEditor/Public/Dataflow/DataflowTransformOutlinerWidget.h
new file mode 100644
--- /dev/null
+++ b/Engine/Plugins/Experimental/Dataflow/Source/DataflowEditor/Public/Dataflow/DataflowTransformOutlinerWidget.h
@@ -0,0 +1,190 @@
+// Copyright Epic Games, Inc. All Rights Reserved.
+
+#pragma once
+
+#include "CoreMinimal.h"
+#include "Dataflow/DataflowSelection.h"
+#include "DetailLayoutBuilder.h"
+#include "Widgets/Views/STableRow.h"
+#include "Widgets/Text/STextBlock.h"
+#include "Widgets/Input/SComboBox.h"
+#include "GeometryCollection/ManagedArrayCollection.h"
+#include "GeometryCollection/ManagedArrayAccessor.h"
+#include "Delegates/Delegate.h"
+#include "UObject/NameTypes.h"
+#include "Styling/StarshipCoreStyle.h"
+#include "Widgets/Views/STreeView.h"
+
+//////////////////////////////////////////////////////////////////////////////////////////
+/* -- Widget to display Collection/Transform group data in a hierarchy -- */
+//////////////////////////////////////////////////////////////////////////////////////////
+
+class FGeometryCollection;
+class FGeometryCollectionTreeItem;
+class FGeometryCollectionTreeItemBone;
+class STransformOutliner;
+
+typedef TArray<TSharedPtr<FGeometryCollectionTreeItem>> FGeometryCollectionTreeItemList;
+typedef TSharedPtr<FGeometryCollectionTreeItem> FGeometryCollectionTreeItemPtr;
+typedef TSharedPtr<class FGeometryCollectionTreeItemBone> FGeometryCollectionTreeItemBonePtr;
+
+class FGeometryCollectionTreeItem : public TSharedFromThis<FGeometryCollectionTreeItem>
+{
+public:
+	virtual ~FGeometryCollectionTreeItem() {}
+	virtual TSharedRef<ITableRow> MakeTreeRowWidget(const TSharedRef<STableViewBase>& InOwnerTable, bool bIsPinned = false) = 0;
+	virtual void GetChildren(FGeometryCollectionTreeItemList& OutChildren) = 0;
+
+	virtual int32 GetBoneIndex() const { return INDEX_NONE; }
+
+	void GenerateContextMenu(UToolMenu* Menu, STransformOutliner& Outliner);
+};
+
+class FGeometryCollectionTreeItemCollection : public FGeometryCollectionTreeItem
+{
+public:
+	FGeometryCollectionTreeItemCollection(TSharedPtr<FManagedArrayCollection> InCollection, FName InOutputName, TSharedPtr<STreeView<FGeometryCollectionTreeItemPtr>> InTreeView)
+		: Collection(InCollection)
+		, OutputName(InOutputName)
+		, TreeView(InTreeView)
+	{
+		RegenerateChildren();
+	}
+
+	virtual ~FGeometryCollectionTreeItemCollection() {}
+
+	// FGeometryCollectionTreeItem interface
+	virtual TSharedRef<ITableRow> MakeTreeRowWidget(const TSharedRef<STableViewBase>& InOwnerTable, bool bIsPinned = false);
+	virtual void GetChildren(FGeometryCollectionTreeItemList& OutChildren) override;
+
+	FGeometryCollectionTreeItemPtr GetItemFromBoneIndex(int32 BoneIndex) const;
+
+	void GetChildrenForBone(FGeometryCollectionTreeItemBone& BoneItem, FGeometryCollectionTreeItemList& OutChildren) const;
+	bool HasChildrenForBone(const FGeometryCollectionTreeItemBone& BoneItem) const;
+
+	void ExpandAll();
+	void RegenerateChildren();
+	void RequestTreeRefresh();
+	void ExpandByLevel(const int32 InLevel);
+
+	TSharedPtr<FManagedArrayCollection> GetCollection() { return Collection; }
+
+	bool IsValid() const;
+
+	// Mark item as unused/invalid; helpful because slate defers destroying tree items and can still run callbacks on them until tick
+	void Invalidate()
+	{
+		bInvalidated = true;
+	}
+
+private:
+	bool FilterBoneIndex(int32 BoneIndex) const;
+
+	TSharedPtr<FManagedArrayCollection> Collection;
+	FName OutputName;
+
+	TSharedPtr<STreeView<FGeometryCollectionTreeItemPtr>> TreeView;
+
+	// The direct children under this component
+	TArray<FGeometryCollectionTreeItemPtr> ChildItems;
+	TMap<int32, FGeometryCollectionTreeItemPtr> ItemsByBoneIndex;
+	int32 RootIndex;
+
+	// track whether the item has been explicitly invalidated
+	bool bInvalidated = false;
+};
+
+class FGeometryCollectionTreeItemBone : public FGeometryCollectionTreeItem
+{
+public:
+	FGeometryCollectionTreeItemBone(const int32 InBoneIndex, FGeometryCollectionTreeItemCollection* InParentCollectionItem)
+		: BoneIndex(InBoneIndex)
+		, ParentCollectionItem(InParentCollectionItem)
+		, ItemColor(FSlateColor::UseForeground())
+	{}
+
+	// FGeometryCollectionTreeItem interface
+	TSharedRef<ITableRow> MakeTreeRowWidget(const TSharedRef<STableViewBase>& InOwnerTable, bool bIsPinned = false);
+	TSharedRef<SWidget> MakeIndexColumnWidget() const;
+//	TSharedRef<SWidget> MakeSimulationTypeColumnWidget() const;
+	TSharedRef<SWidget> MakeEmptyColumnWidget() const;
+
+	virtual void GetChildren(FGeometryCollectionTreeItemList& OutChildren) override;
+	bool IsValidBone() const;
+	virtual int32 GetBoneIndex() const override { return BoneIndex; }
+	FSlateColor GetItemColor() const { return ItemColor; }
+	bool HasChildren() const;
+	TSharedPtr<FManagedArrayCollection> GetCollection() const;
+
+protected:
+
+private:
+	void UpdateItemColorFromCollection();
+
+	const int32 BoneIndex;
+	FGeometryCollectionTreeItemCollection* ParentCollectionItem;
+	FSlateColor ItemColor;
+};
+
+class STransformOutlinerRow : public SMultiColumnTableRow<FGeometryCollectionTreeItemBonePtr>
+{
+protected:
+	FGeometryCollectionTreeItemBonePtr Item;
+
+public:
+	SLATE_BEGIN_ARGS(STransformOutlinerRow) {}
+	SLATE_END_ARGS()
+
+	virtual TSharedRef<SWidget> GenerateWidgetForColumn(const FName& ColumnName) override;
+
+	void Construct(const FArguments& InArgs, TSharedRef<STableViewBase> OwnerTableView, FGeometryCollectionTreeItemBonePtr InItemToEdit)
+	{
+		Item = InItemToEdit;
+		SMultiColumnTableRow<FGeometryCollectionTreeItemBonePtr>::Construct(
+			FSuperRowType::FArguments()
+			.Style(&FStarshipCoreStyle::GetCoreStyle().GetWidgetStyle<FTableRowStyle>("TableView.AlternatingRow"))
+			, OwnerTableView);
+	}
+};
+
+class STransformOutliner : public SCompoundWidget
+{
+	SLATE_BEGIN_ARGS(STransformOutliner) {}
+		SLATE_ARGUMENT(TSharedPtr<SScrollBar>, ExternalVerticalScrollBar)
+	SLATE_END_ARGS()
+
+public:
+	void Construct(const FArguments& InArgs);
+
+	void RegenerateItems();
+	void RegenerateHeader();
+
+	TSharedRef<ITableRow> MakeTreeRowWidget(FGeometryCollectionTreeItemPtr InInfo, const TSharedRef<STableViewBase>& OwnerTable);
+	TSharedRef<ITableRow> OnGeneratePinnedRowWidget(FGeometryCollectionTreeItemPtr InItem, const TSharedRef<STableViewBase>& InOwnerTable, bool bPinned);
+	void OnGetChildren(TSharedPtr<FGeometryCollectionTreeItem> InInfo, TArray< TSharedPtr<FGeometryCollectionTreeItem> >& OutChildren);
+	TSharedPtr<SWidget> OnOpenContextMenu();
+
+	void UpdateGeometryCollection();
+	void SetCollection(const TSharedPtr<FManagedArrayCollection>& InCollection, const FName InOutputName);
+	int32 GetBoneSelectionCount() const;
+
+	void ExpandAll();
+	void ExpandRecursive(TSharedPtr<FGeometryCollectionTreeItem> TreeItem, bool bInExpansionState) const;
+	void ExpandByLevel(const int32 InLevel);
+
+	int32 DisplayLevel = -1;
+
+	void ContextMenuExpandAll();
+	void ContextMenuPreviousLevel();
+	void ContextMenuNextLevel();
+	
+private:
+	void OnSelectionChanged(FGeometryCollectionTreeItemPtr Item, ESelectInfo::Type SelectInfo);
+
+	TSharedPtr<STreeView<FGeometryCollectionTreeItemPtr>> TreeView;
+	TSharedPtr<SHeaderRow> HeaderRowWidget;
+	TArray<TSharedPtr<FGeometryCollectionTreeItemCollection>> RootNodes;
+	bool bPerformingSelection;
+	TSharedPtr<FManagedArrayCollection> Collection;
+};
+
diff --git a/Engine/Plugins/Experimental/Dataflow/Source/DataflowEditor/Public/Dataflow/DataflowVerticesOutlinerWidget.h b/Engine/Plugins/Experimental/Dataflow/Source/DataflowEditor/Public/Dataflow/DataflowVerticesOutlinerWidget.h
new file mode 100644
--- /dev/null
+++ b/Engine/Plugins/Experimental/Dataflow/Source/DataflowEditor/Public/Dataflow/DataflowVerticesOutlinerWidget.h
@@ -0,0 +1,176 @@
+// Copyright Epic Games, Inc. All Rights Reserved.
+
+#pragma once
+
+#include "CoreMinimal.h"
+#include "Dataflow/DataflowSelection.h"
+#include "DetailLayoutBuilder.h"
+#include "Widgets/Views/STableRow.h"
+#include "Widgets/Text/STextBlock.h"
+#include "Widgets/Input/SComboBox.h"
+#include "GeometryCollection/ManagedArrayCollection.h"
+#include "Delegates/Delegate.h"
+#include "UObject/NameTypes.h"
+#include "Styling/StarshipCoreStyle.h"
+#include "Widgets/Views/STreeView.h"
+
+class FGeometryCollection;
+class FGeometryCollectionVerticesTreeItem;
+class FGeometryCollectionVerticesTreeItemVertex;
+class SVerticesOutliner;
+
+typedef TArray<TSharedPtr<FGeometryCollectionVerticesTreeItem>> FGeometryCollectionVerticesTreeItemList;
+typedef TSharedPtr<FGeometryCollectionVerticesTreeItem> FGeometryCollectionVerticesTreeItemPtr;
+typedef TSharedPtr<class FGeometryCollectionVerticesTreeItemVertex> FGeometryCollectionVerticesTreeItemVertexPtr;
+
+class FGeometryCollectionVerticesTreeItem : public TSharedFromThis<FGeometryCollectionVerticesTreeItem>
+{
+public:
+	virtual ~FGeometryCollectionVerticesTreeItem() {}
+	virtual TSharedRef<ITableRow> MakeTreeRowWidget(const TSharedRef<STableViewBase>& InOwnerTable, bool bIsPinned = false) = 0;
+	virtual void GetChildren(FGeometryCollectionVerticesTreeItemList& OutChildren) = 0;
+
+	virtual int32 GetVertexIndex() const { return INDEX_NONE; }
+
+	void GenerateContextMenu(UToolMenu* Menu, SVerticesOutliner& Outliner);
+};
+
+class FGeometryCollectionVerticesTreeItemGeometry : public FGeometryCollectionVerticesTreeItem
+{
+public:
+	FGeometryCollectionVerticesTreeItemGeometry(TSharedPtr<FManagedArrayCollection> InCollection, int32 InGeoemtryIdx, TSharedPtr<STreeView<FGeometryCollectionVerticesTreeItemPtr>> InTreeView)
+		: Collection(InCollection)
+		, GeometryIdx(InGeoemtryIdx)
+		, TreeView(InTreeView)
+	{
+		RegenerateChildren();
+	}
+
+	virtual ~FGeometryCollectionVerticesTreeItemGeometry() {}
+
+	// FGeometryCollectionVerticesTreeItem interface
+	virtual TSharedRef<ITableRow> MakeTreeRowWidget(const TSharedRef<STableViewBase>& InOwnerTable, bool bIsPinned = false);
+	virtual void GetChildren(FGeometryCollectionVerticesTreeItemList& OutChildren) override;
+
+	FGeometryCollectionVerticesTreeItemPtr GetItemFromVertexIndex(int32 InVertexIndex) const;
+
+	void GetChildrenForVertex(FGeometryCollectionVerticesTreeItemVertex& VertexItem, FGeometryCollectionVerticesTreeItemList& OutChildren) const;
+	bool HasChildrenForVertex(const FGeometryCollectionVerticesTreeItemVertex& VertexItem) const;
+
+	void ExpandAll();
+	void RegenerateChildren();
+	void RequestTreeRefresh();
+
+	TSharedPtr<FManagedArrayCollection> GetCollection() { return Collection; }
+
+	bool IsValid() const;
+
+	// Mark item as unused/invalid; helpful because slate defers destroying tree items and can still run callbacks on them until tick
+	void Invalidate()
+	{
+		bInvalidated = true;
+	}
+
+private:
+	bool FilterBoneIndex(int32 BoneIndex) const;
+
+	TSharedPtr<FManagedArrayCollection> Collection;
+	int32 GeometryIdx;
+
+	TSharedPtr<STreeView<FGeometryCollectionVerticesTreeItemPtr>> TreeView;
+
+	// The direct children under this component
+	TArray<FGeometryCollectionVerticesTreeItemPtr> ChildItems;
+	TMap<int32, FGeometryCollectionVerticesTreeItemPtr> ItemsByVertexIndex;
+	int32 RootIndex;
+
+	// track whether the item has been explicitly invalidated
+	bool bInvalidated = false;
+};
+
+class FGeometryCollectionVerticesTreeItemVertex : public FGeometryCollectionVerticesTreeItem
+{
+public:
+	FGeometryCollectionVerticesTreeItemVertex(const int32 InVertexIndex, FGeometryCollectionVerticesTreeItemGeometry* InParentCollectionItem)
+		: VertexIndex(InVertexIndex)
+		, ParentCollectionItem(InParentCollectionItem)
+		, ItemColor(FSlateColor::UseForeground())
+	{}
+
+	// FGeometryCollectionVerticesTreeItem interface
+	TSharedRef<ITableRow> MakeTreeRowWidget(const TSharedRef<STableViewBase>& InOwnerTable, bool bIsPinned = false);
+	TSharedRef<SWidget> MakeIndexColumnWidget() const;
+	TSharedRef<SWidget> MakeEmptyColumnWidget() const;
+
+	virtual void GetChildren(FGeometryCollectionVerticesTreeItemList& OutChildren) override;
+	bool IsValidVertex() const;
+	virtual int32 GetVertexIndex() const override { return VertexIndex; }
+	FSlateColor GetItemColor() const { return ItemColor; }
+	bool HasChildren() const;
+	TSharedPtr<FManagedArrayCollection> GetCollection() const;
+
+protected:
+
+private:
+	void UpdateItemColorFromCollection();
+
+	const int32 VertexIndex;
+	FGeometryCollectionVerticesTreeItemGeometry* ParentCollectionItem;
+	FSlateColor ItemColor;
+};
+
+class SVerticesOutlinerRow : public SMultiColumnTableRow<FGeometryCollectionVerticesTreeItemVertexPtr>
+{
+protected:
+	FGeometryCollectionVerticesTreeItemVertexPtr Item;
+
+public:
+	SLATE_BEGIN_ARGS(SVerticesOutlinerRow) {}
+	SLATE_END_ARGS()
+
+	virtual TSharedRef<SWidget> GenerateWidgetForColumn(const FName& ColumnName) override;
+
+	void Construct(const FArguments& InArgs, TSharedRef<STableViewBase> OwnerTableView, FGeometryCollectionVerticesTreeItemVertexPtr InItemToEdit)
+	{
+		Item = InItemToEdit;
+		SMultiColumnTableRow<FGeometryCollectionVerticesTreeItemVertexPtr>::Construct(
+			FSuperRowType::FArguments()
+			.Style(&FStarshipCoreStyle::GetCoreStyle().GetWidgetStyle<FTableRowStyle>("TableView.AlternatingRow"))
+			, OwnerTableView);
+	}
+};
+
+class SVerticesOutliner : public SCompoundWidget
+{
+	SLATE_BEGIN_ARGS(SVerticesOutliner) {}
+		SLATE_ARGUMENT(TSharedPtr<SScrollBar>, ExternalVerticalScrollBar)
+	SLATE_END_ARGS()
+
+public:
+	void Construct(const FArguments& InArgs);
+
+	void RegenerateItems();
+	void RegenerateHeader();
+
+	TSharedRef<ITableRow> MakeTreeRowWidget(FGeometryCollectionVerticesTreeItemPtr InInfo, const TSharedRef<STableViewBase>& OwnerTable);
+	TSharedRef<ITableRow> OnGeneratePinnedRowWidget(FGeometryCollectionVerticesTreeItemPtr InItem, const TSharedRef<STableViewBase>& InOwnerTable, bool bPinned);
+	void OnGetChildren(TSharedPtr<FGeometryCollectionVerticesTreeItem> InInfo, TArray< TSharedPtr<FGeometryCollectionVerticesTreeItem> >& OutChildren);
+	TSharedPtr<SWidget> OnOpenContextMenu();
+
+	void UpdateGeometryCollection();
+	void SetCollection(const TSharedPtr<FManagedArrayCollection>& InCollection);
+	int32 GetBoneSelectionCount() const;
+
+	void ExpandAll();
+	void ExpandRecursive(TSharedPtr<FGeometryCollectionVerticesTreeItem> TreeItem, bool bInExpansionState) const;
+
+private:
+	void OnSelectionChanged(FGeometryCollectionVerticesTreeItemPtr Item, ESelectInfo::Type SelectInfo);
+
+	TSharedPtr<STreeView<FGeometryCollectionVerticesTreeItemPtr>> TreeView;
+	TSharedPtr<SHeaderRow> HeaderRowWidget;
+	TArray<TSharedPtr<FGeometryCollectionVerticesTreeItemGeometry>> RootNodes;
+	bool bPerformingSelection;
+	TSharedPtr<FManagedArrayCollection> Collection;
+};
+
diff --git a/Engine/Plugins/Experimental/Dataflow/Source/DataflowEditor/Public/Dataflow/SelectionViewWidget.h b/Engine/Plugins/Experimental/Dataflow/Source/DataflowEditor/Public/Dataflow/SelectionViewWidget.h
--- a/Engine/Plugins/Experimental/Dataflow/Source/DataflowEditor/Public/Dataflow/SelectionViewWidget.h
+++ b/Engine/Plugins/Experimental/Dataflow/Source/DataflowEditor/Public/Dataflow/SelectionViewWidget.h
@@ -142,6 +142,9 @@ public:
 	DECLARE_MULTICAST_DELEGATE_OneParam(FOnRefreshLockedChanged, const bool);
 	FOnRefreshLockedChanged& GetOnRefreshLockedChangedDelegate() { return OnRefreshLockedChangedDelegate; }
 
+	const FSlateBrush* GetPinButtonImage() const;
+	const FSlateBrush* GetLockButtonImage() const;
+
 private:
 	TSharedPtr<STextBlock> NodeNameTextBlock;
 	TSharedPtr<SComboBox<FName>> NodeOutputsComboBox;
