Dataflow editor : add AddNewSubGraph command
- also broadcast an OnSubGraphsChanged event when a new SubGraph is added

diff --git a/Engine/Plugins/Experimental/Dataflow/Source/DataflowEditor/Private/Dataflow/DataflowAssetEditUtils.cpp b/Engine/Plugins/Experimental/Dataflow/Source/DataflowEditor/Private/Dataflow/DataflowAssetEditUtils.cpp
--- a/Engine/Plugins/Experimental/Dataflow/Source/DataflowEditor/Private/Dataflow/DataflowAssetEditUtils.cpp
+++ b/Engine/Plugins/Experimental/Dataflow/Source/DataflowEditor/Private/Dataflow/DataflowAssetEditUtils.cpp
@@ -882,8 +882,7 @@ namespace UE::Dataflow
 
 				DataflowAsset.AddSubGraph(NewSubGraph);
 
-				// todo(dataflow) notify of a subgraph update ? 
-				//FDataflowAssetDelegates::OnVariablesChanged.Broadcast(&DataflowAsset, UniqueVariableName);
+				FDataflowAssetDelegates::OnSubGraphsChanged.Broadcast(&DataflowAsset, UniqueSubGraphName);
 				return EChangeResult::Changed;
 			};
 
diff --git a/Engine/Plugins/Experimental/Dataflow/Source/DataflowEditor/Private/Dataflow/DataflowAssetEditUtils.h b/Engine/Plugins/Experimental/Dataflow/Source/DataflowEditor/Private/Dataflow/DataflowAssetEditUtils.h
--- a/Engine/Plugins/Experimental/Dataflow/Source/DataflowEditor/Private/Dataflow/DataflowAssetEditUtils.h
+++ b/Engine/Plugins/Experimental/Dataflow/Source/DataflowEditor/Private/Dataflow/DataflowAssetEditUtils.h
@@ -25,6 +25,7 @@ class UEdGraphPin;
 namespace UE::Dataflow
 {
 	static const FName DefaultNewVariableBaseName = "NewVariable";
+	static const FName DefaultNewSubGraphBaseName = "NewSubGraph";
 
 	struct FEditAssetUtils
 	{
@@ -121,6 +122,6 @@ namespace UE::Dataflow
 		/**
 		* Add a new subgraph to teh Dataflow Asset and return its name
 		*/
-		static FName AddNewSubGraph(UDataflow* DataflowAsset, FName BaseName);
+		static FName AddNewSubGraph(UDataflow* DataflowAsset, FName BaseName = DefaultNewSubGraphBaseName);
 	};
 }
diff --git a/Engine/Plugins/Experimental/Dataflow/Source/DataflowEditor/Private/Dataflow/DataflowEditorCommands.cpp b/Engine/Plugins/Experimental/Dataflow/Source/DataflowEditor/Private/Dataflow/DataflowEditorCommands.cpp
--- a/Engine/Plugins/Experimental/Dataflow/Source/DataflowEditor/Private/Dataflow/DataflowEditorCommands.cpp
+++ b/Engine/Plugins/Experimental/Dataflow/Source/DataflowEditor/Private/Dataflow/DataflowEditorCommands.cpp
@@ -76,6 +76,7 @@ void FDataflowEditorCommandsImpl::RegisterCommands()
 	UI_COMMAND(StepSimulationScene, "StepSimulationScene", "Step the simulation scene", EUserInterfaceActionType::ToggleButton, FInputChord());
 
 	UI_COMMAND(AddNewVariable, "Variable", "Adds a new variable to this dataflow graph.", EUserInterfaceActionType::Button, FInputChord());
+	UI_COMMAND(AddNewSubGraph, "SubGraph", "Adds a new sub-graph to this dataflow graph.", EUserInterfaceActionType::Button, FInputChord());
 
 	for (const TPair<FName, TUniquePtr<UE::Dataflow::IDataflowConstructionViewMode>>& NameAndMode : UE::Dataflow::FRenderingViewModeFactory::GetInstance().GetViewModes())
 	{
diff --git a/Engine/Plugins/Experimental/Dataflow/Source/DataflowEditor/Private/Dataflow/DataflowGraphEditor.cpp b/Engine/Plugins/Experimental/Dataflow/Source/DataflowEditor/Private/Dataflow/DataflowGraphEditor.cpp
--- a/Engine/Plugins/Experimental/Dataflow/Source/DataflowEditor/Private/Dataflow/DataflowGraphEditor.cpp
+++ b/Engine/Plugins/Experimental/Dataflow/Source/DataflowEditor/Private/Dataflow/DataflowGraphEditor.cpp
@@ -152,6 +152,10 @@ void SDataflowGraphEditor::Construct(const FArguments& InArgs, UObject* InAssetO
 				FDataflowEditorCommands::Get().AddNewVariable,
 				FExecuteAction::CreateSP(this, &SDataflowGraphEditor::AddNewVariable)
 			);
+			GraphEditorCommands->MapAction(
+					FDataflowEditorCommands::Get().AddNewSubGraph,
+					FExecuteAction::CreateSP(this, &SDataflowGraphEditor::AddNewSubGraph)
+			);
 		}
 	}
 
@@ -337,6 +341,11 @@ void SDataflowGraphEditor::AddNewVariable() const
 	UE::Dataflow::FEditAssetUtils::AddNewVariable(DataflowAsset.Get());
 }
 
+void SDataflowGraphEditor::AddNewSubGraph() const
+{
+	UE::Dataflow::FEditAssetUtils::AddNewSubGraph(DataflowAsset.Get());
+}
+
 void SDataflowGraphEditor::OnSelectedNodesChanged(const TSet<UObject*>& NewSelection)
 {
 	// Set the currently selected graph editor before running any callback
diff --git a/Engine/Plugins/Experimental/Dataflow/Source/DataflowEditor/Public/Dataflow/DataflowEditorCommands.h b/Engine/Plugins/Experimental/Dataflow/Source/DataflowEditor/Public/Dataflow/DataflowEditorCommands.h
--- a/Engine/Plugins/Experimental/Dataflow/Source/DataflowEditor/Public/Dataflow/DataflowEditorCommands.h
+++ b/Engine/Plugins/Experimental/Dataflow/Source/DataflowEditor/Public/Dataflow/DataflowEditorCommands.h
@@ -66,6 +66,7 @@ public:
 	TSharedPtr<FUICommandInfo> RemoveOptionPin;
 	TSharedPtr<FUICommandInfo> ZoomToFitGraph;
 	TSharedPtr<FUICommandInfo> AddNewVariable;
+	TSharedPtr<FUICommandInfo> AddNewSubGraph;
 
 	TMap<FName, TSharedPtr<FUICommandInfo>> SetConstructionViewModeCommands;
 
diff --git a/Engine/Plugins/Experimental/Dataflow/Source/DataflowEditor/Public/Dataflow/DataflowGraphEditor.h b/Engine/Plugins/Experimental/Dataflow/Source/DataflowEditor/Public/Dataflow/DataflowGraphEditor.h
--- a/Engine/Plugins/Experimental/Dataflow/Source/DataflowEditor/Public/Dataflow/DataflowGraphEditor.h
+++ b/Engine/Plugins/Experimental/Dataflow/Source/DataflowEditor/Public/Dataflow/DataflowGraphEditor.h
@@ -144,6 +144,9 @@ public:
 	/** Add a new variable for this dataflow graph */
 	void AddNewVariable() const;
 
+	/** Add a new SubGraph for this dataflow graph */
+	void AddNewSubGraph() const;
+
 	SGraphEditor* GetGraphEditor() { return (SGraphEditor*)this; }
 
 	/** FGCObject interface */
