Dataflow outliner : rename the tree item class to avoid conflicting with the one in fracture editor

diff --git a/Engine/Plugins/Experimental/Dataflow/Source/DataflowEditor/Private/Dataflow/DataflowTransformOutlinerWidget.cpp b/Engine/Plugins/Experimental/Dataflow/Source/DataflowEditor/Private/Dataflow/DataflowTransformOutlinerWidget.cpp
--- a/Engine/Plugins/Experimental/Dataflow/Source/DataflowEditor/Private/Dataflow/DataflowTransformOutlinerWidget.cpp
+++ b/Engine/Plugins/Experimental/Dataflow/Source/DataflowEditor/Private/Dataflow/DataflowTransformOutlinerWidget.cpp
@@ -11,9 +11,9 @@
 
 #define LOCTEXT_NAMESPACE "TransformOutliner"
 
-// --- FGeometryCollectionTreeItem ---
+// --- FDataflowOutlinerTreeItem ---
 
-void FGeometryCollectionTreeItem::GenerateContextMenu(UToolMenu* Menu, STransformOutliner& Outliner)
+void FDataflowOutlinerTreeItem::GenerateContextMenu(UToolMenu* Menu, STransformOutliner& Outliner)
 {
 	auto SharedOutliner = StaticCastSharedRef<STransformOutliner>(Outliner.AsShared());
 
@@ -32,9 +32,9 @@ static T GetAttributeValue(const TManagedArrayAccessor<T>& Attribute, int32 Inde
 	return (Attribute.IsValid()) ? Attribute.Get()[Index] : Default;
 }
 
-// --- FGeometryCollectionTreeItemCollection ---
+// --- FDataflowOutlinerTreeItemCollection ---
 
-TSharedRef<ITableRow> FGeometryCollectionTreeItemCollection::MakeTreeRowWidget(const TSharedRef<STableViewBase>& InOwnerTable, bool bNoExtraColumn)
+TSharedRef<ITableRow> FDataflowOutlinerTreeItemCollection::MakeTreeRowWidget(const TSharedRef<STableViewBase>& InOwnerTable, bool bNoExtraColumn)
 {
 	FString Output = OutputName.ToString();
 	FString NumVertsString = "";
@@ -48,24 +48,24 @@ TSharedRef<ITableRow> FGeometryCollectionTreeItemCollection::MakeTreeRowWidget(c
 		NumFacesString = FString::FormatAsNumber(NumFaces);
 	}
 
-	return SNew(STableRow<FGeometryCollectionTreeItemPtr>, InOwnerTable)
+	return SNew(STableRow<FDataflowOutlinerTreeItemPtr>, InOwnerTable)
 		.Content()
 		[
 			SNew(STextBlock).Text(FText::FromString(Output + "      Faces: " + NumFacesString + "    Verts: " + NumVertsString))
 		];
 }
 
-void FGeometryCollectionTreeItemCollection::GetChildren(FGeometryCollectionTreeItemList& OutChildren)
+void FDataflowOutlinerTreeItemCollection::GetChildren(FDataflowOutlinerTreeItemList& OutChildren)
 {
 	OutChildren = ChildItems;
 }
 
-FGeometryCollectionTreeItemPtr FGeometryCollectionTreeItemCollection::GetItemFromBoneIndex(int32 BoneIndex) const
+FDataflowOutlinerTreeItemPtr FDataflowOutlinerTreeItemCollection::GetItemFromBoneIndex(int32 BoneIndex) const
 {
 	return ItemsByBoneIndex.FindRef(BoneIndex);
 }
 
-void FGeometryCollectionTreeItemCollection::GetChildrenForBone(FGeometryCollectionTreeItemBone& BoneItem, FGeometryCollectionTreeItemList& OutChildren) const
+void FDataflowOutlinerTreeItemCollection::GetChildrenForBone(FDataflowOutlinerTreeItemBone& BoneItem, FDataflowOutlinerTreeItemList& OutChildren) const
 {
 	if (!Collection.IsValid())
 	{
@@ -80,7 +80,7 @@ void FGeometryCollectionTreeItemCollection::GetChildrenForBone(FGeometryCollecti
 			const TManagedArray<TSet<int32>>& Children = Collection->GetAttribute<TSet<int32>>("Children", FGeometryCollection::TransformGroup);
 			for (const int32 ChildIndex : Children[BoneIndex])
 			{
-				FGeometryCollectionTreeItemPtr ChildPtr = ItemsByBoneIndex.FindRef(ChildIndex);
+				FDataflowOutlinerTreeItemPtr ChildPtr = ItemsByBoneIndex.FindRef(ChildIndex);
 				if (ChildPtr.IsValid())
 				{
 					OutChildren.Add(ChildPtr);
@@ -90,7 +90,7 @@ void FGeometryCollectionTreeItemCollection::GetChildrenForBone(FGeometryCollecti
 	}
 }
 
-bool FGeometryCollectionTreeItemCollection::HasChildrenForBone(const FGeometryCollectionTreeItemBone& BoneItem) const
+bool FDataflowOutlinerTreeItemCollection::HasChildrenForBone(const FDataflowOutlinerTreeItemBone& BoneItem) const
 {
 	if (!Collection.IsValid())
 	{
@@ -110,7 +110,7 @@ bool FGeometryCollectionTreeItemCollection::HasChildrenForBone(const FGeometryCo
 	return false;
 }
 
-void FGeometryCollectionTreeItemCollection::ExpandAll()
+void FDataflowOutlinerTreeItemCollection::ExpandAll()
 {
 	TreeView->SetItemExpansion(AsShared(), true);
 
@@ -120,7 +120,7 @@ void FGeometryCollectionTreeItemCollection::ExpandAll()
 	}
 }
 
-void FGeometryCollectionTreeItemCollection::ExpandByLevel(const int32 InLevel)
+void FDataflowOutlinerTreeItemCollection::ExpandByLevel(const int32 InLevel)
 {
 	if (Collection->HasAttribute("Level", DataCollectionGroup))
 	{
@@ -144,7 +144,7 @@ void FGeometryCollectionTreeItemCollection::ExpandByLevel(const int32 InLevel)
 	}
 }
 
-void FGeometryCollectionTreeItemCollection::RegenerateChildren()
+void FDataflowOutlinerTreeItemCollection::RegenerateChildren()
 {
 	if (Collection.IsValid())
 	{
@@ -165,7 +165,7 @@ void FGeometryCollectionTreeItemCollection::RegenerateChildren()
 					{
 						if (FilterBoneIndex(BoneIndex))
 						{
-							TSharedRef<FGeometryCollectionTreeItemBone> NewItem = MakeShared<FGeometryCollectionTreeItemBone>(BoneIndex, this);
+							TSharedRef<FDataflowOutlinerTreeItemBone> NewItem = MakeShared<FDataflowOutlinerTreeItemBone>(BoneIndex, this);
 							if (Parents[BoneIndex] == RootIndex)
 							{
 								// The actual children directly beneath this node are the ones without a parent.  The rest are children of children
@@ -180,7 +180,7 @@ void FGeometryCollectionTreeItemCollection::RegenerateChildren()
 	}
 }
 
-void FGeometryCollectionTreeItemCollection::RequestTreeRefresh()
+void FDataflowOutlinerTreeItemCollection::RequestTreeRefresh()
 {
 	if (TreeView.IsValid())
 	{
@@ -188,7 +188,7 @@ void FGeometryCollectionTreeItemCollection::RequestTreeRefresh()
 	}
 }
 
-bool FGeometryCollectionTreeItemCollection::FilterBoneIndex(int32 BoneIndex) const
+bool FDataflowOutlinerTreeItemCollection::FilterBoneIndex(int32 BoneIndex) const
 {
 	if (Collection.IsValid())
 	{
@@ -222,7 +222,7 @@ bool FGeometryCollectionTreeItemCollection::FilterBoneIndex(int32 BoneIndex) con
 	return true;
 }
 
-bool FGeometryCollectionTreeItemCollection::IsValid() const
+bool FDataflowOutlinerTreeItemCollection::IsValid() const
 {
 	if (bInvalidated)
 	{
@@ -237,21 +237,21 @@ bool FGeometryCollectionTreeItemCollection::IsValid() const
 	return false;
 }
 
-// --- FGeometryCollectionTreeItemBone ---
+// --- FDataflowOutlinerTreeItemBone ---
 
 
-TSharedPtr<FManagedArrayCollection> FGeometryCollectionTreeItemBone::GetCollection() const
+TSharedPtr<FManagedArrayCollection> FDataflowOutlinerTreeItemBone::GetCollection() const
 {
 	ensure(ParentCollectionItem);
 	return ParentCollectionItem->GetCollection();
 }
 
-void FGeometryCollectionTreeItemBone::UpdateItemColorFromCollection()
+void FDataflowOutlinerTreeItemBone::UpdateItemColorFromCollection()
 {
 	ItemColor = UE::Dataflow::CollectionSpreadSheetHelpers::UpdateItemColorFromCollection(GetCollection(), FGeometryCollection::TransformGroup, GetBoneIndex());
 }
 
-TSharedRef<ITableRow> FGeometryCollectionTreeItemBone::MakeTreeRowWidget(const TSharedRef<STableViewBase>& InOwnerTable, bool bIsPinned)
+TSharedRef<ITableRow> FDataflowOutlinerTreeItemBone::MakeTreeRowWidget(const TSharedRef<STableViewBase>& InOwnerTable, bool bIsPinned)
 {
 	UpdateItemColorFromCollection();
 
@@ -260,7 +260,7 @@ TSharedRef<ITableRow> FGeometryCollectionTreeItemBone::MakeTreeRowWidget(const T
 	// This gets displayed in the pinned rows
 	if (bIsPinned)
 	{
-		return SNew(STableRow<FGeometryCollectionTreeItemPtr>, InOwnerTable)
+		return SNew(STableRow<FDataflowOutlinerTreeItemPtr>, InOwnerTable)
 			.Content()
 			[
 				SNew(SHorizontalBox)
@@ -286,7 +286,7 @@ TSharedRef<ITableRow> FGeometryCollectionTreeItemBone::MakeTreeRowWidget(const T
 	return SNew(STransformOutlinerRow, InOwnerTable, SharedThis(this));
 }
 
-bool FGeometryCollectionTreeItemBone::IsValidBone() const
+bool FDataflowOutlinerTreeItemBone::IsValidBone() const
 {
 	TSharedPtr<FManagedArrayCollection> Collection = GetCollection();
 
@@ -297,7 +297,7 @@ bool FGeometryCollectionTreeItemBone::IsValidBone() const
 	return false;
 }
 
-TSharedRef<SWidget> FGeometryCollectionTreeItemBone::MakeIndexColumnWidget() const
+TSharedRef<SWidget> FDataflowOutlinerTreeItemBone::MakeIndexColumnWidget() const
 {
 	return SNew(SHorizontalBox)
 		+ SHorizontalBox::Slot()
@@ -309,7 +309,7 @@ TSharedRef<SWidget> FGeometryCollectionTreeItemBone::MakeIndexColumnWidget() con
 		];
 }
 
-TSharedRef<SWidget> FGeometryCollectionTreeItemBone::MakeEmptyColumnWidget() const
+TSharedRef<SWidget> FDataflowOutlinerTreeItemBone::MakeEmptyColumnWidget() const
 {
 	return SNew(SHorizontalBox)
 		+ SHorizontalBox::Slot()
@@ -320,12 +320,12 @@ TSharedRef<SWidget> FGeometryCollectionTreeItemBone::MakeEmptyColumnWidget() con
 		];
 }
 
-void FGeometryCollectionTreeItemBone::GetChildren(FGeometryCollectionTreeItemList& OutChildren)
+void FDataflowOutlinerTreeItemBone::GetChildren(FDataflowOutlinerTreeItemList& OutChildren)
 {
 	ParentCollectionItem->GetChildrenForBone(*this, OutChildren);
 }
 
-bool FGeometryCollectionTreeItemBone::HasChildren() const
+bool FDataflowOutlinerTreeItemBone::HasChildren() const
 {
 	return ParentCollectionItem->HasChildrenForBone(*this);
 }
@@ -389,8 +389,8 @@ void STransformOutliner::Construct(const FArguments& InArgs)
 			SNew(SOverlay)
 			+ SOverlay::Slot()
 			[
-				SAssignNew(TreeView, STreeView<FGeometryCollectionTreeItemPtr>)
-				.TreeItemsSource(reinterpret_cast<FGeometryCollectionTreeItemList*>(&RootNodes))
+				SAssignNew(TreeView, STreeView<FDataflowOutlinerTreeItemPtr>)
+				.TreeItemsSource(reinterpret_cast<FDataflowOutlinerTreeItemList*>(&RootNodes))
 				.OnSelectionChanged(this, &STransformOutliner::OnSelectionChanged)
 				.OnGenerateRow(this, &STransformOutliner::MakeTreeRowWidget)
 				.OnGetChildren(this, &STransformOutliner::OnGetChildren)
@@ -473,24 +473,24 @@ void STransformOutliner::RegenerateItems()
 	TreeView->RebuildList();
 }
 
-TSharedRef<ITableRow> STransformOutliner::MakeTreeRowWidget(FGeometryCollectionTreeItemPtr InItem, const TSharedRef<STableViewBase>& InOwnerTable)
+TSharedRef<ITableRow> STransformOutliner::MakeTreeRowWidget(FDataflowOutlinerTreeItemPtr InItem, const TSharedRef<STableViewBase>& InOwnerTable)
 {
 	return InItem->MakeTreeRowWidget(InOwnerTable);
 }
 
-TSharedRef<ITableRow> STransformOutliner::OnGeneratePinnedRowWidget(FGeometryCollectionTreeItemPtr InItem, const TSharedRef<STableViewBase>& InOwnerTable, bool bPinned)
+TSharedRef<ITableRow> STransformOutliner::OnGeneratePinnedRowWidget(FDataflowOutlinerTreeItemPtr InItem, const TSharedRef<STableViewBase>& InOwnerTable, bool bPinned)
 {
 	return InItem->MakeTreeRowWidget(InOwnerTable, true);
 }
 
-void STransformOutliner::OnGetChildren(FGeometryCollectionTreeItemPtr InItem, TArray<FGeometryCollectionTreeItemPtr>& OutChildren)
+void STransformOutliner::OnGetChildren(FDataflowOutlinerTreeItemPtr InItem, TArray<FDataflowOutlinerTreeItemPtr>& OutChildren)
 {
 	InItem->GetChildren(OutChildren);
 }
 
 TSharedPtr<SWidget> STransformOutliner::OnOpenContextMenu()
 {
-	FGeometryCollectionTreeItemList SelectedItems;
+	FDataflowOutlinerTreeItemList SelectedItems;
 	TreeView->GetSelectedItems(SelectedItems);
 
 	if (SelectedItems.Num())
@@ -526,7 +526,7 @@ void STransformOutliner::SetCollection(const TSharedPtr<FManagedArrayCollection>
 	TreeView->ClearSelection();
 
 	// explicitly mark the root nodes as invalid before emptying, so we know we can safely ignore them in case slate still triggers callbacks for them (they will not be deleted until the tree view refresh, on tick)
-	for (TSharedPtr<FGeometryCollectionTreeItemCollection>& RootNode : RootNodes)
+	for (TSharedPtr<FDataflowOutlinerTreeItemCollection>& RootNode : RootNodes)
 	{
 		if (RootNode)
 		{
@@ -535,7 +535,7 @@ void STransformOutliner::SetCollection(const TSharedPtr<FManagedArrayCollection>
 	}
 	RootNodes.Empty();
 
-	RootNodes.Add(MakeShared<FGeometryCollectionTreeItemCollection>(InCollection, InOutputName, TreeView));
+	RootNodes.Add(MakeShared<FDataflowOutlinerTreeItemCollection>(InCollection, InOutputName, TreeView));
 
 	TreeView->RequestTreeRefresh();
 	ExpandAll();
@@ -545,7 +545,7 @@ void STransformOutliner::SetCollection(const TSharedPtr<FManagedArrayCollection>
 
 void STransformOutliner::ExpandAll()
 {
-	for (TSharedPtr<FGeometryCollectionTreeItemCollection> ItemPtr : RootNodes)
+	for (TSharedPtr<FDataflowOutlinerTreeItemCollection> ItemPtr : RootNodes)
 	{
 		ItemPtr->ExpandAll();
 	}
@@ -553,17 +553,17 @@ void STransformOutliner::ExpandAll()
 
 void STransformOutliner::ExpandByLevel(const int32 InLevel)
 {
-	for (TSharedPtr<FGeometryCollectionTreeItemCollection> ItemPtr : RootNodes)
+	for (TSharedPtr<FDataflowOutlinerTreeItemCollection> ItemPtr : RootNodes)
 	{
 		ItemPtr->ExpandByLevel(InLevel);
 	}
 }
 
-void STransformOutliner::ExpandRecursive(TSharedPtr<FGeometryCollectionTreeItem> ItemPtr, bool bInExpansionState) const
+void STransformOutliner::ExpandRecursive(TSharedPtr<FDataflowOutlinerTreeItem> ItemPtr, bool bInExpansionState) const
 {
 	TreeView->SetItemExpansion(ItemPtr, bInExpansionState);
 
-	FGeometryCollectionTreeItemList ItemChildren;
+	FDataflowOutlinerTreeItemList ItemChildren;
 	ItemPtr->GetChildren(ItemChildren);
 	for (auto& Child : ItemChildren)
 	{
@@ -576,7 +576,7 @@ int32 STransformOutliner::GetBoneSelectionCount() const
 	return TreeView->GetSelectedItems().Num();
 }
 
-void STransformOutliner::OnSelectionChanged(FGeometryCollectionTreeItemPtr Item, ESelectInfo::Type SelectInfo)
+void STransformOutliner::OnSelectionChanged(FDataflowOutlinerTreeItemPtr Item, ESelectInfo::Type SelectInfo)
 {
 }
 
diff --git a/Engine/Plugins/Experimental/Dataflow/Source/DataflowEditor/Public/Dataflow/DataflowTransformOutlinerWidget.h b/Engine/Plugins/Experimental/Dataflow/Source/DataflowEditor/Public/Dataflow/DataflowTransformOutlinerWidget.h
--- a/Engine/Plugins/Experimental/Dataflow/Source/DataflowEditor/Public/Dataflow/DataflowTransformOutlinerWidget.h
+++ b/Engine/Plugins/Experimental/Dataflow/Source/DataflowEditor/Public/Dataflow/DataflowTransformOutlinerWidget.h
@@ -19,31 +19,30 @@
 /* -- Widget to display Collection/Transform group data in a hierarchy -- */
 //////////////////////////////////////////////////////////////////////////////////////////
 
-class FGeometryCollection;
-class FGeometryCollectionTreeItem;
-class FGeometryCollectionTreeItemBone;
+class FDataflowOutlinerTreeItem;
+class FDataflowOutlinerTreeItemBone;
 class STransformOutliner;
 
-typedef TArray<TSharedPtr<FGeometryCollectionTreeItem>> FGeometryCollectionTreeItemList;
-typedef TSharedPtr<FGeometryCollectionTreeItem> FGeometryCollectionTreeItemPtr;
-typedef TSharedPtr<class FGeometryCollectionTreeItemBone> FGeometryCollectionTreeItemBonePtr;
+typedef TArray<TSharedPtr<FDataflowOutlinerTreeItem>> FDataflowOutlinerTreeItemList;
+typedef TSharedPtr<FDataflowOutlinerTreeItem> FDataflowOutlinerTreeItemPtr;
+typedef TSharedPtr<class FDataflowOutlinerTreeItemBone> FDataflowOutlinerTreeItemBonePtr;
 
-class FGeometryCollectionTreeItem : public TSharedFromThis<FGeometryCollectionTreeItem>
+class FDataflowOutlinerTreeItem : public TSharedFromThis<FDataflowOutlinerTreeItem>
 {
 public:
-	virtual ~FGeometryCollectionTreeItem() {}
+	virtual ~FDataflowOutlinerTreeItem() {}
 	virtual TSharedRef<ITableRow> MakeTreeRowWidget(const TSharedRef<STableViewBase>& InOwnerTable, bool bIsPinned = false) = 0;
-	virtual void GetChildren(FGeometryCollectionTreeItemList& OutChildren) = 0;
+	virtual void GetChildren(FDataflowOutlinerTreeItemList& OutChildren) = 0;
 
 	virtual int32 GetBoneIndex() const { return INDEX_NONE; }
 
 	void GenerateContextMenu(UToolMenu* Menu, STransformOutliner& Outliner);
 };
 
-class FGeometryCollectionTreeItemCollection : public FGeometryCollectionTreeItem
+class FDataflowOutlinerTreeItemCollection : public FDataflowOutlinerTreeItem
 {
 public:
-	FGeometryCollectionTreeItemCollection(TSharedPtr<FManagedArrayCollection> InCollection, FName InOutputName, TSharedPtr<STreeView<FGeometryCollectionTreeItemPtr>> InTreeView)
+	FDataflowOutlinerTreeItemCollection(TSharedPtr<FManagedArrayCollection> InCollection, FName InOutputName, TSharedPtr<STreeView<FDataflowOutlinerTreeItemPtr>> InTreeView)
 		: Collection(InCollection)
 		, OutputName(InOutputName)
 		, TreeView(InTreeView)
@@ -51,16 +50,16 @@ public:
 		RegenerateChildren();
 	}
 
-	virtual ~FGeometryCollectionTreeItemCollection() {}
+	virtual ~FDataflowOutlinerTreeItemCollection() {}
 
-	// FGeometryCollectionTreeItem interface
+	// FDataflowOutlinerTreeItem interface
 	virtual TSharedRef<ITableRow> MakeTreeRowWidget(const TSharedRef<STableViewBase>& InOwnerTable, bool bIsPinned = false);
-	virtual void GetChildren(FGeometryCollectionTreeItemList& OutChildren) override;
+	virtual void GetChildren(FDataflowOutlinerTreeItemList& OutChildren) override;
 
-	FGeometryCollectionTreeItemPtr GetItemFromBoneIndex(int32 BoneIndex) const;
+	FDataflowOutlinerTreeItemPtr GetItemFromBoneIndex(int32 BoneIndex) const;
 
-	void GetChildrenForBone(FGeometryCollectionTreeItemBone& BoneItem, FGeometryCollectionTreeItemList& OutChildren) const;
-	bool HasChildrenForBone(const FGeometryCollectionTreeItemBone& BoneItem) const;
+	void GetChildrenForBone(FDataflowOutlinerTreeItemBone& BoneItem, FDataflowOutlinerTreeItemList& OutChildren) const;
+	bool HasChildrenForBone(const FDataflowOutlinerTreeItemBone& BoneItem) const;
 
 	void ExpandAll();
 	void RegenerateChildren();
@@ -83,33 +82,33 @@ private:
 	TSharedPtr<FManagedArrayCollection> Collection;
 	FName OutputName;
 
-	TSharedPtr<STreeView<FGeometryCollectionTreeItemPtr>> TreeView;
+	TSharedPtr<STreeView<FDataflowOutlinerTreeItemPtr>> TreeView;
 
 	// The direct children under this component
-	TArray<FGeometryCollectionTreeItemPtr> ChildItems;
-	TMap<int32, FGeometryCollectionTreeItemPtr> ItemsByBoneIndex;
+	TArray<FDataflowOutlinerTreeItemPtr> ChildItems;
+	TMap<int32, FDataflowOutlinerTreeItemPtr> ItemsByBoneIndex;
 	int32 RootIndex;
 
 	// track whether the item has been explicitly invalidated
 	bool bInvalidated = false;
 };
 
-class FGeometryCollectionTreeItemBone : public FGeometryCollectionTreeItem
+class FDataflowOutlinerTreeItemBone : public FDataflowOutlinerTreeItem
 {
 public:
-	FGeometryCollectionTreeItemBone(const int32 InBoneIndex, FGeometryCollectionTreeItemCollection* InParentCollectionItem)
+	FDataflowOutlinerTreeItemBone(const int32 InBoneIndex, FDataflowOutlinerTreeItemCollection* InParentCollectionItem)
 		: BoneIndex(InBoneIndex)
 		, ParentCollectionItem(InParentCollectionItem)
 		, ItemColor(FSlateColor::UseForeground())
 	{}
 
-	// FGeometryCollectionTreeItem interface
+	// FDataflowOutlinerTreeItem interface
 	TSharedRef<ITableRow> MakeTreeRowWidget(const TSharedRef<STableViewBase>& InOwnerTable, bool bIsPinned = false);
 	TSharedRef<SWidget> MakeIndexColumnWidget() const;
 //	TSharedRef<SWidget> MakeSimulationTypeColumnWidget() const;
 	TSharedRef<SWidget> MakeEmptyColumnWidget() const;
 
-	virtual void GetChildren(FGeometryCollectionTreeItemList& OutChildren) override;
+	virtual void GetChildren(FDataflowOutlinerTreeItemList& OutChildren) override;
 	bool IsValidBone() const;
 	virtual int32 GetBoneIndex() const override { return BoneIndex; }
 	FSlateColor GetItemColor() const { return ItemColor; }
@@ -122,14 +121,14 @@ private:
 	void UpdateItemColorFromCollection();
 
 	const int32 BoneIndex;
-	FGeometryCollectionTreeItemCollection* ParentCollectionItem;
+	FDataflowOutlinerTreeItemCollection* ParentCollectionItem;
 	FSlateColor ItemColor;
 };
 
-class STransformOutlinerRow : public SMultiColumnTableRow<FGeometryCollectionTreeItemBonePtr>
+class STransformOutlinerRow : public SMultiColumnTableRow<FDataflowOutlinerTreeItemBonePtr>
 {
 protected:
-	FGeometryCollectionTreeItemBonePtr Item;
+	FDataflowOutlinerTreeItemBonePtr Item;
 
 public:
 	SLATE_BEGIN_ARGS(STransformOutlinerRow) {}
@@ -137,10 +136,10 @@ public:
 
 	virtual TSharedRef<SWidget> GenerateWidgetForColumn(const FName& ColumnName) override;
 
-	void Construct(const FArguments& InArgs, TSharedRef<STableViewBase> OwnerTableView, FGeometryCollectionTreeItemBonePtr InItemToEdit)
+	void Construct(const FArguments& InArgs, TSharedRef<STableViewBase> OwnerTableView, FDataflowOutlinerTreeItemBonePtr InItemToEdit)
 	{
 		Item = InItemToEdit;
-		SMultiColumnTableRow<FGeometryCollectionTreeItemBonePtr>::Construct(
+		SMultiColumnTableRow<FDataflowOutlinerTreeItemBonePtr>::Construct(
 			FSuperRowType::FArguments()
 			.Style(&FStarshipCoreStyle::GetCoreStyle().GetWidgetStyle<FTableRowStyle>("TableView.AlternatingRow"))
 			, OwnerTableView);
@@ -159,9 +158,9 @@ public:
 	void RegenerateItems();
 	void RegenerateHeader();
 
-	TSharedRef<ITableRow> MakeTreeRowWidget(FGeometryCollectionTreeItemPtr InInfo, const TSharedRef<STableViewBase>& OwnerTable);
-	TSharedRef<ITableRow> OnGeneratePinnedRowWidget(FGeometryCollectionTreeItemPtr InItem, const TSharedRef<STableViewBase>& InOwnerTable, bool bPinned);
-	void OnGetChildren(TSharedPtr<FGeometryCollectionTreeItem> InInfo, TArray< TSharedPtr<FGeometryCollectionTreeItem> >& OutChildren);
+	TSharedRef<ITableRow> MakeTreeRowWidget(FDataflowOutlinerTreeItemPtr InInfo, const TSharedRef<STableViewBase>& OwnerTable);
+	TSharedRef<ITableRow> OnGeneratePinnedRowWidget(FDataflowOutlinerTreeItemPtr InItem, const TSharedRef<STableViewBase>& InOwnerTable, bool bPinned);
+	void OnGetChildren(TSharedPtr<FDataflowOutlinerTreeItem> InInfo, TArray< TSharedPtr<FDataflowOutlinerTreeItem> >& OutChildren);
 	TSharedPtr<SWidget> OnOpenContextMenu();
 
 	void UpdateGeometryCollection();
@@ -169,7 +168,7 @@ public:
 	int32 GetBoneSelectionCount() const;
 
 	void ExpandAll();
-	void ExpandRecursive(TSharedPtr<FGeometryCollectionTreeItem> TreeItem, bool bInExpansionState) const;
+	void ExpandRecursive(TSharedPtr<FDataflowOutlinerTreeItem> TreeItem, bool bInExpansionState) const;
 	void ExpandByLevel(const int32 InLevel);
 
 	int32 DisplayLevel = -1;
@@ -179,11 +178,11 @@ public:
 	void ContextMenuNextLevel();
 	
 private:
-	void OnSelectionChanged(FGeometryCollectionTreeItemPtr Item, ESelectInfo::Type SelectInfo);
+	void OnSelectionChanged(FDataflowOutlinerTreeItemPtr Item, ESelectInfo::Type SelectInfo);
 
-	TSharedPtr<STreeView<FGeometryCollectionTreeItemPtr>> TreeView;
+	TSharedPtr<STreeView<FDataflowOutlinerTreeItemPtr>> TreeView;
 	TSharedPtr<SHeaderRow> HeaderRowWidget;
-	TArray<TSharedPtr<FGeometryCollectionTreeItemCollection>> RootNodes;
+	TArray<TSharedPtr<FDataflowOutlinerTreeItemCollection>> RootNodes;
 	bool bPerformingSelection;
 	TSharedPtr<FManagedArrayCollection> Collection;
 };
