Dataflow Editor : Add support for SubGraph in the DataflowMembersWidget
- Added a cvar ( false by default ) to enable UI support for Subgraph ( will be removed when SubGraphs are feature complete and ready )
-Added palette item for SubGraphs

diff --git a/Engine/Plugins/Experimental/Dataflow/Source/DataflowEditor/Private/Dataflow/DataflowEditorToolkit.cpp b/Engine/Plugins/Experimental/Dataflow/Source/DataflowEditor/Private/Dataflow/DataflowEditorToolkit.cpp
--- a/Engine/Plugins/Experimental/Dataflow/Source/DataflowEditor/Private/Dataflow/DataflowEditorToolkit.cpp
+++ b/Engine/Plugins/Experimental/Dataflow/Source/DataflowEditor/Private/Dataflow/DataflowEditorToolkit.cpp
@@ -1318,7 +1318,7 @@ TSharedPtr<SDataflowMembersWidget> FDataflowEditorToolkit::CreateDataflowMembers
 {
 	using namespace UE::Dataflow;
 
-	TSharedRef<SDataflowMembersWidget> NewMembersWidget = SNew(SDataflowMembersWidget, GraphEditor);
+	TSharedRef<SDataflowMembersWidget> NewMembersWidget = SNew(SDataflowMembersWidget, SharedThis(this));
 
 	return NewMembersWidget;
 
diff --git a/Engine/Plugins/Experimental/Dataflow/Source/DataflowEditor/Private/Dataflow/DataflowMembersWidget.cpp b/Engine/Plugins/Experimental/Dataflow/Source/DataflowEditor/Private/Dataflow/DataflowMembersWidget.cpp
--- a/Engine/Plugins/Experimental/Dataflow/Source/DataflowEditor/Private/Dataflow/DataflowMembersWidget.cpp
+++ b/Engine/Plugins/Experimental/Dataflow/Source/DataflowEditor/Private/Dataflow/DataflowMembersWidget.cpp
@@ -8,8 +8,11 @@
 #include "Dataflow/DataflowGraphSchemaAction.h"
 #include "Dataflow/DataflowInstance.h"
 #include "Dataflow/DataflowInstanceDetails.h"
+#include "Dataflow/DataflowEditorToolkit.h"
 #include "Dataflow/DataflowObject.h"
+#include "Dataflow/DataflowSubGraph.h"
 #include "Dataflow/DataflowVariablePaletteItem.h"
+#include "Dataflow/DataflowSubGraphPaletteItem.h"
 #include "Dataflow/DataflowVariableNodes.h"
 #include "Editor/PropertyEditor/Private/PropertyNode.h"
 #include "Framework/Commands/UICommandList.h"
@@ -25,6 +28,16 @@
 
 #define LOCTEXT_NAMESPACE "SDataflowMembersWidget"
 
+/////////////////////////////////////////////////////////////////////////////////////////////////////////////////////////
+namespace UE::Dataflow::DataflowMembersWidget::Private
+{
+	bool bEnableSubGraphs = false;
+	FAutoConsoleVariableRef CVarEnableSubGraphs(
+		TEXT("p.Dataflow.Editor.EnableSubgraphs"),
+		bEnableSubGraphs,
+		TEXT("When true, enable Dataflow SubGraph access for the UI")
+	);
+}
 /////////////////////////////////////////////////////////////////////////////////////////////////////////////////////////
 
 namespace UE::Dataflow::Private
@@ -96,28 +109,109 @@ namespace UE::Dataflow::Private
 		SDataflowMembersWidget::FButton AddButton;
 	};
 
-	TSharedPtr<SDataflowMembersWidget::ISection> MakeVariablesSection()
+	struct FSubGraphsSection : 
+		public FMemberSection<FEdGraphSchemaAction_DataflowSubGraph>
 	{
-		using FVariablesSection = FMemberSection<FEdGraphSchemaAction_DataflowVariable>;
-		
-		TSharedPtr<FVariablesSection> VariablesSection = MakeShared<FVariablesSection>();
+		static TSharedPtr<FSubGraphsSection> Make()
+		{
+			TSharedPtr<FSubGraphsSection> SubGraphsSection = MakeShared<FSubGraphsSection>();
+
+			SubGraphsSection->Title = LOCTEXT("SubGraphs", "SubGraphs");
+			SubGraphsSection->AddButton =
+			{
+				.Tooltip = LOCTEXT("AddNewSubGraph", "Add New Sub-Graph"),
+				.MetadataTag = TEXT("AddNewSubGraph"),
+				.Command = FDataflowEditorCommands::Get().AddNewSubGraph,
+			};
+			return SubGraphsSection;
+		}
 
-		VariablesSection->Title = LOCTEXT("Variables", "Variables");
-		VariablesSection->AddButton = 
+		virtual TSharedRef<SWidget> CreateWidgetForAction(FCreateWidgetForActionData* const InCreateData, TSharedPtr<SDataflowGraphEditor> Editor) const override
 		{
-			.Tooltip = LOCTEXT("AddNewVariable", "Add New Variable"),
-			.MetadataTag = TEXT("AddNewVariable"),
-			.Command = FDataflowEditorCommands::Get().AddNewVariable,
-		};
-		return VariablesSection;
-	}
+			return SNew(SDataflowSubGraphPaletteItem, InCreateData, Editor);
+		}
+
+		virtual void CollectActions(IDataflowInstanceInterface* Interface, TArray<TSharedPtr<FEdGraphSchemaAction>>& OutActions) const override
+		{
+			if (Interface)
+			{
+				if (const UDataflow* DataflowAsset = Interface->GetDataflowInstance().GetDataflowAsset())
+				{
+					for (TObjectPtr<const UDataflowSubGraph> SubGraph : DataflowAsset->GetSubGraphs())
+					{
+						if (ensure(SubGraph))
+						{
+							TSharedPtr<FEdGraphSchemaAction_DataflowSubGraph> NewSubGraphAction =
+								MakeShareable(new FEdGraphSchemaAction_DataflowSubGraph(Interface, SubGraph->GetFName()));
+							OutActions.Add(NewSubGraphAction);
+						}
+					}
+				}
+			}
+		}
+
+		virtual void OnDoubleClicked(FEdGraphSchemaAction& InAction, FDataflowEditorToolkit& InToolkit) const override
+		{
+			if (FEdGraphSchemaAction_DataflowSubGraph* TypedAction = AsTypedAction(InAction))
+			{
+				InToolkit.OpenSubGraphTab(TypedAction->GetSubGraphName());
+			}
+		}
+	};
+
+	struct FVariablesSection :
+		public FMemberSection<FEdGraphSchemaAction_DataflowVariable>
+	{
+		static TSharedPtr<FVariablesSection> Make()
+		{
+			TSharedPtr<FVariablesSection> VariablesSection = MakeShared<FVariablesSection>();
+
+			VariablesSection->Title = LOCTEXT("Variables", "Variables");
+			VariablesSection->AddButton =
+			{
+				.Tooltip = LOCTEXT("AddNewVariable", "Add New Variable"),
+				.MetadataTag = TEXT("AddNewVariable"),
+				.Command = FDataflowEditorCommands::Get().AddNewVariable,
+			};
+			return VariablesSection;
+		}
+
+		virtual TSharedRef<SWidget> CreateWidgetForAction(FCreateWidgetForActionData* const InCreateData, TSharedPtr<SDataflowGraphEditor> Editor) const override
+		{
+			return SNew(SDataflowVariablePaletteItem, InCreateData, Editor);
+		}
+
+		virtual void CollectActions(IDataflowInstanceInterface* Interface, TArray<TSharedPtr<FEdGraphSchemaAction>>& OutActions) const override
+		{
+			if (Interface)
+			{
+				if (const UDataflow* DataflowAsset = Interface->GetDataflowInstance().GetDataflowAsset())
+				{
+					if (const UPropertyBag* PropertyBag = DataflowAsset->Variables.GetPropertyBagStruct())
+					{
+						for (const FPropertyBagPropertyDesc& ConfigDesc : PropertyBag->GetPropertyDescs())
+						{
+							TSharedPtr<FEdGraphSchemaAction_DataflowVariable> NewVarAction =
+								MakeShareable(new FEdGraphSchemaAction_DataflowVariable(Interface, ConfigDesc));
+							OutActions.Add(NewVarAction);
+						}
+					}
+				}
+			}
+		}
+
+		virtual void OnDoubleClicked(FEdGraphSchemaAction&, FDataflowEditorToolkit&) const override
+		{
+			// todo : Jump to the node 
+		}
+	};
 }
 
 /////////////////////////////////////////////////////////////////////////////////////////////////////////////////////////
 
-void SDataflowMembersWidget::Construct(const FArguments& InArgs, TSharedPtr<SDataflowGraphEditor> InGraphEditor)
+void SDataflowMembersWidget::Construct(const FArguments& InArgs, TSharedPtr<FDataflowEditorToolkit> InEditorToolkit)
 {
-	GraphEditor = InGraphEditor;
+	EditorToolkitWeakPtr = InEditorToolkit;
 	CacheEditedAsset();
 
 	InitializeCommands();
@@ -217,6 +311,7 @@ void SDataflowMembersWidget::Construct(const FArguments& InArgs, TSharedPtr<SDat
 	];
 
 	FCoreUObjectDelegates::OnObjectPropertyChanged.AddRaw(this, &SDataflowMembersWidget::OnObjectPropertyChanged);
+	FDataflowAssetDelegates::OnSubGraphsChanged.AddRaw(this, &SDataflowMembersWidget::OnSubGraphsChanged);
 }
 
 void SDataflowMembersWidget::CreateVariableOverrideDetailView()
@@ -273,10 +368,22 @@ void SDataflowMembersWidget::RefreshVariableOverrideDetailView()
 SDataflowMembersWidget::~SDataflowMembersWidget()
 {
 	FCoreUObjectDelegates::OnObjectPropertyChanged.RemoveAll(this);
+	FDataflowAssetDelegates::OnSubGraphsChanged.RemoveAll(this);
+}
+
+const TSharedPtr<SDataflowGraphEditor> SDataflowMembersWidget::GetGraphEditor() const
+{
+	if (EditorToolkitWeakPtr.IsValid())
+	{
+		return EditorToolkitWeakPtr.Pin()->GetDataflowGraphEditor();
+	}
+	return {};
 }
 
 void SDataflowMembersWidget::InitializeCommands()
 {
+	const TSharedPtr<SDataflowGraphEditor> GraphEditor = GetGraphEditor();
+
 	// initialize command list ( merge from Graph editor )
 	CommandList = MakeShareable(new FUICommandList);
 	if (GraphEditor)
@@ -315,8 +422,11 @@ void SDataflowMembersWidget::InitializeSections()
 {
 	using namespace UE::Dataflow;
 	SectionMap.Reset();
-
-	SectionMap.Emplace((int32)ESchemaActionSectionID::VARIABLES, Private::MakeVariablesSection());
+	if (UE::Dataflow::DataflowMembersWidget::Private::bEnableSubGraphs)
+	{
+		SectionMap.Emplace((int32)ESchemaActionSectionID::SUBGRAPHS, Private::FSubGraphsSection::Make());
+	}
+	SectionMap.Emplace((int32)ESchemaActionSectionID::VARIABLES, Private::FVariablesSection::Make());
 }
 
 const TSharedPtr<SDataflowMembersWidget::ISection> SDataflowMembersWidget::GetSectionById(int32 SectionId) const
@@ -340,6 +450,7 @@ FReply SDataflowMembersWidget::OnKeyDown(const FGeometry& MyGeometry, const FKey
 
 void SDataflowMembersWidget::CacheEditedAsset()
 {
+	const TSharedPtr<SDataflowGraphEditor> GraphEditor = GetGraphEditor();
 	if (!GraphEditor)
 	{
 		return;
@@ -414,6 +525,11 @@ void SDataflowMembersWidget::OnObjectPropertyChanged(UObject* InObject, FPropert
 	}
 }
 
+void SDataflowMembersWidget::OnSubGraphsChanged(const UDataflow* InDataflowAsset, FName SubGraphName)
+{
+	Refresh();
+}
+
 void SDataflowMembersWidget::Refresh()
 {
 	if (GraphActionMenu)
@@ -434,28 +550,33 @@ FText SDataflowMembersWidget::GetFilterText() const
 
 TSharedRef<SWidget> SDataflowMembersWidget::OnCreateWidgetForAction(FCreateWidgetForActionData* const InCreateData)
 {
-	using namespace UE::Dataflow::Private;
-	return SNew(SDataflowVariablePaletteItem, InCreateData, GraphEditor);
+	if (InCreateData && InCreateData->Action)
+	{
+		if (const TSharedPtr<ISection> Section = GetSectionById(InCreateData->Action->SectionID))
+		{
+			return Section->CreateWidgetForAction(InCreateData, GetGraphEditor());
+		}
+	}
+	return SNullWidget::NullWidget;
 }
 
 void SDataflowMembersWidget::CollectAllActions(FGraphActionListBuilderBase& OutAllActions)
 {
-	// todo : this should be using the Section map and running a callback from each of the entry 
+	TArray<TSharedPtr<FEdGraphSchemaAction>> Actions;
 	if (IDataflowInstanceInterface* DataflowInstanceInterface = GetDataflowInstanceInterface())
 	{
-		if (const UDataflow* DataflowAsset = DataflowInstanceInterface->GetDataflowInstance().GetDataflowAsset())
+		for (const TPair<int32, TSharedPtr<ISection>>& SectionEntry : SectionMap)
 		{
-			if (const UPropertyBag* PropertyBag = DataflowAsset->Variables.GetPropertyBagStruct())
+			if (ensure(SectionEntry.Value))
 			{
-				for (const FPropertyBagPropertyDesc& ConfigDesc : PropertyBag->GetPropertyDescs())
-				{
-					TSharedPtr<FEdGraphSchemaAction_DataflowVariable> NewVarAction =
-						MakeShareable(new FEdGraphSchemaAction_DataflowVariable(DataflowInstanceInterface, ConfigDesc));
-					OutAllActions.AddAction(NewVarAction);
-				}
+				SectionEntry.Value->CollectActions(DataflowInstanceInterface, Actions);
 			}
 		}
 	}
+	for (TSharedPtr<FEdGraphSchemaAction> Action : Actions)
+	{
+		OutAllActions.AddAction(Action);
+	}
 }
 
 void SDataflowMembersWidget::CollectStaticSections(TArray<int32>& StaticSectionIDs)
@@ -485,7 +606,19 @@ FReply SDataflowMembersWidget::OnActionDragged(const TArray<TSharedPtr<FEdGraphS
 
 void SDataflowMembersWidget::OnActionDoubleClicked(const TArray< TSharedPtr<FEdGraphSchemaAction> >& InActions)
 {
-	// TODO : select the first variable node in the graph
+	if (TSharedPtr<FDataflowEditorToolkit> Toolkit = EditorToolkitWeakPtr.Pin())
+	{
+		for (TSharedPtr<FEdGraphSchemaAction> Action : InActions)
+		{
+			if (Action)
+			{
+				if (const TSharedPtr<ISection> Section = GetSectionById(Action->GetSectionID()))
+				{
+					Section->OnDoubleClicked(*Action, *Toolkit);
+				}
+			}
+		}
+	}
 }
 
 TSharedPtr<SWidget> SDataflowMembersWidget::OnContextMenuOpening()
diff --git a/Engine/Plugins/Experimental/Dataflow/Source/DataflowEditor/Private/Dataflow/DataflowSubGraphPaletteItem.cpp b/Engine/Plugins/Experimental/Dataflow/Source/DataflowEditor/Private/Dataflow/DataflowSubGraphPaletteItem.cpp
new file mode 100644
--- /dev/null
+++ b/Engine/Plugins/Experimental/Dataflow/Source/DataflowEditor/Private/Dataflow/DataflowSubGraphPaletteItem.cpp
@@ -0,0 +1,65 @@
+// Copyright Epic Games, Inc. All Rights Reserved.
+
+#include "Dataflow/DataflowSubGraphPaletteItem.h"
+
+#include "Dataflow/DataflowGraphSchemaAction.h"
+#include "DetailLayoutBuilder.h"
+#include "Widgets/Text/SInlineEditableTextBlock.h"
+
+void SDataflowSubGraphPaletteItem::Construct(const FArguments& InArgs, FCreateWidgetForActionData* const InCreateData, TSharedPtr<SDataflowGraphEditor> InEditor)
+{
+	if (InCreateData)
+	{
+		if (InCreateData->Action || InCreateData->Action->GetTypeId() == FEdGraphSchemaAction_DataflowSubGraph::StaticGetTypeId())
+		{
+			SubGraphAction = StaticCastSharedPtr<FEdGraphSchemaAction_DataflowSubGraph>(InCreateData->Action);
+		}
+	}
+	SGraphPaletteItem::Construct(SGraphPaletteItem::FArguments(), InCreateData);
+}
+
+TSharedRef<SWidget> SDataflowSubGraphPaletteItem::CreateTextSlotWidget(FCreateWidgetForActionData* const InCreateData, TAttribute<bool> bIsReadOnlyIn)
+{
+	if (SubGraphAction == nullptr)
+	{
+		return SNullWidget::NullWidget;
+	}
+
+	const FText VariableName = FText::FromString(SubGraphAction->GetSubGraphName().ToString());
+
+	TSharedPtr<SInlineEditableTextBlock> EditableTextElement;
+
+	// create a type selection widget
+	TSharedPtr<SWidget> Widget =
+		SNew(SHorizontalBox)
+
+		// search bar 
+		+ SHorizontalBox::Slot()
+		.AutoWidth()
+		.HAlign(HAlign_Left)
+		[
+			SAssignNew(EditableTextElement, SInlineEditableTextBlock)
+				.Text(VariableName)
+				.OnTextCommitted(this, &SDataflowSubGraphPaletteItem::OnNameTextCommitted)
+				.OnVerifyTextChanged(this, &SDataflowSubGraphPaletteItem::OnNameTextVerifyChanged)
+		]
+		;
+
+	InlineRenameWidget = EditableTextElement.ToSharedRef();
+	InCreateData->OnRenameRequest->BindSP(InlineRenameWidget.Get(), &SInlineEditableTextBlock::EnterEditingMode);
+
+	return Widget.ToSharedRef();
+}
+
+bool SDataflowSubGraphPaletteItem::OnNameTextVerifyChanged(const FText& InNewText, FText& OutErrorMessage)
+{
+	return SubGraphAction ? SubGraphAction->CanRenameItem(InNewText) : false;
+}
+
+void SDataflowSubGraphPaletteItem::OnNameTextCommitted(const FText& InNewText, ETextCommit::Type InTextCommit)
+{
+	if (SubGraphAction)
+	{
+		SubGraphAction->RenameItem(InNewText);
+	}
+}
diff --git a/Engine/Plugins/Experimental/Dataflow/Source/DataflowEditor/Private/Dataflow/DataflowSubGraphPaletteItem.h b/Engine/Plugins/Experimental/Dataflow/Source/DataflowEditor/Private/Dataflow/DataflowSubGraphPaletteItem.h
new file mode 100644
--- /dev/null
+++ b/Engine/Plugins/Experimental/Dataflow/Source/DataflowEditor/Private/Dataflow/DataflowSubGraphPaletteItem.h
@@ -0,0 +1,33 @@
+// Copyright Epic Games, Inc. All Rights Reserved.
+
+#pragma once
+
+#include "SGraphPalette.h"
+
+class SDataflowGraphEditor;
+struct FEdGraphSchemaAction_DataflowSubGraph;
+
+/**
+* Widget for displaying a subgraph entry in SDataflowMembersWidget
+*/
+class SDataflowSubGraphPaletteItem : public SGraphPaletteItem
+{
+public:
+	SLATE_BEGIN_ARGS(SDataflowSubGraphPaletteItem) {}
+	SLATE_END_ARGS()
+
+	void Construct(const FArguments& InArgs, FCreateWidgetForActionData* const InCreateData, TSharedPtr<SDataflowGraphEditor> InEditor);
+
+private:
+	// SGraphPaletteItem Interface
+	virtual TSharedRef<SWidget> CreateTextSlotWidget(FCreateWidgetForActionData* const InCreateData, TAttribute<bool> bIsReadOnly) override;
+	virtual bool OnNameTextVerifyChanged(const FText& InNewText, FText& OutErrorMessage) override;
+	virtual void OnNameTextCommitted(const FText& NewText, ETextCommit::Type InTextCommit) override;
+	// End of SGraphPaletteItem Interface
+
+private:
+	TSharedPtr<FEdGraphSchemaAction_DataflowSubGraph> SubGraphAction;
+
+	// held for OnRequestRename calls
+	TSharedPtr<SInlineEditableTextBlock> InlineRenameWidget;
+};
diff --git a/Engine/Plugins/Experimental/Dataflow/Source/DataflowEditor/Public/Dataflow/DataflowGraphSchemaAction.h b/Engine/Plugins/Experimental/Dataflow/Source/DataflowEditor/Public/Dataflow/DataflowGraphSchemaAction.h
--- a/Engine/Plugins/Experimental/Dataflow/Source/DataflowEditor/Public/Dataflow/DataflowGraphSchemaAction.h
+++ b/Engine/Plugins/Experimental/Dataflow/Source/DataflowEditor/Public/Dataflow/DataflowGraphSchemaAction.h
@@ -20,8 +20,8 @@ namespace UE::Dataflow
 	enum class ESchemaActionSectionID: int32
 	{
 		NONE = 0,
-		VARIABLES,
 		SUBGRAPHS,
+		VARIABLES,
 	};
 }
 
@@ -171,7 +171,7 @@ public:
 	static FName StaticGetTypeId() { static FName Type("FEdGraphSchemaAction_DataflowSubGraph"); return Type; }
 	virtual FName GetTypeId() const override { return StaticGetTypeId(); }
 
-	const FName GetSubgraphName() const
+	const FName GetSubGraphName() const
 	{
 		return SubGraphName;
 	}
diff --git a/Engine/Plugins/Experimental/Dataflow/Source/DataflowEditor/Public/Dataflow/DataflowMembersWidget.h b/Engine/Plugins/Experimental/Dataflow/Source/DataflowEditor/Public/Dataflow/DataflowMembersWidget.h
--- a/Engine/Plugins/Experimental/Dataflow/Source/DataflowEditor/Public/Dataflow/DataflowMembersWidget.h
+++ b/Engine/Plugins/Experimental/Dataflow/Source/DataflowEditor/Public/Dataflow/DataflowMembersWidget.h
@@ -6,6 +6,7 @@
 
 class UDataflow;
 class UDataflowEditorVariables;
+class FDataflowEditorToolkit;
 class FUICommandList;
 class FUICommandInfo;
 class SSearchBox;
@@ -32,7 +33,7 @@ public:
 	SLATE_BEGIN_ARGS(SDataflowMembersWidget) {}
 	SLATE_END_ARGS()
 
-	void Construct(const FArguments& InArgs, TSharedPtr<SDataflowGraphEditor> InGraphEditor);
+	void Construct(const FArguments& InArgs, TSharedPtr<FDataflowEditorToolkit> InEditorToolkit);
 	~SDataflowMembersWidget();
 
 public:
@@ -59,6 +60,11 @@ public:
 		virtual void OnPaste(FEdGraphSchemaAction&) const = 0;
 		virtual void OnDuplicate(FEdGraphSchemaAction&) const = 0;
 		virtual void OnDelete(FEdGraphSchemaAction&) const = 0;
+
+		virtual void OnDoubleClicked(FEdGraphSchemaAction&, FDataflowEditorToolkit&) const = 0;
+
+		virtual void CollectActions(IDataflowInstanceInterface* Interface, TArray<TSharedPtr<FEdGraphSchemaAction>>& OutActions) const = 0;
+		virtual TSharedRef<SWidget> CreateWidgetForAction(FCreateWidgetForActionData* const InCreateData, TSharedPtr<SDataflowGraphEditor> Editor) const = 0;
 	};
 
 private:
@@ -66,6 +72,8 @@ private:
 	virtual FReply OnKeyDown(const FGeometry& MyGeometry, const FKeyEvent& InKeyEvent);
 
 	void OnObjectPropertyChanged(UObject* InObject, FPropertyChangedEvent& InPropertyChangedEvent);
+	void OnSubGraphsChanged(const UDataflow* InDataflowAsset, FName SubGraphName);
+
 	void OnFilterTextChanged(const FText& InFilterText);
 	FText GetFilterText() const;
 	TSharedRef<SWidget> OnCreateWidgetForAction(FCreateWidgetForActionData* const InCreateData);
@@ -116,8 +124,10 @@ private:
 	void RefreshVariableOverrideDetailView();
 	void OverridesDetailsViewFinishedChangingProperties(const FPropertyChangedEvent& InPropertyChangedEvent);
 
+	const TSharedPtr<SDataflowGraphEditor> GetGraphEditor() const;
+
 private:
-	TSharedPtr<SDataflowGraphEditor> GraphEditor;
+	TWeakPtr<FDataflowEditorToolkit> EditorToolkitWeakPtr;
 
 	/** asset being edited by the graph - maybe a geometry collection with a dataflow asset bound to it */
 	TWeakObjectPtr<UObject> EditedAsset;
