Dataflow toolkit : make sure to unregister all tab we have registered

diff --git a/Engine/Plugins/Experimental/Dataflow/Source/DataflowEditor/Private/Dataflow/DataflowEditorToolkit.cpp b/Engine/Plugins/Experimental/Dataflow/Source/DataflowEditor/Private/Dataflow/DataflowEditorToolkit.cpp
--- a/Engine/Plugins/Experimental/Dataflow/Source/DataflowEditor/Private/Dataflow/DataflowEditorToolkit.cpp
+++ b/Engine/Plugins/Experimental/Dataflow/Source/DataflowEditor/Private/Dataflow/DataflowEditorToolkit.cpp
@@ -1735,7 +1735,12 @@ void FDataflowEditorToolkit::UnregisterTabSpawners(const TSharedRef<class FTabMa
 {
 	FBaseAssetToolkit::UnregisterTabSpawners(InTabManager);
 
+	InTabManager->UnregisterTabSpawner(ViewportTabID);
+	InTabManager->UnregisterTabSpawner(SimulationViewportTabId);
+	InTabManager->UnregisterTabSpawner(DetailsTabID);
+	InTabManager->UnregisterTabSpawner(PreviewSceneTabId);
 	InTabManager->UnregisterTabSpawner(GraphCanvasTabId);
+	InTabManager->UnregisterTabSpawner(SubGraphCanvasTabId);
 	InTabManager->UnregisterTabSpawner(NodeDetailsTabId);
 	InTabManager->UnregisterTabSpawner(SkeletonViewTabId);
 	InTabManager->UnregisterTabSpawner(SelectionViewTabId_1);
@@ -1746,7 +1751,8 @@ void FDataflowEditorToolkit::UnregisterTabSpawners(const TSharedRef<class FTabMa
 	InTabManager->UnregisterTabSpawner(CollectionSpreadSheetTabId_2);
 	InTabManager->UnregisterTabSpawner(CollectionSpreadSheetTabId_3);
 	InTabManager->UnregisterTabSpawner(CollectionSpreadSheetTabId_4);
-	InTabManager->UnregisterTabSpawner(SimulationViewportTabId);
+	InTabManager->UnregisterTabSpawner(SimulationVisualizationTabId);
+	InTabManager->UnregisterTabSpawner(MembersWidgetTabId);
 }
 
 void FDataflowEditorToolkit::OnTabClosed(TSharedRef<SDockTab> Tab)
