Dataflow : reference the dataflow asset directly from the members widget to allow non-bound dataflow assets to show their sub-graphs and variables

diff --git a/Engine/Plugins/Experimental/Dataflow/Source/DataflowEditor/Private/Dataflow/DataflowGraphSchemaAction.cpp b/Engine/Plugins/Experimental/Dataflow/Source/DataflowEditor/Private/Dataflow/DataflowGraphSchemaAction.cpp
--- a/Engine/Plugins/Experimental/Dataflow/Source/DataflowEditor/Private/Dataflow/DataflowGraphSchemaAction.cpp
+++ b/Engine/Plugins/Experimental/Dataflow/Source/DataflowEditor/Private/Dataflow/DataflowGraphSchemaAction.cpp
@@ -19,7 +19,7 @@
 
 TSharedPtr<FAssetSchemaAction_Dataflow_CreateNode_DataflowEdNode> FAssetSchemaAction_Dataflow_CreateNode_DataflowEdNode::CreateAction(const UEdGraph* ParentGraph, const FName& InNodeTypeName, const FName& InOverrideNodeName)
 {
-	if (const UDataflow* Dataflow = Cast<UDataflow>(ParentGraph))
+	if (const UDataflow* Dataflow = UDataflow::GetDataflowAssetFromEdGraph(ParentGraph))
 	{
 		if (UE::Dataflow::FNodeFactory* Factory = UE::Dataflow::FNodeFactory::GetInstance())
 		{
@@ -69,7 +69,7 @@ FEdGraphSchemaAction_DataflowVariable::FEdGraphSchemaAction_DataflowVariable()
 	: FEdGraphSchemaAction()
 {}
 
-FEdGraphSchemaAction_DataflowVariable::FEdGraphSchemaAction_DataflowVariable(IDataflowInstanceInterface* InDataflowInstanceInterface, const FPropertyBagPropertyDesc& PropertyDesc)
+FEdGraphSchemaAction_DataflowVariable::FEdGraphSchemaAction_DataflowVariable(UDataflow* InDataflowAsset, const FPropertyBagPropertyDesc& PropertyDesc)
 	: FEdGraphSchemaAction(
 		FText::FromString(CategoryFromFullName(PropertyDesc.Name)),
 		LOCTEXT("Description", "Description"),
@@ -81,7 +81,7 @@ FEdGraphSchemaAction_DataflowVariable::FEdGraphSchemaAction_DataflowVariable(IDa
 	, VariableName(NameFromFullName(PropertyDesc.Name))
 	, VariableCategory(CategoryFromFullName(PropertyDesc.Name))
 	, VariableType(UE::StructUtils::GetPropertyDescAsPin(PropertyDesc))
-	, DataflowInstanceInterface(InDataflowInstanceInterface)
+	, DataflowAssetWeakPtr(InDataflowAsset)
 {
 }
 
@@ -107,18 +107,9 @@ FString FEdGraphSchemaAction_DataflowVariable::NameFromFullName(FName FullName)
 	return StrFullName;
 }
 
-UDataflow* FEdGraphSchemaAction_DataflowVariable::GetDataflowAsset() const
-{
-	if (DataflowInstanceInterface)
-	{
-		return DataflowInstanceInterface->GetDataflowInstance().GetDataflowAsset();
-	}
-	return nullptr;
-}
-
 bool FEdGraphSchemaAction_DataflowVariable::CanRenameItem(FText NewName) const
 {
-	if (const UDataflow* DataflowAsset = GetDataflowAsset())
+	if (TStrongObjectPtr<const UDataflow> DataflowAsset = DataflowAssetWeakPtr.Pin())
 	{
 		const FString NewVariableName(NewName.ToString());
 		if (NewVariableName.Contains("|"))
@@ -133,7 +124,7 @@ bool FEdGraphSchemaAction_DataflowVariable::CanRenameItem(FText NewName) const
 
 void FEdGraphSchemaAction_DataflowVariable::RenameItem(FText NewName)
 {
-	if (UDataflow* DataflowAsset = GetDataflowAsset())
+	if (TStrongObjectPtr<UDataflow> DataflowAsset = DataflowAssetWeakPtr.Pin())
 	{
 		FName NewVariableName;
 		if (VariableCategory.IsEmpty())
@@ -144,36 +135,48 @@ void FEdGraphSchemaAction_DataflowVariable::RenameItem(FText NewName)
 		{
 			NewVariableName = FName(FString::Format(TEXT("{0}|{1}"), { VariableCategory, NewName.ToString() }));
 		}
-		UE::Dataflow::FEditAssetUtils::RenameVariable(DataflowAsset, FullVariableName, NewVariableName);
+		UE::Dataflow::FEditAssetUtils::RenameVariable(DataflowAsset.Get(), FullVariableName, NewVariableName);
 	}
 }
 
 void FEdGraphSchemaAction_DataflowVariable::CopyItemToClipboard()
 {
-	UE::Dataflow::FEditAssetUtils::CopyVariableToClipboard(GetDataflowAsset(), GetFullVariableName());
+	if (TStrongObjectPtr<UDataflow> DataflowAsset = DataflowAssetWeakPtr.Pin())
+	{
+		UE::Dataflow::FEditAssetUtils::CopyVariableToClipboard(DataflowAsset.Get(), GetFullVariableName());
+	}
 }
 
 void FEdGraphSchemaAction_DataflowVariable::PasteItemFromClipboard()
 {
-	UE::Dataflow::FEditAssetUtils::PasteVariableFromClipboard(GetDataflowAsset());
+	if (TStrongObjectPtr<UDataflow> DataflowAsset = DataflowAssetWeakPtr.Pin())
+	{
+		UE::Dataflow::FEditAssetUtils::PasteVariableFromClipboard(DataflowAsset.Get());
+	}
 }
 
 void FEdGraphSchemaAction_DataflowVariable::SetVariableType(const FEdGraphPinType& PinType)
 {
-	if (UDataflow* DataflowAsset = GetDataflowAsset())
+	if (TStrongObjectPtr<UDataflow> DataflowAsset = DataflowAssetWeakPtr.Pin())
 	{
-		UE::Dataflow::FEditAssetUtils::SetVariableType(DataflowAsset, FullVariableName, PinType);
+		UE::Dataflow::FEditAssetUtils::SetVariableType(DataflowAsset.Get(), FullVariableName, PinType);
 	}
 }
 
 void FEdGraphSchemaAction_DataflowVariable::DeleteItem()
 {
-	UE::Dataflow::FEditAssetUtils::DeleteVariable(GetDataflowAsset(), GetFullVariableName());
+	if (TStrongObjectPtr<UDataflow> DataflowAsset = DataflowAssetWeakPtr.Pin())
+	{
+		UE::Dataflow::FEditAssetUtils::DeleteVariable(DataflowAsset.Get(), GetFullVariableName());
+	}
 }
 
 void FEdGraphSchemaAction_DataflowVariable::DuplicateItem()
 {
-	UE::Dataflow::FEditAssetUtils::DuplicateVariable(GetDataflowAsset(), GetFullVariableName());
+	if (TStrongObjectPtr<UDataflow> DataflowAsset = DataflowAssetWeakPtr.Pin())
+	{
+		UE::Dataflow::FEditAssetUtils::DuplicateVariable(DataflowAsset.Get(), GetFullVariableName());
+	}
 }
 
 /** Execute this action, given the graph and schema, and possibly a pin that we were dragged from. Returns a node that was created by this action (if any). */
@@ -190,9 +193,9 @@ UEdGraphNode* FEdGraphSchemaAction_DataflowVariable::PerformAction(class UEdGrap
 		{
 			if (TSharedPtr<FGetDataflowVariableNode> VariableNode = StaticCastSharedPtr<FGetDataflowVariableNode>(DataflowEdNode->GetDataflowNode()))
 			{
-				if (UDataflow* DataflowAsset = GetDataflowAsset())
+				if (TStrongObjectPtr<UDataflow> DataflowAsset = DataflowAssetWeakPtr.Pin())
 				{
-					VariableNode->SetVariable(DataflowAsset, FullVariableName);
+					VariableNode->SetVariable(DataflowAsset.Get(), FullVariableName);
 					DataflowEdNode->UpdatePinsFromDataflowNode();
 					ParentGraph->NotifyNodeChanged(DataflowEdNode);
 				}
@@ -222,7 +225,7 @@ FEdGraphSchemaAction_DataflowSubGraph::FEdGraphSchemaAction_DataflowSubGraph()
 {
 }
 
-FEdGraphSchemaAction_DataflowSubGraph::FEdGraphSchemaAction_DataflowSubGraph(IDataflowInstanceInterface* InDataflowInstanceInterface, FName InSubgraphName)
+FEdGraphSchemaAction_DataflowSubGraph::FEdGraphSchemaAction_DataflowSubGraph(UDataflow* InDataflowAsset, FName InSubgraphName)
 	: FEdGraphSchemaAction(
 		FText::GetEmpty(),
 		LOCTEXT("Description", "Description"),
@@ -231,7 +234,7 @@ FEdGraphSchemaAction_DataflowSubGraph::FEdGraphSchemaAction_DataflowSubGraph(IDa
 		/*Keywords*/FText::GetEmpty(),
 		(int32)UE::Dataflow::ESchemaActionSectionID::SUBGRAPHS)
 	, SubGraphName(InSubgraphName)
-	, DataflowInstanceInterface(InDataflowInstanceInterface)
+	, DataflowAssetWeakPtr(InDataflowAsset)
 {
 }
 
@@ -244,7 +247,7 @@ UEdGraphNode* FEdGraphSchemaAction_DataflowSubGraph::PerformAction(class UEdGrap
 
 bool FEdGraphSchemaAction_DataflowSubGraph::CanRenameItem(FText NewName) const
 {
-	if (const UDataflow* DataflowAsset = GetDataflowAsset())
+	if (TStrongObjectPtr<UDataflow> DataflowAsset = DataflowAssetWeakPtr.Pin())
 	{
 		const bool bSubGraphAlreadyExists = (nullptr != DataflowAsset->FindSubGraphByName(FName(NewName.ToString())));
 		return !bSubGraphAlreadyExists;
@@ -277,14 +280,4 @@ void FEdGraphSchemaAction_DataflowSubGraph::DuplicateItem()
 	// NOT YET IMPLEMENTED
 }
 
-UDataflow* FEdGraphSchemaAction_DataflowSubGraph::GetDataflowAsset() const
-{
-	if (DataflowInstanceInterface)
-	{
-		return DataflowInstanceInterface->GetDataflowInstance().GetDataflowAsset();
-	}
-	return nullptr;
-}
-
-
 #undef LOCTEXT_NAMESPACE
\ No newline at end of file
diff --git a/Engine/Plugins/Experimental/Dataflow/Source/DataflowEditor/Private/Dataflow/DataflowMembersWidget.cpp b/Engine/Plugins/Experimental/Dataflow/Source/DataflowEditor/Private/Dataflow/DataflowMembersWidget.cpp
--- a/Engine/Plugins/Experimental/Dataflow/Source/DataflowEditor/Private/Dataflow/DataflowMembersWidget.cpp
+++ b/Engine/Plugins/Experimental/Dataflow/Source/DataflowEditor/Private/Dataflow/DataflowMembersWidget.cpp
@@ -131,20 +131,17 @@ namespace UE::Dataflow::Private
 			return SNew(SDataflowSubGraphPaletteItem, InCreateData, Editor);
 		}
 
-		virtual void CollectActions(IDataflowInstanceInterface* Interface, TArray<TSharedPtr<FEdGraphSchemaAction>>& OutActions) const override
+		virtual void CollectActions(UDataflow* DataflowAsset, TArray<TSharedPtr<FEdGraphSchemaAction>>& OutActions) const override
 		{
-			if (Interface)
+			if (DataflowAsset)
 			{
-				if (const UDataflow* DataflowAsset = Interface->GetDataflowInstance().GetDataflowAsset())
+				for (TObjectPtr<const UDataflowSubGraph> SubGraph : DataflowAsset->GetSubGraphs())
 				{
-					for (TObjectPtr<const UDataflowSubGraph> SubGraph : DataflowAsset->GetSubGraphs())
+					if (ensure(SubGraph))
 					{
-						if (ensure(SubGraph))
-						{
-							TSharedPtr<FEdGraphSchemaAction_DataflowSubGraph> NewSubGraphAction =
-								MakeShareable(new FEdGraphSchemaAction_DataflowSubGraph(Interface, SubGraph->GetFName()));
-							OutActions.Add(NewSubGraphAction);
-						}
+						TSharedPtr<FEdGraphSchemaAction_DataflowSubGraph> NewSubGraphAction =
+							MakeShareable(new FEdGraphSchemaAction_DataflowSubGraph(DataflowAsset, SubGraph->GetFName()));
+						OutActions.Add(NewSubGraphAction);
 					}
 				}
 			}
@@ -181,20 +178,17 @@ namespace UE::Dataflow::Private
 			return SNew(SDataflowVariablePaletteItem, InCreateData, Editor);
 		}
 
-		virtual void CollectActions(IDataflowInstanceInterface* Interface, TArray<TSharedPtr<FEdGraphSchemaAction>>& OutActions) const override
+		virtual void CollectActions(UDataflow* DataflowAsset, TArray<TSharedPtr<FEdGraphSchemaAction>>& OutActions) const override
 		{
-			if (Interface)
+			if (DataflowAsset)
 			{
-				if (const UDataflow* DataflowAsset = Interface->GetDataflowInstance().GetDataflowAsset())
+				if (const UPropertyBag* PropertyBag = DataflowAsset->Variables.GetPropertyBagStruct())
 				{
-					if (const UPropertyBag* PropertyBag = DataflowAsset->Variables.GetPropertyBagStruct())
+					for (const FPropertyBagPropertyDesc& ConfigDesc : PropertyBag->GetPropertyDescs())
 					{
-						for (const FPropertyBagPropertyDesc& ConfigDesc : PropertyBag->GetPropertyDescs())
-						{
-							TSharedPtr<FEdGraphSchemaAction_DataflowVariable> NewVarAction =
-								MakeShareable(new FEdGraphSchemaAction_DataflowVariable(Interface, ConfigDesc));
-							OutActions.Add(NewVarAction);
-						}
+						TSharedPtr<FEdGraphSchemaAction_DataflowVariable> NewVarAction =
+							MakeShareable(new FEdGraphSchemaAction_DataflowVariable(DataflowAsset, ConfigDesc));
+						OutActions.Add(NewVarAction);
 					}
 				}
 			}
@@ -212,7 +206,7 @@ namespace UE::Dataflow::Private
 void SDataflowMembersWidget::Construct(const FArguments& InArgs, TSharedPtr<FDataflowEditorToolkit> InEditorToolkit)
 {
 	EditorToolkitWeakPtr = InEditorToolkit;
-	CacheEditedAsset();
+	CacheAssets();
 
 	InitializeCommands();
 	InitializeSections();
@@ -359,7 +353,7 @@ void SDataflowMembersWidget::RefreshVariableOverrideDetailView()
 	// now assign the corresponding objects and structures
 	if (IDataflowInstanceInterface* Interface = GetDataflowInstanceInterface())
 	{
-		OverridesDetailsView->GetDetailsView()->SetObject(EditedAsset.Get());
+		OverridesDetailsView->GetDetailsView()->SetObject(EditedAssetWeakPtr.Get());
 		TSharedPtr<FStructOnScope> StructOnScope = Interface->GetDataflowInstance().MakeStructOnScope();
 		OverridesDetailsView->SetStructureData(StructOnScope);
 	}
@@ -448,8 +442,17 @@ FReply SDataflowMembersWidget::OnKeyDown(const FGeometry& MyGeometry, const FKey
 	return FReply::Unhandled();
 }
 
-void SDataflowMembersWidget::CacheEditedAsset()
+void SDataflowMembersWidget::CacheAssets()
 {
+	if (EditorToolkitWeakPtr.IsValid())
+	{
+		TObjectPtr<UDataflowBaseContent> Content = EditorToolkitWeakPtr.Pin()->GetEditorContent();
+		if (Content)
+		{
+			DataflowAssetWeakPtr = Content->GetDataflowAsset();
+		}
+	}
+
 	const TSharedPtr<SDataflowGraphEditor> GraphEditor = GetGraphEditor();
 	if (!GraphEditor)
 	{
@@ -468,19 +471,14 @@ void SDataflowMembersWidget::CacheEditedAsset()
 		return;
 	}
 
-	EditedAsset = EngineContext->Owner;
+	EditedAssetWeakPtr = EngineContext->Owner;
 }
 
 IDataflowInstanceInterface* SDataflowMembersWidget::GetDataflowInstanceInterface() const
 {
-	return Cast<IDataflowInstanceInterface>(EditedAsset);
-}
-
-UDataflow* SDataflowMembersWidget::GetDataflowAsset() const
-{
-	if (IDataflowInstanceInterface* DataflowInstanceInterface = GetDataflowInstanceInterface())
+	if (EditedAssetWeakPtr.IsValid())
 	{
-		return DataflowInstanceInterface->GetDataflowInstance().GetDataflowAsset();
+		return Cast<IDataflowInstanceInterface>(EditedAssetWeakPtr.Get());
 	}
 	return nullptr;
 }
@@ -488,7 +486,7 @@ UDataflow* SDataflowMembersWidget::GetDataflowAsset() const
 void SDataflowMembersWidget::OverridesDetailsViewFinishedChangingProperties(const FPropertyChangedEvent& InPropertyChangedEvent)
 {
 	const FName VariableName = InPropertyChangedEvent.GetPropertyName();
-	if (UDataflow* DataflowAsset = GetDataflowAsset())
+	if (TStrongObjectPtr<UDataflow> DataflowAsset = DataflowAssetWeakPtr.Pin())
 	{
 		for (const TSharedPtr<FDataflowNode>& Node : DataflowAsset->Dataflow->GetNodes())
 		{
@@ -509,7 +507,7 @@ void SDataflowMembersWidget::OnObjectPropertyChanged(UObject* InObject, FPropert
 
 	if (InObject)
 	{
-		if (InObject == GetDataflowAsset())
+		if (InObject == DataflowAssetWeakPtr)
 		{
 			if (InPropertyChangedEvent.GetMemberPropertyName() == DataflowVariablesPropertyName
 				|| InPropertyChangedEvent.GetPropertyName() == DataflowVariablesPropertyName)
@@ -518,7 +516,7 @@ void SDataflowMembersWidget::OnObjectPropertyChanged(UObject* InObject, FPropert
 			}
 			RefreshVariableOverrideDetailView();
 		}
-		else if (InObject == EditedAsset)
+		else if (InObject == EditedAssetWeakPtr)
 		{
 			RefreshVariableOverrideDetailView();
 		}
@@ -563,13 +561,13 @@ TSharedRef<SWidget> SDataflowMembersWidget::OnCreateWidgetForAction(FCreateWidge
 void SDataflowMembersWidget::CollectAllActions(FGraphActionListBuilderBase& OutAllActions)
 {
 	TArray<TSharedPtr<FEdGraphSchemaAction>> Actions;
-	if (IDataflowInstanceInterface* DataflowInstanceInterface = GetDataflowInstanceInterface())
+	if (TStrongObjectPtr<UDataflow> DataflowAsset = DataflowAssetWeakPtr.Pin())
 	{
 		for (const TPair<int32, TSharedPtr<ISection>>& SectionEntry : SectionMap)
 		{
 			if (ensure(SectionEntry.Value))
 			{
-				SectionEntry.Value->CollectActions(DataflowInstanceInterface, Actions);
+				SectionEntry.Value->CollectActions(DataflowAsset.Get(), Actions);
 			}
 		}
 	}
@@ -595,11 +593,8 @@ FReply SDataflowMembersWidget::OnActionDragged(const TArray<TSharedPtr<FEdGraphS
 	if (InAction->GetTypeId() == FEdGraphSchemaAction_DataflowVariable::StaticGetTypeId())
 	{
 		TSharedPtr<FEdGraphSchemaAction_DataflowVariable> VariableAction = StaticCastSharedPtr<FEdGraphSchemaAction_DataflowVariable>(InAction);
-		if (UDataflow* DataflowAsset = GetDataflowAsset())
-		{
-			TSharedRef<FGraphSchemaActionDragDropAction_DataflowVariable> DragOperation = FGraphSchemaActionDragDropAction_DataflowVariable::New(VariableAction);
-			return FReply::Handled().BeginDragDrop(DragOperation);
-		}
+		TSharedRef<FGraphSchemaActionDragDropAction_DataflowVariable> DragOperation = FGraphSchemaActionDragDropAction_DataflowVariable::New(VariableAction);
+		return FReply::Handled().BeginDragDrop(DragOperation);
 	}
 	return FReply::Unhandled();
 }
diff --git a/Engine/Plugins/Experimental/Dataflow/Source/DataflowEditor/Public/Dataflow/DataflowGraphSchemaAction.h b/Engine/Plugins/Experimental/Dataflow/Source/DataflowEditor/Public/Dataflow/DataflowGraphSchemaAction.h
--- a/Engine/Plugins/Experimental/Dataflow/Source/DataflowEditor/Public/Dataflow/DataflowGraphSchemaAction.h
+++ b/Engine/Plugins/Experimental/Dataflow/Source/DataflowEditor/Public/Dataflow/DataflowGraphSchemaAction.h
@@ -60,7 +60,7 @@ struct FEdGraphSchemaAction_DataflowVariable : public FEdGraphSchemaAction
 public:
 	FEdGraphSchemaAction_DataflowVariable();
 
-	FEdGraphSchemaAction_DataflowVariable(IDataflowInstanceInterface* InDataflowInstanceInterface, const FPropertyBagPropertyDesc& PropertyDesc);
+	FEdGraphSchemaAction_DataflowVariable(UDataflow* InDataflowAsset, const FPropertyBagPropertyDesc& PropertyDesc);
 
 	// Simple type info
 	static FName StaticGetTypeId() { static FName Type("FEdGraphSchemaAction_DataflowVariable"); return Type; }
@@ -95,8 +95,6 @@ public:
 	void DeleteItem();
 	void DuplicateItem();
 
-	UDataflow* GetDataflowAsset() const;
-
 private:
 	static FString CategoryFromFullName(FName FullName);
 	static FString NameFromFullName(FName FullName);
@@ -124,11 +122,8 @@ private:
 	*/
 	FEdGraphPinType VariableType;
 
-	/**
-	* interface from where to read the variable information
-	* this allow access to the dataflow asset as well as the override from the asset using this dataflow asset
-	*/
-	IDataflowInstanceInterface* DataflowInstanceInterface = nullptr;;
+	/** Dataflow asset associated with this action */
+	TWeakObjectPtr<UDataflow> DataflowAssetWeakPtr;
 };
 
 ////////////////////////////////////////////////////////////////////////////////////////////////////////////////////////////////
@@ -165,7 +160,7 @@ struct FEdGraphSchemaAction_DataflowSubGraph : public FEdGraphSchemaAction
 public:
 	FEdGraphSchemaAction_DataflowSubGraph();
 
-	FEdGraphSchemaAction_DataflowSubGraph(IDataflowInstanceInterface* InDataflowInstanceInterface, FName SubGraphName);
+	FEdGraphSchemaAction_DataflowSubGraph(UDataflow* InDataflowAsset, FName SubGraphName);
 
 	// Simple type info
 	static FName StaticGetTypeId() { static FName Type("FEdGraphSchemaAction_DataflowSubGraph"); return Type; }
@@ -189,15 +184,10 @@ public:
 	void DeleteItem();
 	void DuplicateItem();
 
-	UDataflow* GetDataflowAsset() const;
-
 private:
 	/** Name of the subgraph */
 	FName SubGraphName;
 
-	/**
-	* interface from where to read the variable information
-	* this allow access to the dataflow asset as well as the override from the asset using this dataflow asset
-	*/
-	IDataflowInstanceInterface* DataflowInstanceInterface = nullptr;;
+	/** Dataflow asset associated with this action */
+	TWeakObjectPtr<UDataflow> DataflowAssetWeakPtr;
 };
diff --git a/Engine/Plugins/Experimental/Dataflow/Source/DataflowEditor/Public/Dataflow/DataflowMembersWidget.h b/Engine/Plugins/Experimental/Dataflow/Source/DataflowEditor/Public/Dataflow/DataflowMembersWidget.h
--- a/Engine/Plugins/Experimental/Dataflow/Source/DataflowEditor/Public/Dataflow/DataflowMembersWidget.h
+++ b/Engine/Plugins/Experimental/Dataflow/Source/DataflowEditor/Public/Dataflow/DataflowMembersWidget.h
@@ -63,7 +63,7 @@ public:
 
 		virtual void OnDoubleClicked(FEdGraphSchemaAction&, FDataflowEditorToolkit&) const = 0;
 
-		virtual void CollectActions(IDataflowInstanceInterface* Interface, TArray<TSharedPtr<FEdGraphSchemaAction>>& OutActions) const = 0;
+		virtual void CollectActions(UDataflow* DataflowAsset, TArray<TSharedPtr<FEdGraphSchemaAction>>& OutActions) const = 0;
 		virtual TSharedRef<SWidget> CreateWidgetForAction(FCreateWidgetForActionData* const InCreateData, TSharedPtr<SDataflowGraphEditor> Editor) const = 0;
 	};
 
@@ -97,7 +97,7 @@ private:
 
 	const TSharedPtr<ISection> GetSectionById(int32 SectionId) const;
 
-	void CacheEditedAsset();
+	void CacheAssets();
 	UDataflow* GetDataflowAsset() const;
 	IDataflowInstanceInterface* GetDataflowInstanceInterface() const;
 
@@ -129,8 +129,14 @@ private:
 private:
 	TWeakPtr<FDataflowEditorToolkit> EditorToolkitWeakPtr;
 
-	/** asset being edited by the graph - maybe a geometry collection with a dataflow asset bound to it */
-	TWeakObjectPtr<UObject> EditedAsset;
+	/** 
+	* asset being edited by the graph - maybe a geometry collection with a dataflow asset bound to it 
+	* Can be null if the dataflow asset graph is edited directly outside of the context of a specific asset
+	*/
+	TWeakObjectPtr<UObject> EditedAssetWeakPtr;
+
+	/** Dataflow asset  */
+	TWeakObjectPtr<UDataflow> DataflowAssetWeakPtr;
 
 	/** List of UI Commands for this scope */
 	TSharedPtr<FUICommandList> CommandList;
