Dataflow : Make sure subgraph are created non-transient
Mark the right graph modified when adding a node

diff --git a/Engine/Plugins/Experimental/Dataflow/Source/DataflowEditor/Private/Dataflow/DataflowAssetEditUtils.cpp b/Engine/Plugins/Experimental/Dataflow/Source/DataflowEditor/Private/Dataflow/DataflowAssetEditUtils.cpp
--- a/Engine/Plugins/Experimental/Dataflow/Source/DataflowEditor/Private/Dataflow/DataflowAssetEditUtils.cpp
+++ b/Engine/Plugins/Experimental/Dataflow/Source/DataflowEditor/Private/Dataflow/DataflowAssetEditUtils.cpp
@@ -118,9 +118,10 @@ namespace UE::Dataflow
 			{
 				if (UDataflowEdNode* EdNode = NewObject<UDataflowEdNode>(EdGraph, UDataflowEdNode::StaticClass(), DataflowNode->GetName()))
 				{
+					Dataflow->Modify();
 					EdGraph->Modify();
 
-					Dataflow->AddNode(EdNode, /*bUserAction*/true, /*bSelectNewNode*/false);
+					EdGraph->AddNode(EdNode, /*bUserAction*/true, /*bSelectNewNode*/false);
 					
 					EdNode->CreateNewGuid();
 					EdNode->PostPlacedNewNode();
@@ -878,7 +879,7 @@ namespace UE::Dataflow
 				UDataflowSubGraph* NewSubGraph = NewObject<UDataflowSubGraph>(&DataflowAsset, UniqueSubGraphName);
 				ensure(NewSubGraph->GetFName() == UniqueSubGraphName);
 				NewSubGraph->Schema = UDataflowSchema::StaticClass();
-				NewSubGraph->SetFlags(RF_Transient);
+				NewSubGraph->SetFlags(RF_Transactional);
 
 				DataflowAsset.AddSubGraph(NewSubGraph);
 
