Dataflow - Added new Freeze icon.

#rnx

diff --git a/Engine/Plugins/Experimental/Dataflow/Source/DataflowEditor/Private/Dataflow/DataflowEditorStyle.h b/Engine/Plugins/Experimental/Dataflow/Source/DataflowEditor/Private/Dataflow/DataflowEditorStyle.h
--- a/Engine/Plugins/Experimental/Dataflow/Source/DataflowEditor/Private/Dataflow/DataflowEditorStyle.h
+++ b/Engine/Plugins/Experimental/Dataflow/Source/DataflowEditor/Private/Dataflow/DataflowEditorStyle.h
@@ -18,6 +18,7 @@ public:
 		const FVector2D Icon16x16(16.f, 16.f);
 		const FVector2D Icon28x14(28.f, 14.f);
 		const FVector2D Icon20x20(20.f, 20.f);
+		const FVector2D Icon24x24(24.f, 24.f);
 		const FVector2D Icon40x40(40.f, 40.f);
 		const FVector2D Icon64x64(64.f, 64.f);
 
@@ -40,6 +41,8 @@ public:
 		Set("Dataflow.SelectFace", new FSlateImageBrush(RootToContentDir(TEXT("Slate/Dataflow_SelectFace_40x.png")), Icon40x40));
 		Set("Dataflow.SelectVertex", new FSlateImageBrush(RootToContentDir(TEXT("Slate/Dataflow_SelectVertex_40x.png")), Icon40x40));
 
+		Set("Dataflow.FreezeNode", new FSlateImageBrush(RootToContentDir(TEXT("Slate/Dataflow_FreezeNode_24x.png")), Icon24x24));  // There's also a 32x version if it is decided that it is too small 
+
 		// @todo(brice) Remove Example Tools
 		//const FString AttributeEditorPropertyName = "DataflowEditor." + FDataflowEditorCommandsImpl::BeginAttributeEditorToolIdentifier;
 		//Set(*AttributeEditorPropertyName, new FSlateImageBrush(RootToContentDir(TEXT("Slate/Dataflow_SelectObject_20x.png")), Icon20x20));
diff --git a/Engine/Plugins/Experimental/Dataflow/Source/DataflowEditor/Private/Dataflow/DataflowSNode.cpp b/Engine/Plugins/Experimental/Dataflow/Source/DataflowEditor/Private/Dataflow/DataflowSNode.cpp
--- a/Engine/Plugins/Experimental/Dataflow/Source/DataflowEditor/Private/Dataflow/DataflowSNode.cpp
+++ b/Engine/Plugins/Experimental/Dataflow/Source/DataflowEditor/Private/Dataflow/DataflowSNode.cpp
@@ -49,7 +49,7 @@ void SDataflowEdNode::Construct(const FArguments& InArgs, UDataflowEdNode* InNod
 	// Freeze
 	//
 	FreezeImageWidget = SNew(SImage)
-		.Image(FAppStyle::GetBrush("Icons.Lock"))
+		.Image(FDataflowEditorStyle::Get().GetBrush("Dataflow.FreezeNode"))
 		.DesiredSizeOverride(FVector2D(24.f, 24.f))
 		.Visibility_Lambda([this]()->EVisibility
 			{
