Dataflow : Move comment creation to the edit class
Remove unsused comment actions

diff --git a/Engine/Plugins/Experimental/Dataflow/Source/DataflowEditor/Private/Dataflow/DataflowAssetEditUtils.cpp b/Engine/Plugins/Experimental/Dataflow/Source/DataflowEditor/Private/Dataflow/DataflowAssetEditUtils.cpp
--- a/Engine/Plugins/Experimental/Dataflow/Source/DataflowEditor/Private/Dataflow/DataflowAssetEditUtils.cpp
+++ b/Engine/Plugins/Experimental/Dataflow/Source/DataflowEditor/Private/Dataflow/DataflowAssetEditUtils.cpp
@@ -118,6 +118,8 @@ namespace UE::Dataflow
 			{
 				if (UDataflowEdNode* EdNode = NewObject<UDataflowEdNode>(EdGraph, UDataflowEdNode::StaticClass(), DataflowNode->GetName()))
 				{
+					EdNode->SetFlags(RF_Transactional);
+
 					Dataflow->Modify();
 					EdGraph->Modify();
 
@@ -139,47 +141,52 @@ namespace UE::Dataflow
 					EdNode->NodePosX = Location.X;
 					EdNode->NodePosY = Location.Y;
 
-					EdNode->SetFlags(RF_Transactional);
-
 					return EdNode;
 				}
 			}
 			return nullptr;
 		}
 
-		UEdGraphNode_Comment* CreateCommentEdNode(UEdGraph* Graph, const FVector2D& Location, const FString& Comment, const FVector2D Size, const FLinearColor& Color, int32 FontSize)
+		UEdGraphNode_Comment* CreateCommentEdNode(UEdGraph* EdGraph, const FVector2D& Location, const FString& Comment, const FVector2D Size, const FLinearColor& Color, int32 FontSize)
 		{
-			Graph->Modify();
-			
-			UEdGraphNode_Comment* CommentTemplate = NewObject<UEdGraphNode_Comment>(Graph, FName("Comment"));
+			if (UDataflow* Dataflow = UDataflow::GetDataflowAssetFromEdGraph(EdGraph))
 			{
-				CommentTemplate->bCommentBubbleVisible_InDetailsPanel = false;
-				CommentTemplate->bCommentBubbleVisible = false;
-				CommentTemplate->bCommentBubblePinned = false;
+				if (UEdGraphNode_Comment* CommentTemplate = NewObject<UEdGraphNode_Comment>(EdGraph))
+				{
+					CommentTemplate->SetFlags(RF_Transactional);
+
+					Dataflow->Modify();
+					EdGraph->Modify();
 
-				CommentTemplate->SetFlags(RF_Transactional);
+					CommentTemplate->bCommentBubbleVisible_InDetailsPanel = false;
+					CommentTemplate->bCommentBubbleVisible = false;
+					CommentTemplate->bCommentBubblePinned = false;
 
-				// set outer to be the graph so it doesn't go away
-				CommentTemplate->Rename(NULL, Graph, REN_NonTransactional);
-				Graph->AddNode(CommentTemplate, true, /*bSelectNewNode*/false);
+					// set outer to be the graph so it doesn't go away
+					CommentTemplate->Rename(NULL, EdGraph, REN_NonTransactional);
+					EdGraph->AddNode(CommentTemplate, true, /*bSelectNewNode*/false);
 
-				CommentTemplate->CreateNewGuid();
-				CommentTemplate->PostPlacedNewNode();
-				CommentTemplate->AllocateDefaultPins();
+					CommentTemplate->CreateNewGuid();
+					CommentTemplate->PostPlacedNewNode();
+					CommentTemplate->AllocateDefaultPins();
 
-				CommentTemplate->NodePosX = Location.X;
-				CommentTemplate->NodePosY = Location.Y;
-				CommentTemplate->NodeWidth = Size.X;
-				CommentTemplate->NodeHeight = Size.Y;
-				CommentTemplate->SnapToGrid(GetDefault<UEditorStyleSettings>()->GridSnapSize);
-				CommentTemplate->CommentColor = Color;
-				CommentTemplate->FontSize = FontSize;
+					CommentTemplate->NodePosX = Location.X;
+					CommentTemplate->NodePosY = Location.Y;
+					CommentTemplate->NodeWidth = Size.X;
+					CommentTemplate->NodeHeight = Size.Y;
+					CommentTemplate->SnapToGrid(GetDefault<UEditorStyleSettings>()->GridSnapSize);
+					CommentTemplate->CommentColor = Color;
+					CommentTemplate->FontSize = FontSize;
 
-				CommentTemplate->NodeComment = Comment;
+					CommentTemplate->NodeComment = Comment;
 
-				//ParentGraph->NotifyGraphChanged();
+
+					EdGraph->NotifyGraphChanged();
+
+					return CommentTemplate;
+				}
 			}
-			return CommentTemplate;
+			return nullptr;
 		}
 
 		/** returns true if the variable was modified */
@@ -294,6 +301,37 @@ namespace UE::Dataflow
 		return EdNodeToReturn;
 	}
 
+	//////////////////////////////////////////////////////////////////////////////////////////////////////////
+	UEdGraphNode* FEditAssetUtils::AddNewComment(UEdGraph* EdGraph, const FVector2D& Location, const FVector2D& Size)
+	{
+		using namespace EditAssetUtils::Private;
+
+		const FText TransactionName{ LOCTEXT("AddNewDataflowComment", "Add New Dataflow Comment") };
+
+		UEdGraphNode* EdNodeToReturn = nullptr;
+
+		if (UDataflow* Dataflow = UDataflow::GetDataflowAssetFromEdGraph(EdGraph))
+		{
+			auto AddNewCommentInternal =
+				[&Dataflow, &EdGraph, &Location, &Size, &EdNodeToReturn]
+				(UDataflow& DataflowAsset) -> EChangeResult
+				{
+					const FString DefaultText = "Comment";
+					const FLinearColor DefaultColor = FLinearColor::White;
+					const int32 DefaultFontSize = 18;
+					if (UEdGraphNode* EdNode = CreateCommentEdNode(EdGraph, Location, DefaultText, Size, DefaultColor, DefaultFontSize))
+					{
+						EdNodeToReturn = EdNode;
+						return EChangeResult::Changed;
+					}
+					return EChangeResult::None;
+				};
+
+			ChangeDataflowAssetWithTransaction(Dataflow, TransactionName, AddNewCommentInternal, NAME_None);
+		}
+		return EdNodeToReturn;
+	}
+	
 	//////////////////////////////////////////////////////////////////////////////////////////////////////////
 	void FEditAssetUtils::DuplicateNodes(UEdGraph* EdGraph, const TArray<UEdGraphNode*>& EdNodesToDuplicate, const FVector2D& Location, TArray<UEdGraphNode*>& OutDuplicatedNodes)
 	{
diff --git a/Engine/Plugins/Experimental/Dataflow/Source/DataflowEditor/Private/Dataflow/DataflowAssetEditUtils.h b/Engine/Plugins/Experimental/Dataflow/Source/DataflowEditor/Private/Dataflow/DataflowAssetEditUtils.h
--- a/Engine/Plugins/Experimental/Dataflow/Source/DataflowEditor/Private/Dataflow/DataflowAssetEditUtils.h
+++ b/Engine/Plugins/Experimental/Dataflow/Source/DataflowEditor/Private/Dataflow/DataflowAssetEditUtils.h
@@ -45,6 +45,15 @@ namespace UE::Dataflow
 		*/
 		static UDataflowEdNode* AddNewNode(UEdGraph* EdGraph, const FVector2D& Location, const FName NodeName, const FName NodeTypeName, UEdGraphPin* FromPin);
 
+		/**
+		* Add a new comment node to the graph
+		* @param EdGraph	Main Graph or Subgraph to add the node to
+		* @param Location	Location on the graph where to place the comment
+		* @param Location	Size of the comment box
+		* @return the newly created EdGraphNode
+		*/
+		static UEdGraphNode* AddNewComment(UEdGraph* EdGraph, const FVector2D& Location, const FVector2D& Size);
+
 		/**
 		* Duplicate a node
 		* @param EdGraph			Main Graph or Subgraph to duplicate the node in 
diff --git a/Engine/Plugins/Experimental/Dataflow/Source/DataflowEditor/Private/Dataflow/DataflowSCommentNode.cpp b/Engine/Plugins/Experimental/Dataflow/Source/DataflowEditor/Private/Dataflow/DataflowSCommentNode.cpp
--- a/Engine/Plugins/Experimental/Dataflow/Source/DataflowEditor/Private/Dataflow/DataflowSCommentNode.cpp
+++ b/Engine/Plugins/Experimental/Dataflow/Source/DataflowEditor/Private/Dataflow/DataflowSCommentNode.cpp
@@ -11,11 +11,10 @@
 #include "Widgets/Layout/SBorder.h"
 #include "Settings/EditorStyleSettings.h"
 #include "ScopedTransaction.h"
+#include "Dataflow/DataflowAssetEditUtils.h"
 
 #define LOCTEXT_NAMESPACE "SDataflowEdNodeComment"
 
-constexpr const char* CopyPostfix = "_copy";
-
 //
 // Add a menu option to create a graph node.
 //
@@ -29,145 +28,12 @@ TSharedPtr<FAssetSchemaAction_Dataflow_CreateCommentNode_DataflowEdNode> FAssetS
 //
 UEdGraphNode* FAssetSchemaAction_Dataflow_CreateCommentNode_DataflowEdNode::PerformAction(class UEdGraph* ParentGraph, UEdGraphPin* FromPin, const FVector2D Location, bool bSelectNewNode)
 {
-	UEdGraphNode_Comment* CommentTemplate = NewObject<UEdGraphNode_Comment>(ParentGraph, FName("Comment"));
-
-	FVector2D SpawnLocation = Location;
 	FSlateRect Bounds;
-
 	if (GraphEditor->GetBoundsForSelectedNodes(Bounds, 50.0f))
 	{
-		CommentTemplate->SetBounds(Bounds);
-		SpawnLocation.X = CommentTemplate->NodePosX;
-		SpawnLocation.Y = CommentTemplate->NodePosY;
+		return UE::Dataflow::FEditAssetUtils::AddNewComment(ParentGraph, Bounds.GetTopLeft(), Bounds.GetSize());
 	}
-
-	CommentTemplate->bCommentBubbleVisible_InDetailsPanel = false;
-	CommentTemplate->bCommentBubbleVisible = false;
-	CommentTemplate->bCommentBubblePinned = false;
-
-	const FScopedTransaction Transaction(NSLOCTEXT("UnrealEd", "DataflowEditorNewNode", "Dataflow Editor: New Comment Node"));
-	ParentGraph->Modify();
-
-	CommentTemplate->SetFlags(RF_Transactional);
-
-	// set outer to be the graph so it doesn't go away
-	CommentTemplate->Rename(NULL, ParentGraph, REN_NonTransactional);
-	ParentGraph->AddNode(CommentTemplate, true, bSelectNewNode);
-
-	CommentTemplate->CreateNewGuid();
-	CommentTemplate->PostPlacedNewNode();
-	CommentTemplate->AllocateDefaultPins();
-	CommentTemplate->AutowireNewNode(FromPin);
-
-	CommentTemplate->NodePosX = SpawnLocation.X;
-	CommentTemplate->NodePosY = SpawnLocation.Y;
-	CommentTemplate->SnapToGrid(GetDefault<UEditorStyleSettings>()->GridSnapSize);
-
-	ParentGraph->NotifyGraphChanged();
-
-	return CommentTemplate;
-}
-
-//void FAssetSchemaAction_Dataflow_CreateCommentNode_DataflowEdNode::AddReferencedObjects(FReferenceCollector& Collector)
-//{
-//	FEdGraphSchemaAction::AddReferencedObjects(Collector);
-//	Collector.AddReferencedObject(NodeTemplate);
-//}
-
-TSharedPtr<FAssetSchemaAction_Dataflow_DuplicateCommentNode_DataflowEdNode> FAssetSchemaAction_Dataflow_DuplicateCommentNode_DataflowEdNode::CreateAction(UEdGraph* ParentGraph, const TSharedPtr<SGraphEditor>& GraphEditor)
-{
-	return MakeShared<FAssetSchemaAction_Dataflow_DuplicateCommentNode_DataflowEdNode>(GraphEditor);
-}
-
-//
-//  Duplicate comment node
-//
-UEdGraphNode* FAssetSchemaAction_Dataflow_DuplicateCommentNode_DataflowEdNode::PerformAction(class UEdGraph* ParentGraph, UEdGraphPin* FromPin, const FVector2D Location, bool bSelectNewNode)
-{
-	UEdGraphNode_Comment* CommentTemplate = NewObject<UEdGraphNode_Comment>(ParentGraph, FName("Comment"));
-
-	CommentTemplate->bCommentBubbleVisible_InDetailsPanel = false;
-	CommentTemplate->bCommentBubbleVisible = false;
-	CommentTemplate->bCommentBubblePinned = false;
-
-//	const FScopedTransaction Transaction(NSLOCTEXT("UnrealEd", "DataflowEditorDuplicatedNode", "Dataflow Editor: New Comment Node (by using Duplicate)"));
-
-	ParentGraph->Modify();
-
-	CommentTemplate->SetFlags(RF_Transactional);
-
-	// set outer to be the graph so it doesn't go away
-	CommentTemplate->Rename(NULL, ParentGraph, REN_NonTransactional);
-	ParentGraph->AddNode(CommentTemplate, true, bSelectNewNode);
-
-	CommentTemplate->CreateNewGuid();
-	CommentTemplate->PostPlacedNewNode();
-	CommentTemplate->AllocateDefaultPins();
-	CommentTemplate->AutowireNewNode(FromPin);
-
-	CommentTemplate->NodePosX = Location.X;
-	CommentTemplate->NodePosY = Location.Y;
-	CommentTemplate->NodeWidth = CommentNodeToDuplicate->NodeWidth;
-	CommentTemplate->NodeHeight = CommentNodeToDuplicate->NodeHeight;
-	CommentTemplate->SnapToGrid(GetDefault<UEditorStyleSettings>()->GridSnapSize);
-	CommentTemplate->CommentColor = CommentNodeToDuplicate->CommentColor;
-	CommentTemplate->FontSize = CommentNodeToDuplicate->FontSize;
-
-	CommentTemplate->NodeComment = CommentNodeToDuplicate->NodeComment;;
-
-	ParentGraph->NotifyGraphChanged();
-
-	return CommentTemplate;
-}
-
-
-void FAssetSchemaAction_Dataflow_DuplicateCommentNode_DataflowEdNode::AddReferencedObjects(FReferenceCollector& Collector)
-{
-	Collector.AddReferencedObject(CommentNodeToDuplicate);
-}
-
-TSharedPtr<FAssetSchemaAction_Dataflow_PasteCommentNode_DataflowEdNode> FAssetSchemaAction_Dataflow_PasteCommentNode_DataflowEdNode::CreateAction(UEdGraph* ParentGraph, const TSharedPtr<SGraphEditor>& GraphEditor)
-{
-	return MakeShared<FAssetSchemaAction_Dataflow_PasteCommentNode_DataflowEdNode>(GraphEditor);
-}
-
-//
-//  Paste comment node
-//
-UEdGraphNode* FAssetSchemaAction_Dataflow_PasteCommentNode_DataflowEdNode::PerformAction(class UEdGraph* ParentGraph, UEdGraphPin* FromPin, const FVector2D Location, bool bSelectNewNode)
-{
-	UEdGraphNode_Comment* CommentTemplate = NewObject<UEdGraphNode_Comment>(ParentGraph, FName("Comment"));
-
-	CommentTemplate->bCommentBubbleVisible_InDetailsPanel = false;
-	CommentTemplate->bCommentBubbleVisible = false;
-	CommentTemplate->bCommentBubblePinned = false;
-
-	ParentGraph->Modify();
-
-	CommentTemplate->SetFlags(RF_Transactional);
-
-	// set outer to be the graph so it doesn't go away
-	CommentTemplate->Rename(NULL, ParentGraph, REN_NonTransactional);
-	ParentGraph->AddNode(CommentTemplate, true, bSelectNewNode);
-
-	CommentTemplate->CreateNewGuid();
-	CommentTemplate->PostPlacedNewNode();
-	CommentTemplate->AllocateDefaultPins();
-	CommentTemplate->AutowireNewNode(FromPin);
-
-	CommentTemplate->NodePosX = Location.X;
-	CommentTemplate->NodePosY = Location.Y;
-	CommentTemplate->NodeWidth = Size.X;
-	CommentTemplate->NodeHeight = Size.Y;
-	CommentTemplate->SnapToGrid(GetDefault<UEditorStyleSettings>()->GridSnapSize);
-	CommentTemplate->CommentColor = Color;
-	CommentTemplate->FontSize = FontSize;
-
-	CommentTemplate->NodeComment = NodeName.ToString();
-
-	ParentGraph->NotifyGraphChanged();
-
-	return CommentTemplate;
+	return nullptr;
 }
 
 #undef LOCTEXT_NAMESPACE
diff --git a/Engine/Plugins/Experimental/Dataflow/Source/DataflowEditor/Public/Dataflow/DataflowSCommentNode.h b/Engine/Plugins/Experimental/Dataflow/Source/DataflowEditor/Public/Dataflow/DataflowSCommentNode.h
--- a/Engine/Plugins/Experimental/Dataflow/Source/DataflowEditor/Public/Dataflow/DataflowSCommentNode.h
+++ b/Engine/Plugins/Experimental/Dataflow/Source/DataflowEditor/Public/Dataflow/DataflowSCommentNode.h
@@ -45,57 +45,3 @@ public:
 
 	TSharedPtr<SGraphEditor> GraphEditor;
 };
-
-//
-// Action to duplicate a comment node in the graph
-//
-USTRUCT()
-struct DATAFLOWEDITOR_API FAssetSchemaAction_Dataflow_DuplicateCommentNode_DataflowEdNode : public FEdGraphSchemaAction, public FGCObject
-{
-	GENERATED_USTRUCT_BODY();
-
-public:
-	FAssetSchemaAction_Dataflow_DuplicateCommentNode_DataflowEdNode(const TSharedPtr<SGraphEditor>& InGraphEditor) : FEdGraphSchemaAction()
-		, GraphEditor(InGraphEditor)
-	{}
-
-	FAssetSchemaAction_Dataflow_DuplicateCommentNode_DataflowEdNode() {}
-
-	static TSharedPtr<FAssetSchemaAction_Dataflow_DuplicateCommentNode_DataflowEdNode> CreateAction(UEdGraph* ParentGraph, const TSharedPtr<SGraphEditor>& GraphEditor);
-
-	virtual UEdGraphNode* PerformAction(class UEdGraph* ParentGraph, UEdGraphPin* FromPin, const FVector2D Location, bool bSelectNewNode = true) override;
-
-	/** FGCObject interface */
-	virtual void AddReferencedObjects(FReferenceCollector& Collector) override;
-	virtual FString GetReferencerName() const override { return TEXT("FDataflowConnectionDrawingPolicy"); }
-
-	TSharedPtr<SGraphEditor> GraphEditor;
-
-	TObjectPtr<UEdGraphNode_Comment> CommentNodeToDuplicate;
-};
-
-//
-// Action to paste a comment node in the graph
-//
-USTRUCT()
-struct DATAFLOWEDITOR_API FAssetSchemaAction_Dataflow_PasteCommentNode_DataflowEdNode : public FEdGraphSchemaAction
-{
-	GENERATED_USTRUCT_BODY();
-
-public:
-	FAssetSchemaAction_Dataflow_PasteCommentNode_DataflowEdNode(const TSharedPtr<SGraphEditor>& InGraphEditor) : FEdGraphSchemaAction()
-		, GraphEditor(InGraphEditor)
-	{}
-
-	FAssetSchemaAction_Dataflow_PasteCommentNode_DataflowEdNode() {}
-
-	static TSharedPtr<FAssetSchemaAction_Dataflow_PasteCommentNode_DataflowEdNode> CreateAction(UEdGraph* ParentGraph, const TSharedPtr<SGraphEditor>& GraphEditor);
-
-	virtual UEdGraphNode* PerformAction(class UEdGraph* ParentGraph, UEdGraphPin* FromPin, const FVector2D Location, bool bSelectNewNode = true) override;
-
-	TSharedPtr<SGraphEditor> GraphEditor;
-	FName NodeName;
-	FVector2D Size;
-	FLinearColor Color;
-	int32 FontSize;
-};
