Dataflow : Update commands to support working in sub-graphs

diff --git a/Engine/Plugins/Experimental/Dataflow/Source/DataflowEditor/Private/Dataflow/DataflowEditorCommands.cpp b/Engine/Plugins/Experimental/Dataflow/Source/DataflowEditor/Private/Dataflow/DataflowEditorCommands.cpp
--- a/Engine/Plugins/Experimental/Dataflow/Source/DataflowEditor/Private/Dataflow/DataflowEditorCommands.cpp
+++ b/Engine/Plugins/Experimental/Dataflow/Source/DataflowEditor/Private/Dataflow/DataflowEditorCommands.cpp
@@ -400,9 +400,9 @@ void FDataflowEditorCommands::UnfreezeNodes(UE::Dataflow::FContext& Context, con
 	}
 }
 
-void FDataflowEditorCommands::DeleteNodes(UDataflow* Graph, const FGraphPanelSelectionSet& SelectedNodes)
+void FDataflowEditorCommands::DeleteNodes(UEdGraph* EdGraph, const FGraphPanelSelectionSet& SelectedNodes)
 {
-	if (ensureMsgf(Graph != nullptr, TEXT("Warning : Failed to find valid graph.")))
+	if (ensureMsgf(EdGraph != nullptr, TEXT("Warning : Failed to find valid graph.")))
 	{
 		for (UObject* Node : SelectedNodes)
 		{
@@ -410,7 +410,7 @@ void FDataflowEditorCommands::DeleteNodes(UDataflow* Graph, const FGraphPanelSel
 			{
 				if (const TSharedPtr<UE::Dataflow::FGraph> DataflowGraph = EdNode->GetDataflowGraph())
 				{
-					Graph->RemoveNode(EdNode);
+					EdGraph->RemoveNode(EdNode);
 					if (TSharedPtr<FDataflowNode> DataflowNode = DataflowGraph->FindBaseNode(EdNode->GetDataflowNodeGuid()))
 					{
 						DataflowGraph->RemoveNode(DataflowNode);
@@ -419,7 +419,7 @@ void FDataflowEditorCommands::DeleteNodes(UDataflow* Graph, const FGraphPanelSel
 			}
 			else if (UEdGraphNode_Comment* CommentNode = Cast<UEdGraphNode_Comment>(Node))
 			{
-				Graph->RemoveNode(CommentNode);
+				EdGraph->RemoveNode(CommentNode);
 			}
 
 			// Auto-rename node so that its current name is made available until it is garbage collected
@@ -519,9 +519,9 @@ static void ShowNotificationMessage(const FText& Message, const SNotificationIte
 	}
 }
 
-void FDataflowEditorCommands::DuplicateNodes(UDataflow* Graph, const TSharedPtr<SDataflowGraphEditor>& DataflowGraphEditor, const FGraphPanelSelectionSet& SelectedNodes)
+void FDataflowEditorCommands::DuplicateNodes(UEdGraph* EdGraph, const TSharedPtr<SDataflowGraphEditor>& DataflowGraphEditor, const FGraphPanelSelectionSet& SelectedNodes)
 {
-	if (ensureMsgf(Graph != nullptr, TEXT("Warning : Failed to find valid graph.")))
+	if (ensureMsgf(EdGraph != nullptr, TEXT("Warning : Failed to find valid graph.")))
 	{
 		TArray<UEdGraphNode*> SelectedEdNodes;
 		for (UObject* SelectedNode : SelectedNodes)
@@ -533,7 +533,7 @@ void FDataflowEditorCommands::DuplicateNodes(UDataflow* Graph, const TSharedPtr<
 		}
 		TArray<UEdGraphNode*> DuplicatedEdNodes;
 		const FVector2D PasteLocation = DataflowGraphEditor->GetPasteLocation();
-		UE::Dataflow::FEditAssetUtils::DuplicateNodes(Graph, SelectedEdNodes, PasteLocation, DuplicatedEdNodes);
+		UE::Dataflow::FEditAssetUtils::DuplicateNodes(EdGraph, SelectedEdNodes, PasteLocation, DuplicatedEdNodes);
 
 		// Update the selection in the Editor
 		if (DuplicatedEdNodes.Num())
@@ -559,9 +559,9 @@ void FDataflowEditorCommands::DuplicateNodes(UDataflow* Graph, const TSharedPtr<
 	}
 }
 
-void FDataflowEditorCommands::CopyNodes(UDataflow* InGraph, const TSharedPtr<SDataflowGraphEditor>& DataflowGraphEditor, const FGraphPanelSelectionSet& InSelectedNodes)
+void FDataflowEditorCommands::CopyNodes(UEdGraph* EdGraph, const TSharedPtr<SDataflowGraphEditor>& DataflowGraphEditor, const FGraphPanelSelectionSet& InSelectedNodes)
 {
-	if (ensureMsgf(InGraph != nullptr, TEXT("Warning : Failed to find valid graph.")))
+	if (ensureMsgf(EdGraph != nullptr, TEXT("Warning : Failed to find valid graph.")))
 	{
 		if (InSelectedNodes.Num() > 0)
 		{
@@ -595,11 +595,11 @@ void FDataflowEditorCommands::CopyNodes(UDataflow* InGraph, const TSharedPtr<SDa
 	}
 }
 
-void FDataflowEditorCommands::PasteNodes(UDataflow* Graph, const TSharedPtr<SDataflowGraphEditor>& DataflowGraphEditor)
+void FDataflowEditorCommands::PasteNodes(UEdGraph* EdGraph, const TSharedPtr<SDataflowGraphEditor>& DataflowGraphEditor)
 {
 	TArray<UEdGraphNode*> PastedEdNodes;
 	const FVector2D PasteLocation = DataflowGraphEditor->GetPasteLocation();
-	UE::Dataflow::FEditAssetUtils::PasteNodesFromClipboard(Graph, PasteLocation, PastedEdNodes);
+	UE::Dataflow::FEditAssetUtils::PasteNodesFromClipboard(EdGraph, PasteLocation, PastedEdNodes);
 
 	// select the pasted nodes
 	if (PastedEdNodes.Num())
diff --git a/Engine/Plugins/Experimental/Dataflow/Source/DataflowEditor/Private/Dataflow/DataflowGraphEditor.cpp b/Engine/Plugins/Experimental/Dataflow/Source/DataflowEditor/Private/Dataflow/DataflowGraphEditor.cpp
--- a/Engine/Plugins/Experimental/Dataflow/Source/DataflowEditor/Private/Dataflow/DataflowGraphEditor.cpp
+++ b/Engine/Plugins/Experimental/Dataflow/Source/DataflowEditor/Private/Dataflow/DataflowGraphEditor.cpp
@@ -29,7 +29,7 @@ void SDataflowGraphEditor::Construct(const FArguments& InArgs, UObject* InAssetO
 {
 	check(InArgs._GraphToEdit);
 	AssetOwner = InAssetOwner; // nullptr is valid
-	DataflowAsset = Cast<UDataflow>(InArgs._GraphToEdit);
+	EdGraphWeakPtr = InArgs._GraphToEdit;
 	DetailsView = InArgs._DetailsView;
 	EvaluateGraphCallback = InArgs._EvaluateGraph;
 	OnDragDropEventCallback = InArgs._OnDragDropEvent;
@@ -215,6 +215,15 @@ TSharedPtr<UE::Dataflow::FContext> SDataflowGraphEditor::GetDataflowContext() co
 	return TSharedPtr<UE::Dataflow::FContext>();
 }
 
+UDataflow* SDataflowGraphEditor::GetDataflowAsset() const
+{
+	if (TStrongObjectPtr<UEdGraph> EdGraph = EdGraphWeakPtr.Pin())
+	{
+		return UDataflow::GetDataflowAssetFromEdGraph(EdGraph.Get());
+	}
+	return nullptr;
+}
+
 void SDataflowGraphEditor::OnRenderToggleChanged() const
 {
 	// need to refresh the UI by refreshing the selection
@@ -265,7 +274,7 @@ void SDataflowGraphEditor::EvaluateNode()
 
 void SDataflowGraphEditor::DeleteNode()
 {
-	if (DataflowAsset.Get())
+	if (TStrongObjectPtr<UEdGraph> EdGraph = EdGraphWeakPtr.Pin())
 	{
 		if (DetailsView)
 		{
@@ -275,18 +284,19 @@ void SDataflowGraphEditor::DeleteNode()
 		const FGraphPanelSelectionSet& SelectedNodes = GetSelectedNodes();
 		if (SelectedNodes.Num() > 0)
 		{
-			if (UDataflow* const Graph = DataflowAsset.Get())
+			if (UDataflow* Dataflow = UDataflow::GetDataflowAssetFromEdGraph(EdGraph.Get()))
 			{
 				const FScopedTransaction Transaction(LOCTEXT("DeleteSelectedNodes", "Delete selected nodes"));
 
-				Graph->Modify();
+				Dataflow->Modify();
+				EdGraph->Modify();
 
 				for (FGraphPanelSelectionSet::TConstIterator It(SelectedNodes); It; ++It)
 				{
 					(*It)->Modify();
 				}
 
-				FDataflowEditorCommands::DeleteNodes(Graph, SelectedNodes);
+				FDataflowEditorCommands::DeleteNodes(EdGraph.Get(), SelectedNodes);
 
 				OnNodeDeletedMulticast.Broadcast(SelectedNodes);
 			}
@@ -296,7 +306,7 @@ void SDataflowGraphEditor::DeleteNode()
 
 void SDataflowGraphEditor::RenameNode()
 {
-	if (UDataflow* Graph = DataflowAsset.Get())
+	if (TStrongObjectPtr<UEdGraph> EdGraph = EdGraphWeakPtr.Pin())
 	{
 		const TSharedPtr<SDataflowGraphEditor>& DataflowGraphEditor = SharedThis(this);
 		const FGraphPanelSelectionSet& SelectedNodes = GetSelectedNodes();
@@ -338,12 +348,18 @@ bool SDataflowGraphEditor::CanRenameNode() const
 
 void SDataflowGraphEditor::AddNewVariable() const
 {
-	UE::Dataflow::FEditAssetUtils::AddNewVariable(DataflowAsset.Get());
+	if (UDataflow* DataflowAsset = GetDataflowAsset())
+	{
+		UE::Dataflow::FEditAssetUtils::AddNewVariable(DataflowAsset);
+	}
 }
 
 void SDataflowGraphEditor::AddNewSubGraph() const
 {
-	UE::Dataflow::FEditAssetUtils::AddNewSubGraph(DataflowAsset.Get());
+	if (UDataflow* DataflowAsset = GetDataflowAsset())
+	{
+		UE::Dataflow::FEditAssetUtils::AddNewSubGraph(DataflowAsset);
+	}
 }
 
 void SDataflowGraphEditor::OnSelectedNodesChanged(const TSet<UObject*>& NewSelection)
@@ -354,14 +370,17 @@ void SDataflowGraphEditor::OnSelectedNodesChanged(const TSet<UObject*>& NewSelec
 
 	OnSelectionChangedMulticast.Broadcast(NewSelection);
 
-	if (DataflowAsset.Get() && DetailsView)
+	if (DetailsView)
 	{
-		auto AsObjectPointers = [](const TSet<UObject*>& Set) {
-			TSet<TObjectPtr<UObject> > Objs; for (UObject* Elem : Set) Objs.Add(Elem);
-			return Objs;
-		};
+		if (UDataflow* DataflowAsset = GetDataflowAsset())
+		{
+			auto AsObjectPointers = [](const TSet<UObject*>& Set) {
+				TSet<TObjectPtr<UObject> > Objs; for (UObject* Elem : Set) Objs.Add(Elem);
+				return Objs;
+				};
 
-		FDataflowEditorCommands::OnSelectedNodesChanged(DetailsView, AssetOwner.Get(), DataflowAsset.Get(), AsObjectPointers(NewSelection) );
+			FDataflowEditorCommands::OnSelectedNodesChanged(DetailsView, AssetOwner.Get(), DataflowAsset, AsObjectPointers(NewSelection));
+		}
 	}
 
 	// Clear the current selected editor
@@ -502,30 +521,37 @@ void SDataflowGraphEditor::OnDragLeave(const FDragDropEvent& DragDropEvent)
 
 void SDataflowGraphEditor::CreateComment()
 {
-	UDataflow* Graph = DataflowAsset.Get();
-	const TSharedPtr<SGraphEditor>& InGraphEditor = SharedThis(GetGraphEditor());
+	if (TStrongObjectPtr<UEdGraph> EdGraph = EdGraphWeakPtr.Pin())
+	{
+		const TSharedPtr<SGraphEditor>& InGraphEditor = SharedThis(GetGraphEditor());
 
-	TSharedPtr<FAssetSchemaAction_Dataflow_CreateCommentNode_DataflowEdNode> CommentAction = FAssetSchemaAction_Dataflow_CreateCommentNode_DataflowEdNode::CreateAction(Graph, InGraphEditor);
-	CommentAction->PerformAction(Graph, nullptr, GetGraphEditor()->GetPasteLocation(), false);
+		TSharedPtr<FAssetSchemaAction_Dataflow_CreateCommentNode_DataflowEdNode> CommentAction = FAssetSchemaAction_Dataflow_CreateCommentNode_DataflowEdNode::CreateAction(EdGraph.Get(), InGraphEditor);
+		CommentAction->PerformAction(EdGraph.Get(), nullptr, GetGraphEditor()->GetPasteLocation(), false);
+	}
 }
 
 void SDataflowGraphEditor::CreateVertexSelectionNode(const FString & InArray)
 {
-	UDataflow* Graph = DataflowAsset.Get();
-	const TSharedPtr<SGraphEditor>& InGraphEditor = SharedThis(GetGraphEditor());
-
-	if (TSharedPtr<FAssetSchemaAction_Dataflow_CreateNode_DataflowEdNode> NodeAction = FAssetSchemaAction_Dataflow_CreateNode_DataflowEdNode::CreateAction(Graph, "FSelectionSetDataflowNode"))
+	if (TStrongObjectPtr<UEdGraph> EdGraph = EdGraphWeakPtr.Pin())
 	{
-		if (UEdGraphNode* NewEdNode = NodeAction->PerformAction(Graph, nullptr, GetGraphEditor()->GetPasteLocation(), false))
+		const TSharedPtr<SGraphEditor>& InGraphEditor = SharedThis(GetGraphEditor());
+
+		if (TSharedPtr<FAssetSchemaAction_Dataflow_CreateNode_DataflowEdNode> NodeAction = FAssetSchemaAction_Dataflow_CreateNode_DataflowEdNode::CreateAction(EdGraph.Get(), "FSelectionSetDataflowNode"))
 		{
-			FDataflowAssetEdit Edit = Graph->EditDataflow();
-			if (UE::Dataflow::FGraph* DataflowGraph = Edit.GetGraph())
+			if (UEdGraphNode* NewEdNode = NodeAction->PerformAction(EdGraph.Get(), nullptr, GetGraphEditor()->GetPasteLocation(), false))
 			{
-				if (TSharedPtr<FDataflowNode> Node = DataflowGraph->FindBaseNode(((UDataflowEdNode*)NewEdNode)->DataflowNodeGuid))
+				if (UDataflow* DataflowAsset = UDataflow::GetDataflowAssetFromEdGraph(EdGraph.Get()))
 				{
-					if (FSelectionSetDataflowNode * SelectionNode = Node->AsType<FSelectionSetDataflowNode>())
+					FDataflowAssetEdit Edit = DataflowAsset->EditDataflow();
+					if (UE::Dataflow::FGraph* DataflowGraph = Edit.GetGraph())
 					{
-						SelectionNode->Indices = InArray;
+						if (TSharedPtr<FDataflowNode> Node = DataflowGraph->FindBaseNode(((UDataflowEdNode*)NewEdNode)->DataflowNodeGuid))
+						{
+							if (FSelectionSetDataflowNode* SelectionNode = Node->AsType<FSelectionSetDataflowNode>())
+							{
+								SelectionNode->Indices = InArray;
+							}
+						}
 					}
 				}
 			}
@@ -580,34 +606,43 @@ void SDataflowGraphEditor::DistributeVertically()
 
 void SDataflowGraphEditor::ToggleEnabledState()
 {
-	FDataflowEditorCommands::ToggleEnabledState(DataflowAsset.Get());
+	if (UDataflow* DataflowAsset = GetDataflowAsset())
+	{
+		FDataflowEditorCommands::ToggleEnabledState(DataflowAsset);
+	}
 }
 
 void SDataflowGraphEditor::OnAddOptionPin()
 {
-	UDataflow* const Graph = DataflowAsset.Get();
-	FDataflowAssetEdit Edit = Graph->EditDataflow();
-	if (UE::Dataflow::FGraph* const DataflowGraph = Edit.GetGraph())
+	if (TStrongObjectPtr<UEdGraph> EdGraph = EdGraphWeakPtr.Pin())
 	{
-		const FGraphPanelSelectionSet& SelectedNodes = GetSelectedNodes();
-
-		// Iterate over all nodes, and add the pin
-		for (FGraphPanelSelectionSet::TConstIterator It(SelectedNodes); It; ++It)
+		if (UDataflow* DataflowAsset = UDataflow::GetDataflowAssetFromEdGraph(EdGraph.Get()))
 		{
-			UDataflowEdNode* const EdNode = CastChecked<UDataflowEdNode>(*It);
-
-			if (const TSharedPtr<FDataflowNode> Node = DataflowGraph->FindBaseNode(EdNode->DataflowNodeGuid))
+			FDataflowAssetEdit Edit = DataflowAsset->EditDataflow();
+			if (UE::Dataflow::FGraph* const DataflowGraph = Edit.GetGraph())
 			{
-				if (Node->CanAddPin())
+				const FGraphPanelSelectionSet& SelectedNodes = GetSelectedNodes();
+
+				// Iterate over all nodes, and add the pin
+				for (FGraphPanelSelectionSet::TConstIterator It(SelectedNodes); It; ++It)
 				{
-					const FScopedTransaction Transaction(LOCTEXT("AddOptionPin", "Add Option Pin"));
-					Graph->Modify();
-					EdNode->Modify();
+					UDataflowEdNode* const EdNode = CastChecked<UDataflowEdNode>(*It);
 
-					EdNode->AddOptionPin();
+					if (const TSharedPtr<FDataflowNode> Node = DataflowGraph->FindBaseNode(EdNode->DataflowNodeGuid))
+					{
+						if (Node->CanAddPin())
+						{
+							const FScopedTransaction Transaction(LOCTEXT("AddOptionPin", "Add Option Pin"));
+							DataflowAsset->Modify();
+							EdGraph->Modify();
+							EdNode->Modify();
 
-					const UDataflowSchema* const Schema = CastChecked<UDataflowSchema>(Graph->GetSchema());
-					Schema->ReconstructNode(*EdNode);
+							EdNode->AddOptionPin();
+
+							const UDataflowSchema* const Schema = CastChecked<UDataflowSchema>(EdGraph->GetSchema());
+							Schema->ReconstructNode(*EdNode);
+						}
+					}
 				}
 			}
 		}
@@ -618,28 +653,33 @@ bool SDataflowGraphEditor::CanAddOptionPin() const
 {
 	bool bCanAddOptionPin = false;
 
-	const UDataflow* const Graph = DataflowAsset.Get();
-	if (const UE::Dataflow::FGraph* const DataflowGraph = Graph->GetDataflow().Get())
+	if (TStrongObjectPtr<UEdGraph> EdGraph = EdGraphWeakPtr.Pin())
 	{
-		const FGraphPanelSelectionSet& SelectedNodes = GetSelectedNodes();
-
-		// Iterate over all nodes, and add the pin
-		for (FGraphPanelSelectionSet::TConstIterator It(SelectedNodes); It; ++It)
+		if (UDataflow* DataflowAsset = UDataflow::GetDataflowAssetFromEdGraph(EdGraph.Get()))
 		{
-			if (const UDataflowEdNode* const EdNode = Cast<UDataflowEdNode>(*It))
+			if (const UE::Dataflow::FGraph* const DataflowGraph = DataflowAsset->GetDataflow().Get())
 			{
-				if (const TSharedPtr<const FDataflowNode> Node = DataflowGraph->FindBaseNode(EdNode->DataflowNodeGuid))
-				{
-					bCanAddOptionPin = Node->CanAddPin();
-				}
-				else
-				{
-					bCanAddOptionPin = false;
-				}
+				const FGraphPanelSelectionSet& SelectedNodes = GetSelectedNodes();
 
-				if (!bCanAddOptionPin)
+				// Iterate over all nodes, and add the pin
+				for (FGraphPanelSelectionSet::TConstIterator It(SelectedNodes); It; ++It)
 				{
-					break;  // One bad node is good enough to return false
+					if (const UDataflowEdNode* const EdNode = Cast<UDataflowEdNode>(*It))
+					{
+						if (const TSharedPtr<const FDataflowNode> Node = DataflowGraph->FindBaseNode(EdNode->DataflowNodeGuid))
+						{
+							bCanAddOptionPin = Node->CanAddPin();
+						}
+						else
+						{
+							bCanAddOptionPin = false;
+						}
+
+						if (!bCanAddOptionPin)
+						{
+							break;  // One bad node is good enough to return false
+						}
+					}
 				}
 			}
 		}
@@ -650,29 +690,35 @@ bool SDataflowGraphEditor::CanAddOptionPin() const
 
 void SDataflowGraphEditor::OnRemoveOptionPin()
 {
-	UDataflow* const Graph = DataflowAsset.Get();
-	FDataflowAssetEdit Edit = Graph->EditDataflow();
-	if (UE::Dataflow::FGraph* const DataflowGraph = Edit.GetGraph())
+	if (TStrongObjectPtr<UEdGraph> EdGraph = EdGraphWeakPtr.Pin())
 	{
-		const FGraphPanelSelectionSet& SelectedNodes = GetSelectedNodes();
-
-		// Iterate over all nodes, and remove a pin
-		for (FGraphPanelSelectionSet::TConstIterator It(SelectedNodes); It; ++It)
+		if (UDataflow* DataflowAsset = UDataflow::GetDataflowAssetFromEdGraph(EdGraph.Get()))
 		{
-			UDataflowEdNode* const EdNode = CastChecked<UDataflowEdNode>(*It);
-
-			if (const TSharedPtr<FDataflowNode> Node = DataflowGraph->FindBaseNode(EdNode->DataflowNodeGuid))
+			FDataflowAssetEdit Edit = DataflowAsset->EditDataflow();
+			if (UE::Dataflow::FGraph* const DataflowGraph = Edit.GetGraph())
 			{
-				if (Node->CanRemovePin())
+				const FGraphPanelSelectionSet& SelectedNodes = GetSelectedNodes();
+
+				// Iterate over all nodes, and remove a pin
+				for (FGraphPanelSelectionSet::TConstIterator It(SelectedNodes); It; ++It)
 				{
-					const FScopedTransaction Transaction(LOCTEXT("RemoveOptionPin", "Remove Option Pin"));
-					Graph->Modify();
-					EdNode->Modify();
+					UDataflowEdNode* const EdNode = CastChecked<UDataflowEdNode>(*It);
 
-					EdNode->RemoveOptionPin();
+					if (const TSharedPtr<FDataflowNode> Node = DataflowGraph->FindBaseNode(EdNode->DataflowNodeGuid))
+					{
+						if (Node->CanRemovePin())
+						{
+							const FScopedTransaction Transaction(LOCTEXT("RemoveOptionPin", "Remove Option Pin"));
+							DataflowAsset->Modify();
+							EdGraph->Modify();
+							EdNode->Modify();
 
-					const UDataflowSchema* const Schema = CastChecked<UDataflowSchema>(Graph->GetSchema());
-					Schema->ReconstructNode(*EdNode);
+							EdNode->RemoveOptionPin();
+
+							const UDataflowSchema* const Schema = CastChecked<UDataflowSchema>(EdGraph->GetSchema());
+							Schema->ReconstructNode(*EdNode);
+						}
+					}
 				}
 			}
 		}
@@ -683,28 +729,33 @@ bool SDataflowGraphEditor::CanRemoveOptionPin() const
 {
 	bool bCanRemoveOptionPin = false;
 
-	const UDataflow* const Graph = DataflowAsset.Get();
-	if (const UE::Dataflow::FGraph* const DataflowGraph = Graph->GetDataflow().Get())
+	if (TStrongObjectPtr<UEdGraph> EdGraph = EdGraphWeakPtr.Pin())
 	{
-		const FGraphPanelSelectionSet& SelectedNodes = GetSelectedNodes();
-
-		// Iterate over all nodes, and add the pin
-		for (FGraphPanelSelectionSet::TConstIterator It(SelectedNodes); It; ++It)
+		if (UDataflow* DataflowAsset = UDataflow::GetDataflowAssetFromEdGraph(EdGraph.Get()))
 		{
-			if (const UDataflowEdNode* const EdNode = Cast<UDataflowEdNode>(*It))
+			if (const UE::Dataflow::FGraph* const DataflowGraph = DataflowAsset->GetDataflow().Get())
 			{
-				if (const TSharedPtr<const FDataflowNode> Node = DataflowGraph->FindBaseNode(EdNode->DataflowNodeGuid))
-				{
-					bCanRemoveOptionPin = Node->CanRemovePin();
-				}
-				else
-				{
-					bCanRemoveOptionPin = false;
-				}
+				const FGraphPanelSelectionSet& SelectedNodes = GetSelectedNodes();
 
-				if (!bCanRemoveOptionPin)
+				// Iterate over all nodes, and add the pin
+				for (FGraphPanelSelectionSet::TConstIterator It(SelectedNodes); It; ++It)
 				{
-					break;  // One bad node is good enough to return false
+					if (const UDataflowEdNode* const EdNode = Cast<UDataflowEdNode>(*It))
+					{
+						if (const TSharedPtr<const FDataflowNode> Node = DataflowGraph->FindBaseNode(EdNode->DataflowNodeGuid))
+						{
+							bCanRemoveOptionPin = Node->CanRemovePin();
+						}
+						else
+						{
+							bCanRemoveOptionPin = false;
+						}
+
+						if (!bCanRemoveOptionPin)
+						{
+							break;  // One bad node is good enough to return false
+						}
+					}
 				}
 			}
 		}
@@ -715,14 +766,14 @@ bool SDataflowGraphEditor::CanRemoveOptionPin() const
 
 void SDataflowGraphEditor::DuplicateSelectedNodes()
 {
-	if (UDataflow* Graph = DataflowAsset.Get())
+	if (TStrongObjectPtr<UEdGraph> EdGraph = EdGraphWeakPtr.Pin())
 	{
 		const TSharedPtr<SDataflowGraphEditor>& DataflowGraphEditor = SharedThis(this);
 		const FGraphPanelSelectionSet& SelectedNodes = GetSelectedNodes();
 
 		if (SelectedNodes.Num() > 0)
 		{
-			FDataflowEditorCommands::DuplicateNodes(Graph, DataflowGraphEditor, SelectedNodes);
+			FDataflowEditorCommands::DuplicateNodes(EdGraph.Get(), DataflowGraphEditor, SelectedNodes);
 		}
 	}
 }
@@ -744,47 +795,47 @@ bool SDataflowGraphEditor::GetPinVisibility(SGraphEditor::EPinVisibility PinVisi
 
 void SDataflowGraphEditor::AddReferencedObjects(FReferenceCollector& Collector)
 {
-	Collector.AddReferencedObject(DataflowAsset);
+	Collector.AddReferencedObject(EdGraphWeakPtr);
 	Collector.AddReferencedObject(AssetOwner);
 }
 
 void SDataflowGraphEditor::CopySelectedNodes()
 {
-	if (UDataflow* Graph = DataflowAsset.Get())
+	if (TStrongObjectPtr<UEdGraph> EdGraph = EdGraphWeakPtr.Pin())
 	{
 		const TSharedPtr<SDataflowGraphEditor>& DataflowGraphEditor = SharedThis(this);
 		const FGraphPanelSelectionSet& SelectedNodes = GetSelectedNodes();
 
 		if (SelectedNodes.Num() > 0)
 		{
-			FDataflowEditorCommands::CopyNodes(Graph, DataflowGraphEditor, SelectedNodes);
+			FDataflowEditorCommands::CopyNodes(EdGraph.Get(), DataflowGraphEditor, SelectedNodes);
 		}
 	}
 }
 
 void SDataflowGraphEditor::CutSelectedNodes()
 {
-	if (UDataflow* Graph = DataflowAsset.Get())
+	if (TStrongObjectPtr<UEdGraph> EdGraph = EdGraphWeakPtr.Pin())
 	{
 		const TSharedPtr<SDataflowGraphEditor>& DataflowGraphEditor = SharedThis(this);
 		const FGraphPanelSelectionSet& SelectedNodes = GetSelectedNodes();
 
 		if (SelectedNodes.Num() > 0)
 		{
-			FDataflowEditorCommands::CopyNodes(Graph, DataflowGraphEditor, SelectedNodes);
+			FDataflowEditorCommands::CopyNodes(EdGraph.Get(), DataflowGraphEditor, SelectedNodes);
 
-			FDataflowEditorCommands::DeleteNodes(DataflowAsset.Get(), SelectedNodes);
+			FDataflowEditorCommands::DeleteNodes(EdGraph.Get(), SelectedNodes);
 		}
 	}
 }
 
 void SDataflowGraphEditor::PasteSelectedNodes()
 {
-	if (UDataflow* Graph = DataflowAsset.Get())
+	if (TStrongObjectPtr<UEdGraph> EdGraph = EdGraphWeakPtr.Pin())
 	{
 		const TSharedPtr<SDataflowGraphEditor>& DataflowGraphEditor = SharedThis(this);
 
-		FDataflowEditorCommands::PasteNodes(Graph, DataflowGraphEditor);
+		FDataflowEditorCommands::PasteNodes(EdGraph.Get(), DataflowGraphEditor);
 	}
 }
 
diff --git a/Engine/Plugins/Experimental/Dataflow/Source/DataflowEditor/Public/Dataflow/DataflowEditorCommands.h b/Engine/Plugins/Experimental/Dataflow/Source/DataflowEditor/Public/Dataflow/DataflowEditorCommands.h
--- a/Engine/Plugins/Experimental/Dataflow/Source/DataflowEditor/Public/Dataflow/DataflowEditorCommands.h
+++ b/Engine/Plugins/Experimental/Dataflow/Source/DataflowEditor/Public/Dataflow/DataflowEditorCommands.h
@@ -14,6 +14,7 @@ class IStructureDetailsView;
 class UDataflow;
 class UDataflowEdNode;
 struct FDataflowNode;
+class UEdGraph;
 class UEdGraphNode;
 class SDataflowGraphEditor;
 class UDataflowBaseContent;
@@ -130,7 +131,7 @@ public:
 	/*
 	*  DeleteNodes
 	*/
-	static void DeleteNodes(UDataflow* Graph, const FGraphPanelSelectionSet& SelectedNodes);
+	static void DeleteNodes(UEdGraph* EdGraph, const FGraphPanelSelectionSet& SelectedNodes);
 
 	/*
 	*  FreezeNodes
@@ -176,17 +177,17 @@ public:
 	/*
 	*  DuplicateNodes
 	*/
-	static void DuplicateNodes(UDataflow* Graph, const TSharedPtr<SDataflowGraphEditor>& DataflowGraphEditor, const FGraphPanelSelectionSet& SelectedNodes);
+	static void DuplicateNodes(UEdGraph* EdGraph, const TSharedPtr<SDataflowGraphEditor>& DataflowGraphEditor, const FGraphPanelSelectionSet& SelectedNodes);
 
 	/*
 	*  CopyNodes
 	*/
-	static void CopyNodes(UDataflow* Graph, const TSharedPtr<SDataflowGraphEditor>& DataflowGraphEditor, const FGraphPanelSelectionSet& SelectedNodes);
+	static void CopyNodes(UEdGraph* EdGraph, const TSharedPtr<SDataflowGraphEditor>& DataflowGraphEditor, const FGraphPanelSelectionSet& SelectedNodes);
 
 	/*
 	*  PasteSelectedNodes
 	*/
-	static void PasteNodes(UDataflow* Graph, const TSharedPtr<SDataflowGraphEditor>& DataflowGraphEditor);
+	static void PasteNodes(UEdGraph* EdGraph, const TSharedPtr<SDataflowGraphEditor>& DataflowGraphEditor);
 
 	/*
 	*  RenameNode
diff --git a/Engine/Plugins/Experimental/Dataflow/Source/DataflowEditor/Public/Dataflow/DataflowGraphEditor.h b/Engine/Plugins/Experimental/Dataflow/Source/DataflowEditor/Public/Dataflow/DataflowGraphEditor.h
--- a/Engine/Plugins/Experimental/Dataflow/Source/DataflowEditor/Public/Dataflow/DataflowGraphEditor.h
+++ b/Engine/Plugins/Experimental/Dataflow/Source/DataflowEditor/Public/Dataflow/DataflowGraphEditor.h
@@ -181,6 +181,9 @@ private:
 	/** Text for overlay message in graph editor panel */
 	FText GetGraphEditorOverlayText() const;
 
+	/** get the dataflow asset from the edGraph being edited */
+	UDataflow* GetDataflowAsset() const;
+
 	FDataflowEditorCommands::FOnDragDropEventCallback OnDragDropEventCallback;
 	FDataflowEditorCommands::FGraphEvaluationCallback EvaluateGraphCallback;
 
@@ -188,7 +191,7 @@ private:
 	TWeakObjectPtr<UObject> AssetOwner;
 
 	/** The dataflow asset associated with this graph */
-	TWeakObjectPtr<UDataflow> DataflowAsset;
+	TWeakObjectPtr<UEdGraph> EdGraphWeakPtr;
 
 	/** Command list associated with this graph editor */
 	TSharedPtr<FUICommandList> GraphEditorCommands;
