Include subgraph in name generation for nodes

diff --git a/Engine/Plugins/Experimental/Dataflow/Source/DataflowEditor/Private/Dataflow/DataflowAssetEditUtils.cpp b/Engine/Plugins/Experimental/Dataflow/Source/DataflowEditor/Private/Dataflow/DataflowAssetEditUtils.cpp
--- a/Engine/Plugins/Experimental/Dataflow/Source/DataflowEditor/Private/Dataflow/DataflowAssetEditUtils.cpp
+++ b/Engine/Plugins/Experimental/Dataflow/Source/DataflowEditor/Private/Dataflow/DataflowAssetEditUtils.cpp
@@ -10,6 +10,7 @@
 #include "Dataflow/DataflowEdNode.h"
 #include "Dataflow/DataflowSNode.h"
 #include "Dataflow/DataflowSubGraph.h"
+#include "Editor.h"
 #include "PropertyBagDetails.h"
 #include "ScopedTransaction.h"
 #include "Settings/EditorStyleSettings.h"
@@ -69,30 +70,28 @@ namespace UE::Dataflow
 			}
 		}
 
-		/** Generate a Dataflow node unique name from a BaseName */
-		static FName GenerateUniqueNodeName(UDataflow* Dataflow, FString NodeBaseName)
+		/** Generate a Dataflow child object unique name from a BaseName ( Node or Subgraph for example ) */
+		static FName GenerateUniqueObjectName(UDataflow* Dataflow, const FName InBaseName)
 		{
 			FString Left, Right;
 			int32 NameIndex = 1;
 
 			// Check if NodeBaseName already ends with "_dd"
-			if (NodeBaseName.Split(TEXT("_"), &Left, &Right, ESearchCase::IgnoreCase, ESearchDir::FromEnd))
+			FName BaseName(InBaseName);
+			if (BaseName.ToString().Split(TEXT("_"), &Left, &Right, ESearchCase::IgnoreCase, ESearchDir::FromEnd))
 			{
 				if (Right.IsNumeric())
 				{
 					NameIndex = FCString::Atoi(*Right);
-					NodeBaseName = Left;
+					BaseName = FName(Left);
 				}
 			}
-
-			FName NodeUniqueName{ NodeBaseName };
-			while (Dataflow->GetDataflow()->FindBaseNode(FName(NodeUniqueName)) != nullptr)
+			
+			if (::IsUniqueObjectName(BaseName, Dataflow))
 			{
-				NodeUniqueName = FName(NodeBaseName + FString::Printf(TEXT("_%d"), NameIndex));
-				NameIndex++;
+				return BaseName;
 			}
-
-			return NodeUniqueName;
+			return ::MakeUniqueObjectName(Dataflow, UDataflowEdNode::StaticClass(), BaseName);
 		}
 
 		static TSharedPtr<FDataflowNode> AddDataflowNode(UDataflow* Dataflow, FName NodeName, FName NodeTypeName)
@@ -103,7 +102,7 @@ namespace UE::Dataflow
 				{
 					.Guid = FGuid::NewGuid(),
 					.Type = NodeTypeName,
-					.Name = GenerateUniqueNodeName(Dataflow, NodeName.ToString()),
+					.Name = GenerateUniqueObjectName(Dataflow, NodeName),
 					.OwningObject = Dataflow,
 				};
 
@@ -234,24 +233,6 @@ namespace UE::Dataflow
 			return UniqueName;
 		}
 
-		/** Generate a Dataflow asset subgraph unique name from a BaseName */
-		static FName GenerateUniqueSubGraphName(const UDataflow& DataflowAsset, FName BaseName)
-		{
-			int32 Counter = 1;
-			FName UniqueName{ BaseName };
-			const FString BasenameStr{ BaseName.ToString() };
-
-			while (true)
-			{
-				if (DataflowAsset.FindSubGraphByName(UniqueName) == nullptr)
-				{
-					break; // found an available name exit 
-				}
-				UniqueName = FName(FString::Format(TEXT("{0}_{1}"), { BasenameStr, FString::FromInt(Counter++) }));
-			}
-			return UniqueName;
-		}
-
 		static UEdGraphPin* FindPin(const UDataflowEdNode* Node, const EEdGraphPinDirection Direction, const FName Name)
 		{
 			for (UEdGraphPin* Pin : Node->GetAllPins())
@@ -913,9 +894,9 @@ namespace UE::Dataflow
 		auto AddNewVariableInternal =
 			[&BaseName, &UniqueSubGraphName](UDataflow& DataflowAsset) -> EChangeResult
 			{
-				UniqueSubGraphName = GenerateUniqueSubGraphName(DataflowAsset, BaseName);
+				UniqueSubGraphName = GenerateUniqueObjectName(&DataflowAsset, BaseName);
 				UDataflowSubGraph* NewSubGraph = NewObject<UDataflowSubGraph>(&DataflowAsset, UniqueSubGraphName);
-				ensure(NewSubGraph->GetFName() == UniqueSubGraphName);
+				ensure(NewSubGraph->GetFName().IsEqual(UniqueSubGraphName));
 				NewSubGraph->Schema = UDataflowSchema::StaticClass();
 				NewSubGraph->SetFlags(RF_Transactional);
 
