FEditorViewportClient: make GetMinimumOrthoZoom virtual.
FDataflowConstructionViewportClient: set a custom GetMinimumOrthoZoom, ignoring ULevelEditorViewportSettings::MinimumOrthographicZoom

diff --git a/Engine/Plugins/Experimental/Dataflow/Source/DataflowEditor/Private/Dataflow/DataflowConstructionViewportClient.cpp b/Engine/Plugins/Experimental/Dataflow/Source/DataflowEditor/Private/Dataflow/DataflowConstructionViewportClient.cpp
--- a/Engine/Plugins/Experimental/Dataflow/Source/DataflowEditor/Private/Dataflow/DataflowConstructionViewportClient.cpp
+++ b/Engine/Plugins/Experimental/Dataflow/Source/DataflowEditor/Private/Dataflow/DataflowConstructionViewportClient.cpp
@@ -343,6 +343,12 @@ void FDataflowConstructionViewportClient::DrawCanvas(FViewport& InViewport, FSce
 	}
 }
 
+float FDataflowConstructionViewportClient::GetMinimumOrthoZoom() const
+{
+	// Ignore ULevelEditorViewportSettings::MinimumOrthographicZoom in this viewport client
+	return 1.0f;
+}
+
 void FDataflowConstructionViewportClient::AddReferencedObjects(FReferenceCollector& Collector)
 {
 	FDataflowEditorViewportClientBase::AddReferencedObjects(Collector);
diff --git a/Engine/Plugins/Experimental/Dataflow/Source/DataflowEditor/Public/Dataflow/DataflowConstructionViewportClient.h b/Engine/Plugins/Experimental/Dataflow/Source/DataflowEditor/Public/Dataflow/DataflowConstructionViewportClient.h
--- a/Engine/Plugins/Experimental/Dataflow/Source/DataflowEditor/Public/Dataflow/DataflowConstructionViewportClient.h
+++ b/Engine/Plugins/Experimental/Dataflow/Source/DataflowEditor/Public/Dataflow/DataflowConstructionViewportClient.h
@@ -72,6 +72,7 @@ private:
 	virtual void ProcessClick(FSceneView& View, HHitProxy* HitProxy, FKey Key, EInputEvent Event, uint32 HitX, uint32 HitY) override;
 	virtual void Draw(const FSceneView* View, FPrimitiveDrawInterface* PDI) override;
 	virtual void DrawCanvas(FViewport& InViewport, FSceneView& View, FCanvas& Canvas) override;
+	virtual float GetMinimumOrthoZoom() const override;
 
 	TUniquePtr<IClickDragBehaviorTarget> OrthoScrollBehaviorTarget;
 	TArray<TObjectPtr<UInputBehavior>> BehaviorsFor2DMode;
diff --git a/Engine/Source/Editor/UnrealEd/Public/EditorViewportClient.h b/Engine/Source/Editor/UnrealEd/Public/EditorViewportClient.h
--- a/Engine/Source/Editor/UnrealEd/Public/EditorViewportClient.h
+++ b/Engine/Source/Editor/UnrealEd/Public/EditorViewportClient.h
@@ -1765,7 +1765,7 @@ private:
 	virtual void BeginCameraMovement(bool bHasMovement) {}
 	virtual void EndCameraMovement() {}
 
-	UNREALED_API float GetMinimumOrthoZoom() const;
+	UNREALED_API virtual float GetMinimumOrthoZoom() const;
 
 public:
 	static UNREALED_API const uint32 MaxCameraSpeeds;
