Dataflow subgraphs : create a subgraph call  node when drag and droppping a subgraph from the members widget

diff --git a/Engine/Plugins/Experimental/Dataflow/Source/DataflowEditor/Private/Dataflow/DataflowGraphSchemaAction.cpp b/Engine/Plugins/Experimental/Dataflow/Source/DataflowEditor/Private/Dataflow/DataflowGraphSchemaAction.cpp
--- a/Engine/Plugins/Experimental/Dataflow/Source/DataflowEditor/Private/Dataflow/DataflowGraphSchemaAction.cpp
+++ b/Engine/Plugins/Experimental/Dataflow/Source/DataflowEditor/Private/Dataflow/DataflowGraphSchemaAction.cpp
@@ -8,6 +8,7 @@
 #include "Dataflow/DataflowSNode.h"
 #include "Dataflow/DataflowAssetEditUtils.h"
 #include "Dataflow/DataflowVariableNodes.h"
+#include "Dataflow/DataflowSubGraphNodes.h"
 #include "EdGraph/EdGraphSchema.h"
 #include "PropertyBagDetails.h"
 #include "StructUtils/PropertyBag.h"
@@ -209,7 +210,7 @@ FGraphSchemaActionDragDropAction_DataflowVariable::FGraphSchemaActionDragDropAct
 	: FGraphSchemaActionDragDropAction()
 {}
 
-TSharedRef<FGraphSchemaActionDragDropAction_DataflowVariable> FGraphSchemaActionDragDropAction_DataflowVariable::New(TSharedPtr<FEdGraphSchemaAction_DataflowVariable> InAction)
+TSharedRef<FGraphSchemaActionDragDropAction_DataflowVariable> FGraphSchemaActionDragDropAction_DataflowVariable::New(TSharedPtr<FEdGraphSchemaAction_DataflowVariable>& InAction)
 {
 	TSharedRef<FGraphSchemaActionDragDropAction_DataflowVariable> Operation = MakeShareable(new FGraphSchemaActionDragDropAction_DataflowVariable);
 	Operation->VariableAction = InAction;
@@ -240,9 +241,27 @@ FEdGraphSchemaAction_DataflowSubGraph::FEdGraphSchemaAction_DataflowSubGraph(UDa
 
 UEdGraphNode* FEdGraphSchemaAction_DataflowSubGraph::PerformAction(class UEdGraph* ParentGraph, UEdGraphPin* FromPin, const FVector2D Location, bool bSelectNewNode)
 {
-	// NOT YET IMPLEMENTED
-	// Should drop a node that represents the subgraph
-	return nullptr;
+	// TODO : move this code to the FEditableDataflow eventually 
+	UEdGraphNode* NewEdNode = nullptr;
+	if (TSharedPtr<FAssetSchemaAction_Dataflow_CreateNode_DataflowEdNode>  SubGraphNodeAction = FAssetSchemaAction_Dataflow_CreateNode_DataflowEdNode::CreateAction(ParentGraph, "FDataflowCallSubGraphNode", SubGraphName))
+	{
+		UDataflowEdNode* DataflowEdNode = Cast<UDataflowEdNode>(SubGraphNodeAction->PerformAction(ParentGraph, nullptr, Location, false));
+		NewEdNode = DataflowEdNode;
+
+		if (DataflowEdNode)
+		{
+			if (TSharedPtr<FDataflowCallSubGraphNode> SubGraphCallNode = StaticCastSharedPtr<FDataflowCallSubGraphNode>(DataflowEdNode->GetDataflowNode()))
+			{
+				if (TStrongObjectPtr<UDataflow> DataflowAsset = DataflowAssetWeakPtr.Pin())
+				{
+					SubGraphCallNode->SetSubGraphName(SubGraphName);
+					DataflowEdNode->UpdatePinsFromDataflowNode();
+					ParentGraph->NotifyNodeChanged(DataflowEdNode);
+				}
+			}
+		}
+	}
+	return NewEdNode;
 }
 
 bool FEdGraphSchemaAction_DataflowSubGraph::CanRenameItem(FText NewName) const
@@ -280,4 +299,18 @@ void FEdGraphSchemaAction_DataflowSubGraph::DuplicateItem()
 	// NOT YET IMPLEMENTED
 }
 
+/////////////////////////////////////////////////////////////////////////////////////////////////////////////////////////
+FGraphSchemaActionDragDropAction_DataflowSubGraph::FGraphSchemaActionDragDropAction_DataflowSubGraph()
+	: FGraphSchemaActionDragDropAction()
+{}
+
+TSharedRef<FGraphSchemaActionDragDropAction_DataflowSubGraph> FGraphSchemaActionDragDropAction_DataflowSubGraph::New(TSharedPtr<FEdGraphSchemaAction_DataflowSubGraph>& InAction)
+{
+	TSharedRef<FGraphSchemaActionDragDropAction_DataflowSubGraph> Operation = MakeShareable(new FGraphSchemaActionDragDropAction_DataflowSubGraph);
+	Operation->SubGraphAction = InAction;
+	Operation->SourceAction = InAction;
+	Operation->Construct();
+	return Operation;
+}
+
 #undef LOCTEXT_NAMESPACE
\ No newline at end of file
diff --git a/Engine/Plugins/Experimental/Dataflow/Source/DataflowEditor/Private/Dataflow/DataflowMembersWidget.cpp b/Engine/Plugins/Experimental/Dataflow/Source/DataflowEditor/Private/Dataflow/DataflowMembersWidget.cpp
--- a/Engine/Plugins/Experimental/Dataflow/Source/DataflowEditor/Private/Dataflow/DataflowMembersWidget.cpp
+++ b/Engine/Plugins/Experimental/Dataflow/Source/DataflowEditor/Private/Dataflow/DataflowMembersWidget.cpp
@@ -596,6 +596,13 @@ FReply SDataflowMembersWidget::OnActionDragged(const TArray<TSharedPtr<FEdGraphS
 		TSharedRef<FGraphSchemaActionDragDropAction_DataflowVariable> DragOperation = FGraphSchemaActionDragDropAction_DataflowVariable::New(VariableAction);
 		return FReply::Handled().BeginDragDrop(DragOperation);
 	}
+
+	if (InAction->GetTypeId() == FEdGraphSchemaAction_DataflowSubGraph::StaticGetTypeId())
+	{
+		TSharedPtr<FEdGraphSchemaAction_DataflowSubGraph> SubGraphAction = StaticCastSharedPtr<FEdGraphSchemaAction_DataflowSubGraph>(InAction);
+		TSharedRef<FGraphSchemaActionDragDropAction_DataflowSubGraph> DragOperation = FGraphSchemaActionDragDropAction_DataflowSubGraph::New(SubGraphAction);
+		return FReply::Handled().BeginDragDrop(DragOperation);
+	}
 	return FReply::Unhandled();
 }
 
diff --git a/Engine/Plugins/Experimental/Dataflow/Source/DataflowEditor/Public/Dataflow/DataflowGraphSchemaAction.h b/Engine/Plugins/Experimental/Dataflow/Source/DataflowEditor/Public/Dataflow/DataflowGraphSchemaAction.h
--- a/Engine/Plugins/Experimental/Dataflow/Source/DataflowEditor/Public/Dataflow/DataflowGraphSchemaAction.h
+++ b/Engine/Plugins/Experimental/Dataflow/Source/DataflowEditor/Public/Dataflow/DataflowGraphSchemaAction.h
@@ -137,7 +137,7 @@ struct FGraphSchemaActionDragDropAction_DataflowVariable : public FGraphSchemaAc
 public:
 	DRAG_DROP_OPERATOR_TYPE(FGraphSchemaActionDragDropAction_DataflowVariable, FGraphSchemaActionDragDropAction)
 
-	static TSharedRef<FGraphSchemaActionDragDropAction_DataflowVariable> New(TSharedPtr<FEdGraphSchemaAction_DataflowVariable> InAction);
+	static TSharedRef<FGraphSchemaActionDragDropAction_DataflowVariable> New(TSharedPtr<FEdGraphSchemaAction_DataflowVariable>& InAction);
 
 private:
 	FGraphSchemaActionDragDropAction_DataflowVariable();
@@ -191,3 +191,22 @@ private:
 	/** Dataflow asset associated with this action */
 	TWeakObjectPtr<UDataflow> DataflowAssetWeakPtr;
 };
+
+////////////////////////////////////////////////////////////////////////////////////////////////////////////////////////////////
+
+/**
+* Dataflow variable drag and drop action
+* see FEdGraphSchemaAction_DataflowSubGraph
+*/
+struct FGraphSchemaActionDragDropAction_DataflowSubGraph : public FGraphSchemaActionDragDropAction
+{
+public:
+	DRAG_DROP_OPERATOR_TYPE(FGraphSchemaActionDragDropAction_DataflowSubGraph, FGraphSchemaActionDragDropAction)
+
+	static TSharedRef<FGraphSchemaActionDragDropAction_DataflowSubGraph> New(TSharedPtr<FEdGraphSchemaAction_DataflowSubGraph>& InAction);
+
+private:
+	FGraphSchemaActionDragDropAction_DataflowSubGraph();
+
+	TSharedPtr<FEdGraphSchemaAction_DataflowSubGraph> SubGraphAction;
+};
\ No newline at end of file
