Dataflow - Fixed the property reset to default.

[RN] Dataflow - Fixed the property "reset to default" buttons in the node details UI, that wasn't correctly invalidating the nodes.

diff --git a/Engine/Plugins/Experimental/Dataflow/Source/DataflowEditor/Private/Dataflow/DataflowEditorCommands.cpp b/Engine/Plugins/Experimental/Dataflow/Source/DataflowEditor/Private/Dataflow/DataflowEditorCommands.cpp
--- a/Engine/Plugins/Experimental/Dataflow/Source/DataflowEditor/Private/Dataflow/DataflowEditorCommands.cpp
+++ b/Engine/Plugins/Experimental/Dataflow/Source/DataflowEditor/Private/Dataflow/DataflowEditorCommands.cpp
@@ -319,15 +319,8 @@ void FDataflowEditorCommands::OnAssetPropertyValueChanged(TObjectPtr<UDataflowBa
 
 void FDataflowEditorCommands::OnPropertyValueChanged(UDataflow* OutDataflow, TSharedPtr<UE::Dataflow::FEngineContext>& Context, UE::Dataflow::FTimestamp& OutLastNodeTimestamp, const FPropertyChangedEvent& InPropertyChangedEvent, const TSet<TObjectPtr<UObject> >& SelectedNodes)
 {
-	switch (InPropertyChangedEvent.ChangeType)
+	if (InPropertyChangedEvent.ChangeType != EPropertyChangeType::Interactive)
 	{
-	case EPropertyChangeType::ValueSet:
-	case EPropertyChangeType::ArrayAdd:
-	case EPropertyChangeType::ArrayRemove:
-	case EPropertyChangeType::ArrayClear:
-	case EPropertyChangeType::ArrayMove:
-	case EPropertyChangeType::Duplicate:
-	case EPropertyChangeType::ToggleEditable:
 		if (ensure(OutDataflow && InPropertyChangedEvent.Property && InPropertyChangedEvent.Property->GetOwnerUObject()))
 		{
 			OutDataflow->Modify();  // Modify must be called even if SelectedNodes is empty because comment nodes aren't part of the selection set but still have properties
@@ -360,9 +353,6 @@ void FDataflowEditorCommands::OnPropertyValueChanged(UDataflow* OutDataflow, TSh
 				}
 			}
 		}
-		break;
-	default:
-		break;
 	}
 }
 
