FEditorViewportClient: allow focusing on objects smaller than 10cm. This change replaces the hardcoded 10.0 lower limit with a member variable that subclasses can modify. Default behavior is unchanged.
FDataflowConstructionViewportClient: set the minumum object size to 0.1cm, since we could be dealing with coordinates in the 0-1 range.

diff --git a/Engine/Plugins/Experimental/Dataflow/Source/DataflowEditor/Private/Dataflow/DataflowConstructionViewportClient.cpp b/Engine/Plugins/Experimental/Dataflow/Source/DataflowEditor/Private/Dataflow/DataflowConstructionViewportClient.cpp
--- a/Engine/Plugins/Experimental/Dataflow/Source/DataflowEditor/Private/Dataflow/DataflowConstructionViewportClient.cpp
+++ b/Engine/Plugins/Experimental/Dataflow/Source/DataflowEditor/Private/Dataflow/DataflowConstructionViewportClient.cpp
@@ -33,6 +33,9 @@ FDataflowConstructionViewportClient::FDataflowConstructionViewportClient(FEditor
 	// We want our near clip plane to be quite close so that we can zoom in further.
 	OverrideNearClipPlane(KINDA_SMALL_NUMBER);
 
+	// Allow focusing on small objects
+	MinimumFocusRadius = 0.1f;
+
 	EngineShowFlags.SetSelectionOutline(true);
 	EngineShowFlags.EnableAdvancedFeatures();
 
diff --git a/Engine/Source/Editor/UnrealEd/Private/EditorViewportClient.cpp b/Engine/Source/Editor/UnrealEd/Private/EditorViewportClient.cpp
--- a/Engine/Source/Editor/UnrealEd/Private/EditorViewportClient.cpp
+++ b/Engine/Source/Editor/UnrealEd/Private/EditorViewportClient.cpp
@@ -817,7 +817,7 @@ void FEditorViewportClient::FocusViewportOnBox( const FBox& BoundingBox, bool bI
 	}
 
 	const FVector Position = BoundingBox.GetCenter();
-	float Radius = FMath::Max<FVector::FReal>(BoundingBox.GetExtent().Size(), 10.f);
+	float Radius = FMath::Max<FVector::FReal>(BoundingBox.GetExtent().Size(), MinimumFocusRadius);
 
 	float AspectToUse = AspectRatio;
 	const FIntPoint ViewportSize = Viewport->GetSizeXY();
diff --git a/Engine/Source/Editor/UnrealEd/Public/EditorViewportClient.h b/Engine/Source/Editor/UnrealEd/Public/EditorViewportClient.h
--- a/Engine/Source/Editor/UnrealEd/Public/EditorViewportClient.h
+++ b/Engine/Source/Editor/UnrealEd/Public/EditorViewportClient.h
@@ -2019,6 +2019,9 @@ protected:
 	/** Custom override function that will be called every ::Draw() until override is disabled */
 	TUniqueFunction<void(FEngineShowFlags&)> OverrideShowFlagsFunc;
 
+	/**  When focusing viewport on a box, treat boxes smaller than this size as being this size */
+	FVector::FReal MinimumFocusRadius = 10.0;
+
 protected:
 	// Used for the display of the current preview light after it has been adjusted
 	FVector2D MovingPreviewLightSavedScreenPos;
