Dataflow : Add rename subgraph API for dataflow ( not yet wired to UI )
* also add an APi to test uniqueness of a Dataflow subobject name

diff --git a/Engine/Plugins/Experimental/Dataflow/Source/DataflowEditor/Private/Dataflow/DataflowAssetEditUtils.cpp b/Engine/Plugins/Experimental/Dataflow/Source/DataflowEditor/Private/Dataflow/DataflowAssetEditUtils.cpp
--- a/Engine/Plugins/Experimental/Dataflow/Source/DataflowEditor/Private/Dataflow/DataflowAssetEditUtils.cpp
+++ b/Engine/Plugins/Experimental/Dataflow/Source/DataflowEditor/Private/Dataflow/DataflowAssetEditUtils.cpp
@@ -247,6 +247,11 @@ namespace UE::Dataflow
 		}
 	}
 
+	bool FEditAssetUtils::IsUniqueDataflowSubObjectName(UDataflow* DataflowAsset, FName SubObjectName)
+	{
+		return ::IsUniqueObjectName(SubObjectName, DataflowAsset);
+	}
+
 	//////////////////////////////////////////////////////////////////////////////////////////////////////////
 	// 
 	// NODE API
@@ -902,7 +907,7 @@ namespace UE::Dataflow
 
 				DataflowAsset.AddSubGraph(NewSubGraph);
 
-				FDataflowAssetDelegates::OnSubGraphsChanged.Broadcast(&DataflowAsset, UniqueSubGraphName);
+				FDataflowAssetDelegates::OnSubGraphsChanged.Broadcast(&DataflowAsset, NewSubGraph->GetFName());
 				return EChangeResult::Changed;
 			};
 
@@ -911,6 +916,32 @@ namespace UE::Dataflow
 		return UniqueSubGraphName;
 	}
 
+	void FEditAssetUtils::RenameSubGraph(UDataflow* DataflowAsset, FName OldSubGraphName, FName NewSubGraphName)
+	{
+		using namespace EditAssetUtils::Private;
+
+		if (DataflowAsset && IsUniqueDataflowSubObjectName(DataflowAsset, NewSubGraphName))
+		{
+			if (UDataflowSubGraph* SubGraphToRename = DataflowAsset->FindSubGraphByName(OldSubGraphName))
+			{
+				const FText TransactionName{ LOCTEXT("RenameDataflowSubGraph", "Rename a Dataflow SubGraph") };
+
+				auto RenameSubGraphInternal =
+					[&SubGraphToRename, &NewSubGraphName](UDataflow& DataflowAsset) -> EChangeResult
+					{
+						if (SubGraphToRename->Rename(*NewSubGraphName.ToString()))
+						{
+							FDataflowAssetDelegates::OnSubGraphsChanged.Broadcast(&DataflowAsset, SubGraphToRename->GetFName());
+							return EChangeResult::Changed;
+						}
+						return EChangeResult::Cancel;
+					};
+
+				ChangeDataflowAssetWithTransaction(DataflowAsset, TransactionName, RenameSubGraphInternal, NAME_None);
+			}
+		}
+	}
+
 }
 
 #undef LOCTEXT_NAMESPACE
\ No newline at end of file
diff --git a/Engine/Plugins/Experimental/Dataflow/Source/DataflowEditor/Private/Dataflow/DataflowAssetEditUtils.h b/Engine/Plugins/Experimental/Dataflow/Source/DataflowEditor/Private/Dataflow/DataflowAssetEditUtils.h
--- a/Engine/Plugins/Experimental/Dataflow/Source/DataflowEditor/Private/Dataflow/DataflowAssetEditUtils.h
+++ b/Engine/Plugins/Experimental/Dataflow/Source/DataflowEditor/Private/Dataflow/DataflowAssetEditUtils.h
@@ -29,6 +29,11 @@ namespace UE::Dataflow
 
 	struct FEditAssetUtils
 	{
+		/**
+		* Test if a dataflow sub-object ( Node, subgraph ... ) name is unique 
+		*/
+		static bool IsUniqueDataflowSubObjectName(UDataflow* DataflowAsset, FName Name);
+
 		/////////////////////////////////////////////////////////////////////////////////////////////////////////////////////////
 		//
 		// Nodes API
@@ -132,5 +137,10 @@ namespace UE::Dataflow
 		* Add a new subgraph to teh Dataflow Asset and return its name
 		*/
 		static FName AddNewSubGraph(UDataflow* DataflowAsset, FName BaseName = DefaultNewSubGraphBaseName);
+
+		/**
+		* Rename a Dataflow asset variable
+		*/
+		static void RenameSubGraph(UDataflow* DataflowAsset, FName OldSubGraphName, FName NewSubGraphName);
 	};
 }
