Dataflow : Changed the signature of Subgraph changed event to send a GUID instead of a name as well as a reason

diff --git a/Engine/Plugins/Experimental/Dataflow/Source/DataflowEditor/Private/Dataflow/DataflowAssetEditUtils.cpp b/Engine/Plugins/Experimental/Dataflow/Source/DataflowEditor/Private/Dataflow/DataflowAssetEditUtils.cpp
--- a/Engine/Plugins/Experimental/Dataflow/Source/DataflowEditor/Private/Dataflow/DataflowAssetEditUtils.cpp
+++ b/Engine/Plugins/Experimental/Dataflow/Source/DataflowEditor/Private/Dataflow/DataflowAssetEditUtils.cpp
@@ -907,7 +907,7 @@ namespace UE::Dataflow
 
 				DataflowAsset.AddSubGraph(NewSubGraph);
 
-				FDataflowAssetDelegates::OnSubGraphsChanged.Broadcast(&DataflowAsset, NewSubGraph->GetFName());
+				FDataflowAssetDelegates::OnSubGraphsChanged.Broadcast(&DataflowAsset, NewSubGraph->GetSubGraphGuid(), UE::Dataflow::ESubGraphChangedReason::Created);
 				return EChangeResult::Changed;
 			};
 
@@ -931,7 +931,7 @@ namespace UE::Dataflow
 					{
 						if (SubGraphToRename->Rename(*NewSubGraphName.ToString()))
 						{
-							FDataflowAssetDelegates::OnSubGraphsChanged.Broadcast(&DataflowAsset, SubGraphToRename->GetFName());
+							FDataflowAssetDelegates::OnSubGraphsChanged.Broadcast(&DataflowAsset, SubGraphToRename->GetSubGraphGuid(), UE::Dataflow::ESubGraphChangedReason::Renamed);
 							return EChangeResult::Changed;
 						}
 						return EChangeResult::Cancel;
diff --git a/Engine/Plugins/Experimental/Dataflow/Source/DataflowEditor/Private/Dataflow/DataflowMembersWidget.cpp b/Engine/Plugins/Experimental/Dataflow/Source/DataflowEditor/Private/Dataflow/DataflowMembersWidget.cpp
--- a/Engine/Plugins/Experimental/Dataflow/Source/DataflowEditor/Private/Dataflow/DataflowMembersWidget.cpp
+++ b/Engine/Plugins/Experimental/Dataflow/Source/DataflowEditor/Private/Dataflow/DataflowMembersWidget.cpp
@@ -523,9 +523,12 @@ void SDataflowMembersWidget::OnObjectPropertyChanged(UObject* InObject, FPropert
 	}
 }
 
-void SDataflowMembersWidget::OnSubGraphsChanged(const UDataflow* InDataflowAsset, FName SubGraphName)
+void SDataflowMembersWidget::OnSubGraphsChanged(const UDataflow* InDataflowAsset, const FGuid& InSubGraphGuid, UE::Dataflow::ESubGraphChangedReason InReason)
 {
-	Refresh();
+	if (InDataflowAsset == EditedAssetWeakPtr)
+	{
+		Refresh();
+	}
 }
 
 void SDataflowMembersWidget::Refresh()
diff --git a/Engine/Plugins/Experimental/Dataflow/Source/DataflowEditor/Public/Dataflow/DataflowMembersWidget.h b/Engine/Plugins/Experimental/Dataflow/Source/DataflowEditor/Public/Dataflow/DataflowMembersWidget.h
--- a/Engine/Plugins/Experimental/Dataflow/Source/DataflowEditor/Public/Dataflow/DataflowMembersWidget.h
+++ b/Engine/Plugins/Experimental/Dataflow/Source/DataflowEditor/Public/Dataflow/DataflowMembersWidget.h
@@ -22,6 +22,11 @@ struct FGraphActionNode;
 struct FPropertyBagPropertyDesc;
 struct FDataflowVariableClassFilter;
 
+namespace UE::Dataflow
+{
+	enum class ESubGraphChangedReason : uint8;
+}
+
 /**
 *
 * Widget to interact with Variables, subgraphs in Dataflow editors
@@ -72,7 +77,7 @@ private:
 	virtual FReply OnKeyDown(const FGeometry& MyGeometry, const FKeyEvent& InKeyEvent);
 
 	void OnObjectPropertyChanged(UObject* InObject, FPropertyChangedEvent& InPropertyChangedEvent);
-	void OnSubGraphsChanged(const UDataflow* InDataflowAsset, FName SubGraphName);
+	void OnSubGraphsChanged(const UDataflow* InDataflowAsset, const FGuid& InSubGraphGuid, UE::Dataflow::ESubGraphChangedReason InReason);
 
 	void OnFilterTextChanged(const FText& InFilterText);
 	FText GetFilterText() const;
diff --git a/Engine/Source/Runtime/Experimental/Dataflow/Engine/Public/Dataflow/DataflowObject.h b/Engine/Source/Runtime/Experimental/Dataflow/Engine/Public/Dataflow/DataflowObject.h
--- a/Engine/Source/Runtime/Experimental/Dataflow/Engine/Public/Dataflow/DataflowObject.h
+++ b/Engine/Source/Runtime/Experimental/Dataflow/Engine/Public/Dataflow/DataflowObject.h
@@ -28,6 +28,16 @@ class UDataflowSubGraph;
 class UMaterial;
 namespace UE::Dataflow { class FGraph; }
 
+namespace UE::Dataflow
+{
+	enum class ESubGraphChangedReason: uint8
+	{
+		Created,
+		Renamed,
+		Deleted
+	};
+}
+
 struct FDataflowAssetDelegates
 {
 	/** Called when variables are edited ( add, remove, changetype , setvalue ) */
@@ -35,8 +45,8 @@ struct FDataflowAssetDelegates
 	static DATAFLOWENGINE_API FOnVariablesChanged OnVariablesChanged;
 
 	/** Called when sub-graphs are edited ( add, remove ) */
-	DECLARE_MULTICAST_DELEGATE_TwoParams(FOnSubGraphsChanged, const UDataflow*, FName);
-	static DATAFLOWENGINE_API FOnVariablesChanged OnSubGraphsChanged;
+	DECLARE_MULTICAST_DELEGATE_ThreeParams(FOnSubGraphsChanged, const UDataflow*, const FGuid&, UE::Dataflow::ESubGraphChangedReason);
+	static DATAFLOWENGINE_API FOnSubGraphsChanged OnSubGraphsChanged;
 };
 
 /**
