Dataflow : Allow renaming of SubGraphs

diff --git a/Engine/Plugins/Experimental/Dataflow/Source/DataflowEditor/Private/Dataflow/DataflowGraphSchemaAction.cpp b/Engine/Plugins/Experimental/Dataflow/Source/DataflowEditor/Private/Dataflow/DataflowGraphSchemaAction.cpp
--- a/Engine/Plugins/Experimental/Dataflow/Source/DataflowEditor/Private/Dataflow/DataflowGraphSchemaAction.cpp
+++ b/Engine/Plugins/Experimental/Dataflow/Source/DataflowEditor/Private/Dataflow/DataflowGraphSchemaAction.cpp
@@ -264,19 +264,25 @@ UEdGraphNode* FEdGraphSchemaAction_DataflowSubGraph::PerformAction(class UEdGrap
 	return NewEdNode;
 }
 
-bool FEdGraphSchemaAction_DataflowSubGraph::CanRenameItem(FText NewName) const
+bool FEdGraphSchemaAction_DataflowSubGraph::CanRenameItem(FText NewNameAsText) const
 {
 	if (TStrongObjectPtr<UDataflow> DataflowAsset = DataflowAssetWeakPtr.Pin())
 	{
-		const bool bSubGraphAlreadyExists = (nullptr != DataflowAsset->FindSubGraphByName(FName(NewName.ToString())));
-		return !bSubGraphAlreadyExists;
+		const FName NewName(NewNameAsText.ToString());
+		return UE::Dataflow::FEditAssetUtils::IsUniqueDataflowSubObjectName(DataflowAsset.Get(), NewName);
 	}
 	return false;
 }
 
-void FEdGraphSchemaAction_DataflowSubGraph::RenameItem(FText NewName)
+void FEdGraphSchemaAction_DataflowSubGraph::RenameItem(FText NewNameAsText)
 {
-	// NOT YET IMPLEMENTED
+	if (TStrongObjectPtr<UDataflow> DataflowAsset = DataflowAssetWeakPtr.Pin())
+	{
+		const FName NewSubGraphName(NewNameAsText.ToString());
+		UE::Dataflow::FEditAssetUtils::RenameSubGraph(DataflowAsset.Get(), GetSubGraphName(), NewSubGraphName);
+
+		// todo : shoudl close the existing tab if any 
+	}
 }
 
 void FEdGraphSchemaAction_DataflowSubGraph::CopyItemToClipboard()
