Adding a UDataflowMesh type to Dataflow Engine. Currently it just holds an FDynamicMesh3 and a set of UMaterialInterfaces, but in the future we could look at making it a container for other mesh types (e.g. Static Mesh) as well.
- Update MeshWarpNode to use this mesh class
- Also adding render support and a few utility nodes (make, split, append)

diff --git a/Engine/Plugins/Experimental/Dataflow/Source/DataflowEditor/Private/Dataflow/DataflowEngineRendering.cpp b/Engine/Plugins/Experimental/Dataflow/Source/DataflowEditor/Private/Dataflow/DataflowEngineRendering.cpp
--- a/Engine/Plugins/Experimental/Dataflow/Source/DataflowEditor/Private/Dataflow/DataflowEngineRendering.cpp
+++ b/Engine/Plugins/Experimental/Dataflow/Source/DataflowEditor/Private/Dataflow/DataflowEngineRendering.cpp
@@ -5,6 +5,7 @@
 #include "Dataflow/DataflowConnectionTypes.h"
 #include "Dataflow/DataflowEditorModule.h"
 #include "Dataflow/DataflowRenderingFactory.h"
+#include "Dataflow/DataflowMesh.h"
 #include "DynamicMesh/DynamicMesh3.h"
 #include "Field/FieldSystemTypes.h"
 #include "GeometryCollection/GeometryCollectionUtility.h"
@@ -17,6 +18,116 @@
 
 namespace UE::Dataflow
 {
+	namespace Private
+	{
+		void RenderDynamicMesh(GeometryCollection::Facades::FRenderingFacade& RenderCollection, 
+			const FString& GeometryGroupName, 
+			const UE::Geometry::FDynamicMesh3& DynamicMesh, 
+			const TArray<TObjectPtr<UMaterialInterface>>& Materials)
+		{
+			const int32 NumVertices = DynamicMesh.VertexCount();
+			const int32 NumTriangles = DynamicMesh.TriangleCount();
+
+			if (NumVertices > 0 && NumTriangles > 0)
+			{
+				// This will contain the valid triangles only
+				TArray<FIntVector> Tris; Tris.Reserve(DynamicMesh.TriangleCount());
+				TArray<int32> MaterialIDs; MaterialIDs.Reserve(DynamicMesh.TriangleCount());
+				const UE::Geometry::FDynamicMeshMaterialAttribute* const MaterialAttribute = DynamicMesh.Attributes()->GetMaterialID();
+
+				// DynamicMesh.TrianglesItr() returns the valid triangles only
+				for (const int32 TriangleID : DynamicMesh.TriangleIndicesItr())
+				{
+					const UE::Geometry::FIndex3i Tri = DynamicMesh.GetTriangle(TriangleID);
+					Tris.Add(FIntVector(Tri.A, Tri.B, Tri.C));
+					if (MaterialAttribute)
+					{
+						MaterialIDs.Add(MaterialAttribute->GetValue(TriangleID));
+					}
+				}
+
+				// This will contain all the vertices (invalid ones too)
+				// Otherwise the IDs need to be remaped
+				TArray<FVector3f> Vertices; Vertices.AddZeroed(DynamicMesh.MaxVertexID());
+
+				// DynamicMesh.VertexIndicesItr() returns the valid vertices only
+				for (int32 VertexID : DynamicMesh.VertexIndicesItr())
+				{
+					Vertices[VertexID] = (FVector3f)DynamicMesh.GetVertex(VertexID);
+				}
+
+				// Add VertexNormal and VertexColor
+				TArray<FVector3f> VertexNormals; VertexNormals.AddUninitialized(Vertices.Num());
+				TArray<FLinearColor> VertexColors; VertexColors.AddUninitialized(Vertices.Num());
+				for (int32 VertexIdx = 0; VertexIdx < VertexNormals.Num(); ++VertexIdx)
+				{
+					// TODO: Get the normal from FDynamicMesh3
+					VertexNormals[VertexIdx] = FVector3f(0.f);
+					VertexColors[VertexIdx] = FLinearColor(FDataflowEditorModule::SurfaceColor);
+				}
+
+				const int32 NumUVLayers = DynamicMesh.Attributes()->NumUVLayers();
+
+				const int32 GeometryIndex = RenderCollection.StartGeometryGroup(GeometryGroupName);
+
+				if (MaterialIDs.Num() != Tris.Num() || Materials.IsEmpty() || NumUVLayers == 0)
+				{
+					// No materials.
+					RenderCollection.AddSurface(MoveTemp(Vertices), MoveTemp(Tris), MoveTemp(VertexNormals), MoveTemp(VertexColors));
+				}
+				else
+				{
+					// GetUVs and materials.
+					// The render facade currently assumes there are no internal uv seams (i.e., each vertex has one UV, instead of wedges). 
+					// This will be the case for meshes that come from SKMs/SMs/Cloth
+					TArray<const UE::Geometry::FDynamicMeshUVOverlay*> UVLayers;
+					UVLayers.SetNumUninitialized(NumUVLayers);
+					for (int32 Layer = 0; Layer < NumUVLayers; ++Layer)
+					{
+						UVLayers[Layer] = DynamicMesh.Attributes()->GetUVLayer(Layer);
+					}
+
+					TArray<FVector2f> DefaultUVs;
+					DefaultUVs.SetNumZeroed(NumUVLayers);
+
+					TArray<TArray<FVector2f>> UVs;
+					UVs.Init(DefaultUVs, Vertices.Num()); // RenderFacade->DynamicMesh code expects all vertices to have full UV sets.
+
+					for (const int32 TriangleIndex : DynamicMesh.TriangleIndicesItr())
+					{
+						const UE::Geometry::FIndex3i Tri = DynamicMesh.GetTriangle(TriangleIndex);
+
+						for (int32 TriangleVertexIndex = 0; TriangleVertexIndex < 3; ++TriangleVertexIndex)
+						{
+							const int32 VertexIndex = Tri[TriangleVertexIndex];
+
+							for (int32 Layer = 0; Layer < NumUVLayers; ++Layer)
+							{
+								if (UVLayers[Layer])
+								{
+									UVs[VertexIndex][Layer] = UVLayers[Layer]->GetElementAtVertex(TriangleIndex, VertexIndex);
+								}
+							}
+						}
+					}
+
+					TArray<FString> MaterialPaths;
+					MaterialPaths.Reserve(Materials.Num());
+					for (const TObjectPtr<UMaterialInterface>& Material : Materials)
+					{
+						MaterialPaths.Add(Material ? Material->GetPathName() : "");
+					}
+
+					RenderCollection.AddSurface(MoveTemp(Vertices), MoveTemp(Tris), MoveTemp(VertexNormals), MoveTemp(VertexColors), MoveTemp(UVs), MoveTemp(MaterialIDs), MoveTemp(MaterialPaths));
+
+				}
+
+				RenderCollection.EndGeometryGroup(GeometryIndex);
+			}
+		}
+
+	}
+
 	void RenderBasicGeometryCollection(GeometryCollection::Facades::FRenderingFacade& RenderCollection, const UE::Dataflow::FGraphRenderingState& State, TArray<FLinearColor>* VertexColorOverride = nullptr)
 	{
 		FManagedArrayCollection Default;
@@ -422,119 +533,50 @@ namespace UE::Dataflow
 				{
 					const UE::Geometry::FDynamicMesh3& DynamicMesh = Mesh->GetMeshRef();
 
-					const int32 NumVertices = DynamicMesh.VertexCount();
-					const int32 NumTriangles = DynamicMesh.TriangleCount();
-
-					if (NumVertices > 0 && NumTriangles > 0)
+					TArray<TObjectPtr<UMaterialInterface>> Materials;
+					if (State.GetRenderOutputs().IsValidIndex(1))
 					{
-						// This will contain the valid triangles only
-						TArray<FIntVector> Tris; Tris.Reserve(DynamicMesh.TriangleCount());
-						TArray<int32> MaterialIDs; MaterialIDs.Reserve(DynamicMesh.TriangleCount());
-						const UE::Geometry::FDynamicMeshMaterialAttribute* const MaterialAttribute = DynamicMesh.Attributes()->GetMaterialID();
-
-						// DynamicMesh.TrianglesItr() returns the valid triangles only
-						for (const int32 TriangleID : DynamicMesh.TriangleIndicesItr())
-						{
-							const UE::Geometry::FIndex3i Tri = DynamicMesh.GetTriangle(TriangleID);
-							Tris.Add(FIntVector(Tri.A, Tri.B, Tri.C));
-							if (MaterialAttribute)
-							{
-								MaterialIDs.Add(MaterialAttribute->GetValue(TriangleID));
-							}
-						}
-
-						// This will contain all the vertices (invalid ones too)
-						// Otherwise the IDs need to be remaped
-						TArray<FVector3f> Vertices; Vertices.AddZeroed(DynamicMesh.MaxVertexID());
-
-						// DynamicMesh.VertexIndicesItr() returns the valid vertices only
-						for (int32 VertexID : DynamicMesh.VertexIndicesItr())
-						{
-							Vertices[VertexID] = (FVector3f)DynamicMesh.GetVertex(VertexID);
-						}
-
-						// Add VertexNormal and VertexColor
-						TArray<FVector3f> VertexNormals; VertexNormals.AddUninitialized(Vertices.Num());
-						TArray<FLinearColor> VertexColors; VertexColors.AddUninitialized(Vertices.Num());
-						for (int32 VertexIdx = 0; VertexIdx < VertexNormals.Num(); ++VertexIdx)
-						{
-							// TODO: Get the normal from FDynamicMesh3
-							VertexNormals[VertexIdx] = FVector3f(0.f);
-							VertexColors[VertexIdx] = FLinearColor(FDataflowEditorModule::SurfaceColor);
-						}
-
-						// See if we have materials to render.
-						TArray<TObjectPtr<UMaterialInterface>> Materials;
-						if (State.GetRenderOutputs().IsValidIndex(1))
-						{
-							TArray<TObjectPtr<UMaterialInterface>> MaterialsDefault;
-							Materials = State.GetValue<TArray<TObjectPtr<UMaterialInterface>>>(State.GetRenderOutputs()[1], MaterialsDefault);
-						}
-
-						const int32 NumUVLayers = DynamicMesh.Attributes()->NumUVLayers();
-
-						const int32 GeometryIndex = RenderCollection.StartGeometryGroup(State.GetGuid().ToString());
-
-						if (MaterialIDs.Num() != Tris.Num() || Materials.IsEmpty() || NumUVLayers == 0)
-						{
-							// No materials.
-							RenderCollection.AddSurface(MoveTemp(Vertices), MoveTemp(Tris), MoveTemp(VertexNormals), MoveTemp(VertexColors));
-						}
-						else
-						{
-							// GetUVs and materials.
-							// The render facade currently assumes there are no internal uv seams (i.e., each vertex has one UV, instead of wedges). 
-							// This will be the case for meshes that come from SKMs/SMs/Cloth
-							TArray<const UE::Geometry::FDynamicMeshUVOverlay*> UVLayers;
-							UVLayers.SetNumUninitialized(NumUVLayers);
-							for (int32 Layer = 0; Layer < NumUVLayers; ++Layer)
-							{
-								UVLayers[Layer] = DynamicMesh.Attributes()->GetUVLayer(Layer);
-							}
-
-							TArray<FVector2f> DefaultUVs;
-							DefaultUVs.SetNumZeroed(NumUVLayers);
-
-							TArray<TArray<FVector2f>> UVs;
-							UVs.Init(DefaultUVs, Vertices.Num()); // RenderFacade->DynamicMesh code expects all vertices to have full UV sets.
+						TArray<TObjectPtr<UMaterialInterface>> MaterialsDefault;
+						Materials = State.GetValue<TArray<TObjectPtr<UMaterialInterface>>>(State.GetRenderOutputs()[1], MaterialsDefault);
+					}
 
-							for (const int32 TriangleIndex : DynamicMesh.TriangleIndicesItr())
-							{
-								const UE::Geometry::FIndex3i Tri = DynamicMesh.GetTriangle(TriangleIndex);
+					Private::RenderDynamicMesh(RenderCollection, State.GetGuid().ToString(), DynamicMesh, Materials);
+				}
+			}
+		}
+	};
 
-								for (int32 TriangleVertexIndex = 0; TriangleVertexIndex < 3; ++TriangleVertexIndex)
-								{
-									const int32 VertexIndex = Tri[TriangleVertexIndex];
-
-									for (int32 Layer = 0; Layer < NumUVLayers; ++Layer)
-									{
-										if (UVLayers[Layer])
-										{
-											UVs[VertexIndex][Layer] = UVLayers[Layer]->GetElementAtVertex(TriangleIndex, VertexIndex);
-										}
-									}
-								}
-							}
 
-							TArray<FString> MaterialPaths;
-							MaterialPaths.Reserve(Materials.Num());
-							for (const TObjectPtr<UMaterialInterface>& Material : Materials)
-							{
-								MaterialPaths.Add(Material ? Material->GetPathName() : "");
-							}
+	class FDataflowMeshSurfaceRenderCallbacks : public FRenderingFactory::ICallbackInterface
+	{
+		virtual UE::Dataflow::FRenderKey GetRenderKey() const override
+		{
+			return { "SurfaceRender", FName("UDataflowMesh") };
+		}
 
-							RenderCollection.AddSurface(MoveTemp(Vertices), MoveTemp(Tris), MoveTemp(VertexNormals), MoveTemp(VertexColors), MoveTemp(UVs), MoveTemp(MaterialIDs), MoveTemp(MaterialPaths));
+		virtual bool CanRender(const UE::Dataflow::IDataflowConstructionViewMode& ViewMode) const override
+		{
+			return (ViewMode.GetName() == FDataflowConstruction3DViewMode::Name);
+		}
 
-						}
+		virtual void Render(GeometryCollection::Facades::FRenderingFacade& RenderCollection, const FGraphRenderingState& State) override
+		{
+			if (State.GetRenderOutputs().Num())
+			{
+				FName PrimaryOutput = State.GetRenderOutputs()[0];
 
-						RenderCollection.EndGeometryGroup(GeometryIndex);
+				TObjectPtr<UDataflowMesh> Default;
+				if (const TObjectPtr<UDataflowMesh> Mesh = State.GetValue(PrimaryOutput, Default))
+				{
+					if (const UE::Geometry::FDynamicMesh3* const DynamicMesh = Mesh->GetDynamicMesh())
+					{
+						Private::RenderDynamicMesh(RenderCollection, State.GetGuid().ToString(), *DynamicMesh, Mesh->GetMaterials());
 					}
 				}
 			}
 		}
 	};
 
-
 	class FBoxSurfaceRenderCallbacks : public FRenderingFactory::ICallbackInterface
 	{
 		virtual UE::Dataflow::FRenderKey GetRenderKey() const override
@@ -824,6 +866,7 @@ namespace UE::Dataflow
 		FRenderingFactory::GetInstance()->RegisterCallbacks(MakeUnique<FGeometryCollectionSurfaceRenderCallbacks>());
 		FRenderingFactory::GetInstance()->RegisterCallbacks(MakeUnique<FGeometryCollectionSurfaceWeightsRenderCallbacks>());
 		FRenderingFactory::GetInstance()->RegisterCallbacks(MakeUnique<FDynamicMesh3SurfaceRenderCallbacks>());
+		FRenderingFactory::GetInstance()->RegisterCallbacks(MakeUnique<FDataflowMeshSurfaceRenderCallbacks>());
 		FRenderingFactory::GetInstance()->RegisterCallbacks(MakeUnique<FBoxSurfaceRenderCallbacks>());
 		FRenderingFactory::GetInstance()->RegisterCallbacks(MakeUnique<FSphereSurfaceRenderCallbacks>());
 		FRenderingFactory::GetInstance()->RegisterCallbacks(MakeUnique<FPointRenderCallbacks>());
diff --git a/Engine/Plugins/Experimental/GeometryCollectionPlugin/Source/GeometryCollectionNodes/Private/Dataflow/GeometryCollectionMeshNodes.cpp b/Engine/Plugins/Experimental/GeometryCollectionPlugin/Source/GeometryCollectionNodes/Private/Dataflow/GeometryCollectionMeshNodes.cpp
--- a/Engine/Plugins/Experimental/GeometryCollectionPlugin/Source/GeometryCollectionNodes/Private/Dataflow/GeometryCollectionMeshNodes.cpp
+++ b/Engine/Plugins/Experimental/GeometryCollectionPlugin/Source/GeometryCollectionNodes/Private/Dataflow/GeometryCollectionMeshNodes.cpp
@@ -50,6 +50,9 @@ namespace UE::Dataflow
 		DATAFLOW_NODE_REGISTER_CREATION_FACTORY(FCollectionToMeshDataflowNode);
 		DATAFLOW_NODE_REGISTER_CREATION_FACTORY(FStaticMeshToMeshDataflowNode);
 		DATAFLOW_NODE_REGISTER_CREATION_FACTORY(FMeshAppendDataflowNode);
+		DATAFLOW_NODE_REGISTER_CREATION_FACTORY(FDataflowMeshAppendDataflowNode);
+		DATAFLOW_NODE_REGISTER_CREATION_FACTORY(FMakeDataflowMeshDataflowNode);
+		DATAFLOW_NODE_REGISTER_CREATION_FACTORY(FSplitDataflowMeshDataflowNode);
 		DATAFLOW_NODE_REGISTER_CREATION_FACTORY(FMeshCopyToPointsDataflowNode);
 		DATAFLOW_NODE_REGISTER_CREATION_FACTORY(FGetMeshDataDataflowNode);
 		DATAFLOW_NODE_REGISTER_CREATION_FACTORY(FApplyMeshProcessorToMeshDataflowNode);
@@ -316,6 +319,140 @@ void FMeshAppendDataflowNode::Evaluate(UE::Dataflow::FContext& Context, const FD
 	}
 }
 
+FDataflowMeshAppendDataflowNode::FDataflowMeshAppendDataflowNode(const UE::Dataflow::FNodeParameters& InParam, FGuid InGuid)
+	: FDataflowNode(InParam, InGuid)
+{
+	RegisterInputConnection(&Mesh);
+	RegisterOutputConnection(&Mesh, &Mesh);
+	RegisterInputConnection(&AppendMesh);
+}
+
+void FDataflowMeshAppendDataflowNode::Evaluate(UE::Dataflow::FContext& Context, const FDataflowOutput* Out) const
+{
+	if (Out->IsA(&Mesh))
+	{
+		TObjectPtr<UDataflowMesh> NewMesh = NewObject<UDataflowMesh>();
+
+		if (const UDataflowMesh* const DataflowMesh1 = GetValue(Context, &Mesh))
+		{
+			if (const UDataflowMesh* const DataflowMesh2 = GetValue(Context, &AppendMesh))
+			{
+				if (const UE::Geometry::FDynamicMesh3* const DynamicMesh1 = DataflowMesh1->GetDynamicMesh())
+				{
+					if (const UE::Geometry::FDynamicMesh3* const DynamicMesh2 = DataflowMesh2->GetDynamicMesh())
+					{
+						if (DynamicMesh1->VertexCount() > 0 && DynamicMesh2->VertexCount() > 0)
+						{
+							UE::Geometry::FDynamicMesh3 ResultDynamicMesh;
+							UE::Geometry::FDynamicMeshEditor MeshEditor(&ResultDynamicMesh);
+							ResultDynamicMesh.EnableAttributes();
+							ResultDynamicMesh.Attributes()->EnableMaterialID();
+
+							UE::Geometry::FMeshIndexMappings IndexMaps1;
+							MeshEditor.AppendMesh(DynamicMesh1, IndexMaps1);
+
+							UE::Geometry::FMeshIndexMappings IndexMaps2;
+							MeshEditor.AppendMesh(DynamicMesh2, IndexMaps2);
+
+							// Reindex material IDs
+							if (DynamicMesh1->HasAttributes() && DynamicMesh1->Attributes()->HasMaterialID() && DynamicMesh2->HasAttributes() && DynamicMesh2->Attributes()->HasMaterialID())
+							{
+								const int32 MaterialIDOffset = DataflowMesh1->GetMaterials().Num();
+
+								for (const int32 Mesh2TriangleIndex : DynamicMesh2->TriangleIndicesItr())
+								{
+									int32 InputMaterialID;
+									DynamicMesh2->Attributes()->GetMaterialID()->GetValue(Mesh2TriangleIndex, &InputMaterialID);
+
+									const int32 NewTriangleIndex = IndexMaps2.GetNewTriangle(Mesh2TriangleIndex);
+									ResultDynamicMesh.Attributes()->GetMaterialID()->SetValue(NewTriangleIndex, MaterialIDOffset + InputMaterialID);
+								}
+							}
+
+							NewMesh->SetDynamicMesh(MoveTemp(ResultDynamicMesh));
+						}
+						else if (DynamicMesh1->VertexCount() > 0)
+						{
+							NewMesh->SetDynamicMesh(*DynamicMesh1);
+						}
+						else if (DynamicMesh2->VertexCount() > 0)
+						{
+							NewMesh->SetDynamicMesh(*DynamicMesh2);
+						}
+					}
+				}	// end if DynamicMesh1
+
+				// Materials
+				NewMesh->AddMaterials(DataflowMesh1->GetMaterials());
+				NewMesh->AddMaterials(DataflowMesh2->GetMaterials());
+			}
+		}
+
+		SetValue(Context, NewMesh, &Mesh);
+	}
+}
+
+FMakeDataflowMeshDataflowNode::FMakeDataflowMeshDataflowNode(const UE::Dataflow::FNodeParameters& InParam, FGuid InGuid)
+	: FDataflowNode(InParam, InGuid)
+{
+	RegisterInputConnection(&InMesh);
+	RegisterInputConnection(&InMaterials);
+	RegisterOutputConnection(&Mesh);
+}
+
+void FMakeDataflowMeshDataflowNode::Evaluate(UE::Dataflow::FContext& Context, const FDataflowOutput* Out) const
+{
+	if (Out->IsA(&Mesh))
+	{
+		TObjectPtr<UDataflowMesh> NewMesh = NewObject<UDataflowMesh>();
+		
+		if (UDynamicMesh* const InUDynamicMesh = GetValue(Context, &InMesh))
+		{
+			InUDynamicMesh->ProcessMesh([NewMesh](const UE::Geometry::FDynamicMesh3& InFDynamicMesh)
+			{
+				NewMesh->SetDynamicMesh(InFDynamicMesh);
+			});
+		}
+
+		TArray<TObjectPtr<UMaterialInterface>> MaterialArray = GetValue(Context, &InMaterials);
+		NewMesh->SetMaterials(MoveTemp(MaterialArray));
+
+		SetValue(Context, NewMesh, &Mesh);
+	}
+}
+
+FSplitDataflowMeshDataflowNode::FSplitDataflowMeshDataflowNode(const UE::Dataflow::FNodeParameters& InParam, FGuid InGuid)
+	: FDataflowNode(InParam, InGuid)
+{
+	RegisterInputConnection(&InMesh);
+	RegisterOutputConnection(&Mesh);
+	RegisterOutputConnection(&MaterialArray);
+}
+
+void FSplitDataflowMeshDataflowNode::Evaluate(UE::Dataflow::FContext& Context, const FDataflowOutput* Out) const
+{
+	TObjectPtr<UDynamicMesh> NewMesh = NewObject<UDynamicMesh>();
+	TArray<TObjectPtr<UMaterialInterface>> Materials;
+
+	if (Out->IsA(&Mesh))
+	{
+		if (const UDataflowMesh* const InDataflowMesh = GetValue(Context, &InMesh))
+		{
+			NewMesh->SetMesh(InDataflowMesh->GetDynamicMeshRef());
+		}
+		SetValue(Context, NewMesh, &Mesh);
+	}
+	else if (Out->IsA(&MaterialArray))
+	{
+		if (UDataflowMesh* const InDataflowMesh = GetValue(Context, &InMesh))
+		{
+			Materials = InDataflowMesh->GetMaterials();
+		}
+		SetValue(Context, Materials, &MaterialArray);
+	}
+}
+
+
 void FMeshCopyToPointsDataflowNode::Evaluate(UE::Dataflow::FContext& Context, const FDataflowOutput* Out) const
 {
 	if (Out->IsA<TObjectPtr<UDynamicMesh>>(&Mesh))
diff --git a/Engine/Plugins/Experimental/GeometryCollectionPlugin/Source/GeometryCollectionNodes/Public/Dataflow/GeometryCollectionMeshNodes.h b/Engine/Plugins/Experimental/GeometryCollectionPlugin/Source/GeometryCollectionNodes/Public/Dataflow/GeometryCollectionMeshNodes.h
--- a/Engine/Plugins/Experimental/GeometryCollectionPlugin/Source/GeometryCollectionNodes/Public/Dataflow/GeometryCollectionMeshNodes.h
+++ b/Engine/Plugins/Experimental/GeometryCollectionPlugin/Source/GeometryCollectionNodes/Public/Dataflow/GeometryCollectionMeshNodes.h
@@ -4,6 +4,7 @@
 
 #include "CoreMinimal.h"
 #include "Dataflow/DataflowEngine.h"
+#include "Dataflow/DataflowMesh.h"
 #include "GeometryCollection/ManagedArrayCollection.h"
 #include "Dataflow/DataflowSelection.h"
 #include "DynamicMesh/DynamicMesh3.h"
@@ -481,6 +482,104 @@ public:
 
 };
 
+/**
+ *
+ * Combine two Dataflow meshes
+ *
+ */
+USTRUCT(meta = (DataflowGeometryCollection))
+struct FDataflowMeshAppendDataflowNode : public FDataflowNode
+{
+	GENERATED_USTRUCT_BODY()
+	DATAFLOW_NODE_DEFINE_INTERNAL(FDataflowMeshAppendDataflowNode, "DataflowMeshAppend", "Mesh|Utilities", "")
+	DATAFLOW_NODE_RENDER_TYPE("SurfaceRender", FName("UDataflowMesh"), "Mesh" )
+
+public:
+
+	FDataflowMeshAppendDataflowNode(const UE::Dataflow::FNodeParameters& InParam, FGuid InGuid = FGuid::NewGuid());
+
+private:
+
+	/** Mesh input/output */
+	UPROPERTY(meta = (DataflowInput, DataflowOutput, DataflowPassthrough = "Mesh", DataflowIntrinsic))
+	TObjectPtr<UDataflowMesh> Mesh;
+
+	/** Mesh to append  */
+	UPROPERTY(meta = (DataflowInput, DataflowIntrinsic))
+	TObjectPtr<const UDataflowMesh> AppendMesh;
+
+	virtual void Evaluate(UE::Dataflow::FContext& Context, const FDataflowOutput* Out) const override;
+
+};
+
+/**
+ *
+ * Create a UDataflow mesh from an input UDynamicMesh and material array
+ *
+ */
+USTRUCT(meta = (DataflowGeometryCollection))
+struct FMakeDataflowMeshDataflowNode : public FDataflowNode
+{
+	GENERATED_USTRUCT_BODY()
+	DATAFLOW_NODE_DEFINE_INTERNAL(FMakeDataflowMeshDataflowNode, "MakeDataflowMesh", "Mesh|Utilities", "")
+	DATAFLOW_NODE_RENDER_TYPE("SurfaceRender", FName("UDataflowMesh"), "Mesh")
+
+public:
+
+	FMakeDataflowMeshDataflowNode(const UE::Dataflow::FNodeParameters& InParam, FGuid InGuid = FGuid::NewGuid());
+
+private:
+
+	/** DynamicMesh input */
+	UPROPERTY(meta = (DataflowInput, DataflowIntrinsic))
+	TObjectPtr<UDynamicMesh> InMesh;
+
+	/** Materials input */
+	UPROPERTY(meta = (DataflowInput, DataflowIntrinsic))
+	TArray<TObjectPtr<UMaterialInterface>> InMaterials;
+
+	/** DataflowMesh output */
+	UPROPERTY(meta = (DataflowOutput))
+	TObjectPtr<UDataflowMesh> Mesh;
+
+	virtual void Evaluate(UE::Dataflow::FContext& Context, const FDataflowOutput* Out) const override;
+};
+
+
+/**
+ *
+ * Split a UDataflow mesh into a UDynamicMesh and a material array
+ *
+ */
+USTRUCT(meta = (DataflowGeometryCollection))
+struct FSplitDataflowMeshDataflowNode : public FDataflowNode
+{
+	GENERATED_USTRUCT_BODY()
+	DATAFLOW_NODE_DEFINE_INTERNAL(FSplitDataflowMeshDataflowNode, "SplitDataflowMesh", "Mesh|Utilities", "")
+	DATAFLOW_NODE_RENDER_TYPE("SurfaceRender", FName("FDynamicMesh3"), { "Mesh", "MaterialArray" })
+
+public:
+
+	FSplitDataflowMeshDataflowNode(const UE::Dataflow::FNodeParameters& InParam, FGuid InGuid = FGuid::NewGuid());
+
+private:
+
+	/** DataflowMesh input */
+	UPROPERTY(meta = (DataflowInput, DataflowIntrinsic))
+	TObjectPtr<UDataflowMesh> InMesh;
+
+	/** DyanmicMesh output */
+	UPROPERTY(meta = (DataflowOutput))
+	TObjectPtr<UDynamicMesh> Mesh;
+
+	/** Materials output */
+	UPROPERTY(meta = (DataflowOutput))
+	TArray<TObjectPtr<UMaterialInterface>> MaterialArray;
+
+	virtual void Evaluate(UE::Dataflow::FContext& Context, const FDataflowOutput* Out) const override;
+};
+
+
 
 
 /**
diff --git a/Engine/Plugins/Experimental/MeshResizing/Source/MeshResizingNodes/Private/MeshResizing/MeshWarpNode.cpp b/Engine/Plugins/Experimental/MeshResizing/Source/MeshResizingNodes/Private/MeshResizing/MeshWarpNode.cpp
--- a/Engine/Plugins/Experimental/MeshResizing/Source/MeshResizingNodes/Private/MeshResizing/MeshWarpNode.cpp
+++ b/Engine/Plugins/Experimental/MeshResizing/Source/MeshResizingNodes/Private/MeshResizing/MeshWarpNode.cpp
@@ -3,12 +3,15 @@
 #include "MeshResizing/RBFInterpolation.h"
 #include "DynamicMesh/DynamicMesh3.h"
 #include "DynamicMesh/DynamicMeshAABBTree3.h"
-#include "UDynamicMesh.h"
 #include "DynamicMeshEditor.h"
+#include "UDynamicMesh.h"
+#include "Dataflow/DataflowMesh.h"
 #include "Materials/MaterialInterface.h"
 
 #include UE_INLINE_GENERATED_CPP_BY_NAME(MeshWarpNode)
 
+using UE::Geometry::FDynamicMesh3;
+
 namespace  UE::MeshResizing::Private
 {
 	void BlendMesh(
@@ -17,8 +20,6 @@ namespace  UE::MeshResizing::Private
 		const FDynamicMesh3& TargetMesh,
 		float Alpha)
 	{
-		using namespace UE::Geometry;
-
 		OutMesh.Copy(SourceMesh);
 
 		ensure(SourceMesh.VertexCount() == TargetMesh.VertexCount());
@@ -80,14 +81,13 @@ namespace  UE::MeshResizing::Private
 		{
 			if (EmbeddingInfo[VertexIndex].ClosestTriangle != INDEX_NONE)
 			{
-				const FVector3d Normal = SourceMesh.GetTriNormal(EmbeddingInfo[VertexIndex].ClosestTriangle);
+				const FVector3d Normal = TargetMesh.GetTriNormal(EmbeddingInfo[VertexIndex].ClosestTriangle);
 				FVector3d TriPointA, TriPointB, TriPointC;
 				TargetMesh.GetTriVertices(EmbeddingInfo[VertexIndex].ClosestTriangle, TriPointA, TriPointB, TriPointC);
 
 				const FVector4d& Weights = EmbeddingInfo[VertexIndex].Weights;
 				const FVector PointOnPlane = Weights[0] * TriPointA + Weights[1] * TriPointB + Weights[2] * TriPointC;
-				const FVector3d CurrentVertexPosition = ResizedMesh.GetVertex(VertexIndex);
-				const FVector FinalPoint = (1.0 - Alpha) * CurrentVertexPosition + Alpha * (PointOnPlane + Weights[3] * Normal);
+				const FVector FinalPoint = PointOnPlane + Weights[3] * Normal;
 
 				ResizedMesh.SetVertex(VertexIndex, FinalPoint);
 			}
@@ -120,64 +120,55 @@ FMeshWarpNode::FMeshWarpNode(const UE::Dataflow::FNodeParameters& InParam, FGuid
 	FDataflowNode(InParam, InGuid)
 {
 	RegisterInputConnection(&MeshToResize);
-	RegisterInputConnection(&MeshToResizeMaterialArray);
 	RegisterInputConnection(&SourceMesh);
 	RegisterInputConnection(&TargetMesh);
 
 	RegisterOutputConnection(&BlendedTargetMesh);
 	RegisterOutputConnection(&ResizedMesh);
-	RegisterOutputConnection(&MaterialArray);
 }
 
 void FMeshWarpNode::Evaluate(UE::Dataflow::FContext& Context, const FDataflowOutput* Out) const
 {
 	using namespace UE::Geometry;
 
-	TObjectPtr<UDynamicMesh> OutResizedMesh = NewObject<UDynamicMesh>();
-	TObjectPtr<UDynamicMesh> OutBlendedTargetMesh = NewObject<UDynamicMesh>();
-	TArray<TObjectPtr<UMaterialInterface>> OutMaterials;
+	TObjectPtr<UDataflowMesh> OutResizedMesh = NewObject<UDataflowMesh>();
+	TObjectPtr<UDataflowMesh> OutBlendedTargetMesh = NewObject<UDataflowMesh>();
 
-	if (TObjectPtr<UDynamicMesh> InMeshToResize = GetValue(Context, &MeshToResize))
+	if (TObjectPtr<UDataflowMesh> InMeshToResize = GetValue(Context, &MeshToResize))
 	{
-		if (TObjectPtr<UDynamicMesh> InSourceUMesh = GetValue(Context, &SourceMesh))
+		if (TObjectPtr<UDataflowMesh> InSourceMesh = GetValue(Context, &SourceMesh))
 		{
-			if(const FDynamicMesh3* const InSourceMesh = InSourceUMesh->GetMeshPtr())
+			if (TObjectPtr<UDataflowMesh> InTargetMesh = GetValue(Context, &TargetMesh))
 			{
-				if (TObjectPtr<UDynamicMesh> InTargetUMesh = GetValue(Context, &TargetMesh))
+				if (InMeshToResize->GetDynamicMesh() && InSourceMesh->GetDynamicMesh() && InTargetMesh->GetDynamicMesh())
 				{
-					if (const FDynamicMesh3* const InTargetMesh = InTargetUMesh->GetMeshPtr())
+					FDynamicMesh3 BlendedTargetFMesh;
+					UE::MeshResizing::Private::BlendMesh(BlendedTargetFMesh, InSourceMesh->GetDynamicMeshRef(), InTargetMesh->GetDynamicMeshRef(), Alpha);
+
+					FDynamicMesh3 ResizeMesh;
+					ResizeMesh.Copy(InMeshToResize->GetDynamicMeshRef());
+					switch (WarpMethod)
 					{
-						FDynamicMesh3 BlendedTargetFMesh;
-						UE::MeshResizing::Private::BlendMesh(BlendedTargetFMesh, *InSourceMesh, *InTargetMesh, Alpha);
-
-						FDynamicMesh3 ResizeMesh;
-						ResizeMesh.Copy(InMeshToResize->GetMeshRef());
-						switch (WarpMethod)
-						{
-						case EMeshResizingWarpMethod::RBFInterpolate:
-							UE::MeshResizing::Private::ApplyRefitToMesh_RBFInterpolate(ResizeMesh, *InSourceMesh, BlendedTargetFMesh, NumInterpolationPoints);
-							break;
-						case EMeshResizingWarpMethod::WrapDeform:
-						default:
-							UE::MeshResizing::Private::ApplyRefitToMesh_WrapDeform(ResizeMesh, *InSourceMesh, BlendedTargetFMesh, 1.0f);
-							break;
-						}
-
-						OutBlendedTargetMesh->SetMesh(MoveTemp(BlendedTargetFMesh));
-						OutResizedMesh->SetMesh(MoveTemp(ResizeMesh));
-
-						if (IsConnected(&MeshToResizeMaterialArray))
-						{
-							OutMaterials = GetValue(Context, &MeshToResizeMaterialArray);
-						}
+					case EMeshResizingWarpMethod::RBFInterpolate:
+						UE::MeshResizing::Private::ApplyRefitToMesh_RBFInterpolate(ResizeMesh, InSourceMesh->GetDynamicMeshRef(), BlendedTargetFMesh, NumInterpolationPoints);
+						break;
+					case EMeshResizingWarpMethod::WrapDeform:
+					default:
+						UE::MeshResizing::Private::ApplyRefitToMesh_WrapDeform(ResizeMesh, InSourceMesh->GetDynamicMeshRef(), BlendedTargetFMesh, 1.0f);
+						break;
 					}
+
+					OutBlendedTargetMesh->SetDynamicMesh(MoveTemp(BlendedTargetFMesh));
+					OutResizedMesh->SetDynamicMesh(MoveTemp(ResizeMesh));
 				}
+
+				OutBlendedTargetMesh->SetMaterials(InTargetMesh->GetMaterials());
+				OutResizedMesh->SetMaterials(InMeshToResize->GetMaterials());
 			}
 		}
 	}
 
 	SetValue(Context, MoveTemp(OutResizedMesh), &ResizedMesh);
 	SetValue(Context, MoveTemp(OutBlendedTargetMesh), &BlendedTargetMesh);
-	SetValue(Context, MoveTemp(OutMaterials), &MaterialArray);
 }
 
diff --git a/Engine/Plugins/Experimental/MeshResizing/Source/MeshResizingNodes/Private/MeshResizing/MeshWarpNode.h b/Engine/Plugins/Experimental/MeshResizing/Source/MeshResizingNodes/Private/MeshResizing/MeshWarpNode.h
--- a/Engine/Plugins/Experimental/MeshResizing/Source/MeshResizingNodes/Private/MeshResizing/MeshWarpNode.h
+++ b/Engine/Plugins/Experimental/MeshResizing/Source/MeshResizingNodes/Private/MeshResizing/MeshWarpNode.h
@@ -7,7 +7,7 @@
 
 #include "MeshWarpNode.generated.h"
 
-class UDynamicMesh;
+class UDataflowMesh;
 class UMaterialInterface;
 
 UENUM()
@@ -23,8 +23,8 @@ struct FMeshWarpNode : public FDataflowNode
 	GENERATED_USTRUCT_BODY()
 	DATAFLOW_NODE_DEFINE_INTERNAL(FMeshWarpNode, "MeshWarp", "MeshResizing", "Mesh Warp")
 	DATAFLOW_NODE_RENDER_TYPE_START()
-	DATAFLOW_NODE_RENDER_TYPE_ADD("SurfaceRender", FName("FDynamicMesh3"), "BlendedTargetMesh")
-	DATAFLOW_NODE_RENDER_TYPE_ADD("SurfaceRender", FName("FDynamicMesh3"), { "ResizedMesh", "MaterialArray" })
+	DATAFLOW_NODE_RENDER_TYPE_ADD("SurfaceRender", FName("UDataflowMesh"), "BlendedTargetMesh")
+	DATAFLOW_NODE_RENDER_TYPE_ADD("SurfaceRender", FName("UDataflowMesh"), "ResizedMesh")
 	DATAFLOW_NODE_RENDER_TYPE_END()
 
 public:
@@ -34,26 +34,20 @@ public:
 private:
 
 	UPROPERTY(meta = (DataflowInput))
-	TObjectPtr<UDynamicMesh> MeshToResize;
+	TObjectPtr<UDataflowMesh> MeshToResize;
 
 	UPROPERTY(meta = (DataflowInput))
-	TArray<TObjectPtr<UMaterialInterface>> MeshToResizeMaterialArray;
+	TObjectPtr<UDataflowMesh> SourceMesh;
 
 	UPROPERTY(meta = (DataflowInput))
-	TObjectPtr<UDynamicMesh> SourceMesh;
+	TObjectPtr<UDataflowMesh> TargetMesh;
 
-	UPROPERTY(meta = (DataflowInput))
-	TObjectPtr<UDynamicMesh> TargetMesh;
-
-
-	UPROPERTY(meta = (DataflowOutput))
-	TObjectPtr<UDynamicMesh> BlendedTargetMesh;
 
 	UPROPERTY(meta = (DataflowOutput))
-	TObjectPtr<UDynamicMesh> ResizedMesh;
+	TObjectPtr<UDataflowMesh> BlendedTargetMesh;
 
 	UPROPERTY(meta = (DataflowOutput))
-	TArray<TObjectPtr<UMaterialInterface>> MaterialArray;
+	TObjectPtr<UDataflowMesh> ResizedMesh;
 
 	UPROPERTY(EditAnywhere, Category = "Warp", meta = (Min = 0.0f, UIMin = 0.0f, Max = 1.0f, UIMax = 1.0f))
 	float Alpha;
diff --git a/Engine/Source/Runtime/Experimental/Dataflow/Engine/DataflowEngine.Build.cs b/Engine/Source/Runtime/Experimental/Dataflow/Engine/DataflowEngine.Build.cs
--- a/Engine/Source/Runtime/Experimental/Dataflow/Engine/DataflowEngine.Build.cs
+++ b/Engine/Source/Runtime/Experimental/Dataflow/Engine/DataflowEngine.Build.cs
@@ -16,7 +16,8 @@ namespace UnrealBuildTool.Rules
 					"Engine",
 					"DataflowCore",
 					"Chaos",
-					"ImageCore"
+					"ImageCore",
+					"GeometryCore"
 				}
 			);
 		}
diff --git a/Engine/Source/Runtime/Experimental/Dataflow/Engine/Private/Dataflow/DataflowMesh.cpp b/Engine/Source/Runtime/Experimental/Dataflow/Engine/Private/Dataflow/DataflowMesh.cpp
new file mode 100644
--- /dev/null
+++ b/Engine/Source/Runtime/Experimental/Dataflow/Engine/Private/Dataflow/DataflowMesh.cpp
@@ -0,0 +1,56 @@
+// Copyright Epic Games, Inc. All Rights Reserved.
+
+#include "Dataflow/DataflowMesh.h"
+#include "Materials/MaterialInterface.h"
+
+UDataflowMesh::UDataflowMesh(const FObjectInitializer& ObjectInitializer)
+	: Super(ObjectInitializer)
+{
+}
+
+const UE::Geometry::FDynamicMesh3* UDataflowMesh::GetDynamicMesh() const
+{
+	return DynamicMesh.Get();
+}
+
+const UE::Geometry::FDynamicMesh3& UDataflowMesh::GetDynamicMeshRef() const
+{
+	return *DynamicMesh;
+}
+
+const TArray<TObjectPtr<UMaterialInterface>>& UDataflowMesh::GetMaterials() const
+{
+	return Materials;
+}
+
+void UDataflowMesh::Serialize(FArchive& Ar)
+{
+	Super::Serialize(Ar);
+
+	if (Ar.IsSaving())
+	{
+		bool bValidDynamicMesh = DynamicMesh.IsValid();
+		Ar << bValidDynamicMesh;
+		if (bValidDynamicMesh)
+		{
+			Ar << *DynamicMesh;
+		}
+	} 
+	else if (Ar.IsLoading())
+	{
+		bool bValidDynamicMesh = false;
+		Ar << bValidDynamicMesh;
+		if (bValidDynamicMesh)
+		{
+			DynamicMesh = MakeUnique<UE::Geometry::FDynamicMesh3>();
+			Ar << *DynamicMesh;
+		}
+		else
+		{
+			DynamicMesh.Reset();
+		}
+	}
+
+	Ar << Materials;
+}
+
diff --git a/Engine/Source/Runtime/Experimental/Dataflow/Engine/Public/Dataflow/DataflowMesh.h b/Engine/Source/Runtime/Experimental/Dataflow/Engine/Public/Dataflow/DataflowMesh.h
new file mode 100644
--- /dev/null
+++ b/Engine/Source/Runtime/Experimental/Dataflow/Engine/Public/Dataflow/DataflowMesh.h
@@ -0,0 +1,59 @@
+// Copyright Epic Games, Inc. All Rights Reserved.
+
+#pragma once
+
+#include "UObject/Object.h"
+#include "DynamicMesh/DynamicMesh3.h"
+#include "DataflowMesh.generated.h"
+
+class UMaterialInterface;
+
+/** 
+* Containter for storing mesh and material information in Dataflow graph
+*/
+UCLASS()
+class DATAFLOWENGINE_API UDataflowMesh : public UObject
+{
+	GENERATED_UCLASS_BODY()
+
+public:
+
+	const UE::Geometry::FDynamicMesh3* GetDynamicMesh() const;
+	const UE::Geometry::FDynamicMesh3& GetDynamicMeshRef() const;
+
+	const TArray<TObjectPtr<UMaterialInterface>>& GetMaterials() const;
+
+	template<typename MeshType>
+	void SetDynamicMesh(MeshType&& InMesh)
+	{
+		if (!DynamicMesh)
+		{
+			DynamicMesh = MakeUnique<UE::Geometry::FDynamicMesh3>(Forward<MeshType>(InMesh));
+		}
+		else
+		{
+			*DynamicMesh = Forward<MeshType>(InMesh);
+		}
+	}
+
+	template<typename MaterialArrayType>
+	void SetMaterials(MaterialArrayType&& InMaterials)
+	{
+		Materials = Forward<MaterialArrayType>(InMaterials);
+	}
+
+	template<typename MaterialArrayType>
+	void AddMaterials(MaterialArrayType&& InMaterials)
+	{
+		Materials.Append(Forward<MaterialArrayType>(InMaterials));
+	}
+
+protected:
+
+	TUniquePtr<UE::Geometry::FDynamicMesh3> DynamicMesh;
+
+	UPROPERTY()
+	TArray<TObjectPtr<UMaterialInterface>> Materials;
+
+	virtual void Serialize(FArchive& Archive) override;
+};
