Chaosflesh: Added check condition for property invalidation

diff --git a/Engine/Plugins/Experimental/Dataflow/Source/DataflowEditor/Private/Dataflow/DataflowEditorCommands.cpp b/Engine/Plugins/Experimental/Dataflow/Source/DataflowEditor/Private/Dataflow/DataflowEditorCommands.cpp
--- a/Engine/Plugins/Experimental/Dataflow/Source/DataflowEditor/Private/Dataflow/DataflowEditorCommands.cpp
+++ b/Engine/Plugins/Experimental/Dataflow/Source/DataflowEditor/Private/Dataflow/DataflowEditorCommands.cpp
@@ -341,10 +341,12 @@ void FDataflowEditorCommands::OnPropertyValueChanged(UDataflow* OutDataflow, TSh
 								Node->SetEnabledState(DataflowNode->IsActive(bCheckIsActiveFlagOnly) ? ENodeEnabledState::Enabled : ENodeEnabledState::Disabled);
 							}
 						}
-						// Invalidate the node and reset the editor timestamp
-						DataflowNode->Invalidate();
-						OutLastNodeTimestamp = UE::Dataflow::FTimestamp::Invalid;
-
+						if (DataflowNode->ShouldInvalidateOnPropertyChanged(InPropertyChangedEvent))
+						{
+							// Invalidate the node and reset the editor timestamp
+							DataflowNode->Invalidate();
+							OutLastNodeTimestamp = UE::Dataflow::FTimestamp::Invalid;
+						}
 						if (Context.IsValid())
 						{
 							DataflowNode->OnPropertyChanged(*Context, InPropertyChangedEvent);
diff --git a/Engine/Source/Runtime/Experimental/Dataflow/Core/Public/Dataflow/DataflowNode.h b/Engine/Source/Runtime/Experimental/Dataflow/Core/Public/Dataflow/DataflowNode.h
--- a/Engine/Source/Runtime/Experimental/Dataflow/Core/Public/Dataflow/DataflowNode.h
+++ b/Engine/Source/Runtime/Experimental/Dataflow/Core/Public/Dataflow/DataflowNode.h
@@ -622,7 +622,9 @@ PRAGMA_ENABLE_DEPRECATION_WARNINGS
 	DATAFLOWCORE_API virtual void ValidateProperties();
 
 	// Called when UPROPERTY members of the dataflow node have been changed in the editor
-	virtual void OnPropertyChanged(UE::Dataflow::FContext& Context, const FPropertyChangedEvent& InPropertyChangedEvent) {}
+	DATAFLOWCORE_API virtual void OnPropertyChanged(UE::Dataflow::FContext& Context, const FPropertyChangedEvent& InPropertyChangedEvent) {}
+
+	DATAFLOWCORE_API virtual bool ShouldInvalidateOnPropertyChanged(const FPropertyChangedEvent& InPropertyChangedEvent) const { return true; };
 
 	bool HasValidConnections() const { return bHasValidConnections; }
 
