Dataflow : Add SubGraph delete feature and update of correspoding tabs when renamed and deleted

diff --git a/Engine/Plugins/Experimental/Dataflow/Source/DataflowEditor/Private/Dataflow/DataflowAssetEditUtils.cpp b/Engine/Plugins/Experimental/Dataflow/Source/DataflowEditor/Private/Dataflow/DataflowAssetEditUtils.cpp
--- a/Engine/Plugins/Experimental/Dataflow/Source/DataflowEditor/Private/Dataflow/DataflowAssetEditUtils.cpp
+++ b/Engine/Plugins/Experimental/Dataflow/Source/DataflowEditor/Private/Dataflow/DataflowAssetEditUtils.cpp
@@ -942,6 +942,32 @@ namespace UE::Dataflow
 		}
 	}
 
+	void FEditAssetUtils::DeleteSubGraph(UDataflow* DataflowAsset, FGuid SubGraphGuid)
+	{
+		using namespace EditAssetUtils::Private;
+
+		if (DataflowAsset)
+		{
+			if (UDataflowSubGraph* SubGraphToDelete = DataflowAsset->FindSubGraphByGuid(SubGraphGuid))
+			{
+				const FText TransactionName{ LOCTEXT("DeleteDataflowSubGraph", "Delete a Dataflow SubGraph") };
+
+				auto DeleteSubGraphInternal =
+					[&SubGraphToDelete, SubGraphGuid](UDataflow& DataflowAsset) -> EChangeResult
+					{
+						FDataflowAssetDelegates::OnSubGraphsChanged.Broadcast(&DataflowAsset, SubGraphGuid, UE::Dataflow::ESubGraphChangedReason::Deleting);
+
+						DataflowAsset.RemoveSubGraph(SubGraphToDelete);
+						
+						FDataflowAssetDelegates::OnSubGraphsChanged.Broadcast(&DataflowAsset, SubGraphGuid, UE::Dataflow::ESubGraphChangedReason::Deleted);
+						return EChangeResult::Changed;
+					};
+
+				ChangeDataflowAssetWithTransaction(DataflowAsset, TransactionName, DeleteSubGraphInternal, NAME_None);
+			}
+		}
+	}
+
 }
 
 #undef LOCTEXT_NAMESPACE
\ No newline at end of file
diff --git a/Engine/Plugins/Experimental/Dataflow/Source/DataflowEditor/Private/Dataflow/DataflowAssetEditUtils.h b/Engine/Plugins/Experimental/Dataflow/Source/DataflowEditor/Private/Dataflow/DataflowAssetEditUtils.h
--- a/Engine/Plugins/Experimental/Dataflow/Source/DataflowEditor/Private/Dataflow/DataflowAssetEditUtils.h
+++ b/Engine/Plugins/Experimental/Dataflow/Source/DataflowEditor/Private/Dataflow/DataflowAssetEditUtils.h
@@ -134,13 +134,18 @@ namespace UE::Dataflow
 		/////////////////////////////////////////////////////////////////////////////////////////////////////////////////////////
 
 		/**
-		* Add a new subgraph to teh Dataflow Asset and return its name
+		* Add a new SubGraph to the Dataflow Asset and return its name
 		*/
 		static FName AddNewSubGraph(UDataflow* DataflowAsset, FName BaseName = DefaultNewSubGraphBaseName);
 
 		/**
-		* Rename a Dataflow asset variable
+		* Rename a Dataflow SubGraph 
 		*/
 		static void RenameSubGraph(UDataflow* DataflowAsset, FName OldSubGraphName, FName NewSubGraphName);
+
+		/**
+		* Delete a Dataflow SubGraph 
+		*/
+		static void DeleteSubGraph(UDataflow* DataflowAsset, FGuid SubGraphGuid);
 	};
 }
diff --git a/Engine/Plugins/Experimental/Dataflow/Source/DataflowEditor/Private/Dataflow/DataflowEditorToolkit.cpp b/Engine/Plugins/Experimental/Dataflow/Source/DataflowEditor/Private/Dataflow/DataflowEditorToolkit.cpp
--- a/Engine/Plugins/Experimental/Dataflow/Source/DataflowEditor/Private/Dataflow/DataflowEditorToolkit.cpp
+++ b/Engine/Plugins/Experimental/Dataflow/Source/DataflowEditor/Private/Dataflow/DataflowEditorToolkit.cpp
@@ -1811,20 +1811,46 @@ void FDataflowEditorToolkit::SetSubGraphTabActiveState(TSharedPtr<SDataflowGraph
 	}
 }
 
-void FDataflowEditorToolkit::OpenSubGraphTab(FName SubGraphName)
+UDataflowSubGraph* FDataflowEditorToolkit::GetSubGraph(FName SubGraphName) const
 {
 	if (const TObjectPtr<UDataflowBaseContent>& EditorContent = GetEditorContent())
 	{
 		ensure(EditorContent);
 		if (UDataflow* DataflowAsset = EditorContent->GetDataflowAsset())
 		{
-			if (UDataflowSubGraph* SubGraph = DataflowAsset->FindSubGraphByName(SubGraphName))
-			{
-				TSharedRef<FTabPayload_UObject> Payload = FTabPayload_UObject::Make(SubGraph);
-				DocumentManager->OpenDocument(Payload, FDocumentTracker::EOpenDocumentCause::OpenNewDocument);
-			}
+			return DataflowAsset->FindSubGraphByName(SubGraphName);
 		}
 	}
+	return nullptr;
+}
+
+void FDataflowEditorToolkit::OpenSubGraphTab(FName SubGraphName)
+{
+	OpenSubGraphTab(GetSubGraph(SubGraphName));
+}
+
+void FDataflowEditorToolkit::OpenSubGraphTab(const UDataflowSubGraph* SubGraph)
+{
+	if (SubGraph)
+	{
+		TSharedRef<FTabPayload_UObject> Payload = FTabPayload_UObject::Make(SubGraph);
+		DocumentManager->OpenDocument(Payload, FDocumentTracker::EOpenDocumentCause::OpenNewDocument);
+	}
+}
+
+void FDataflowEditorToolkit::CloseSubGraphTab(const UDataflowSubGraph* SubGraph)
+{
+	if (SubGraph)
+	{
+		TSharedRef<FTabPayload_UObject> Payload = FTabPayload_UObject::Make(SubGraph);
+		DocumentManager->CloseTab(Payload);
+	}
+}
+
+void FDataflowEditorToolkit::ReOpenSubGraphTab(const UDataflowSubGraph* SubGraph)
+{
+	CloseSubGraphTab(SubGraph);
+	OpenSubGraphTab(SubGraph);
 }
 
 TSharedRef<SGraphEditor> FDataflowEditorToolkit::CreateSubGraphEditorWidget(TSharedRef<FTabInfo> InTabInfo, UDataflowSubGraph* InGraph)
diff --git a/Engine/Plugins/Experimental/Dataflow/Source/DataflowEditor/Private/Dataflow/DataflowGraphSchemaAction.cpp b/Engine/Plugins/Experimental/Dataflow/Source/DataflowEditor/Private/Dataflow/DataflowGraphSchemaAction.cpp
--- a/Engine/Plugins/Experimental/Dataflow/Source/DataflowEditor/Private/Dataflow/DataflowGraphSchemaAction.cpp
+++ b/Engine/Plugins/Experimental/Dataflow/Source/DataflowEditor/Private/Dataflow/DataflowGraphSchemaAction.cpp
@@ -310,7 +310,10 @@ void FEdGraphSchemaAction_DataflowSubGraph::PasteItemFromClipboard()
 
 void FEdGraphSchemaAction_DataflowSubGraph::DeleteItem()
 {
-	// NOT YET IMPLEMENTED
+	if (TStrongObjectPtr<UDataflow> DataflowAsset = DataflowAssetWeakPtr.Pin())
+	{
+		UE::Dataflow::FEditAssetUtils::DeleteSubGraph(DataflowAsset.Get(), SubGraphGuid);
+	}
 }
 
 void FEdGraphSchemaAction_DataflowSubGraph::DuplicateItem()
diff --git a/Engine/Plugins/Experimental/Dataflow/Source/DataflowEditor/Private/Dataflow/DataflowMembersWidget.cpp b/Engine/Plugins/Experimental/Dataflow/Source/DataflowEditor/Private/Dataflow/DataflowMembersWidget.cpp
--- a/Engine/Plugins/Experimental/Dataflow/Source/DataflowEditor/Private/Dataflow/DataflowMembersWidget.cpp
+++ b/Engine/Plugins/Experimental/Dataflow/Source/DataflowEditor/Private/Dataflow/DataflowMembersWidget.cpp
@@ -56,7 +56,7 @@ namespace UE::Dataflow::Private
 	template<typename TAction>
 	struct FMemberSection: public SDataflowMembersWidget::ISection
 	{
-		virtual ~FMemberSection() {}
+		virtual ~FMemberSection() override {}
 
 		virtual const FText& GetTitle() const { return Title; }
 		virtual const SDataflowMembersWidget::FButton* GetAddButton() const { return &AddButton; }
@@ -112,6 +112,8 @@ namespace UE::Dataflow::Private
 	struct FSubGraphsSection : 
 		public FMemberSection<FEdGraphSchemaAction_DataflowSubGraph>
 	{
+		using Super = FMemberSection<FEdGraphSchemaAction_DataflowSubGraph>;
+
 		static TSharedPtr<FSubGraphsSection> Make()
 		{
 			TSharedPtr<FSubGraphsSection> SubGraphsSection = MakeShared<FSubGraphsSection>();
@@ -126,6 +128,8 @@ namespace UE::Dataflow::Private
 			return SubGraphsSection;
 		}
 
+		virtual ~FSubGraphsSection() override {}
+
 		virtual TSharedRef<SWidget> CreateWidgetForAction(FCreateWidgetForActionData* const InCreateData, TSharedPtr<SDataflowGraphEditor> Editor) const override
 		{
 			return SNew(SDataflowSubGraphPaletteItem, InCreateData, Editor);
@@ -173,6 +177,8 @@ namespace UE::Dataflow::Private
 			return VariablesSection;
 		}
 
+		virtual ~FVariablesSection() override {}
+
 		virtual TSharedRef<SWidget> CreateWidgetForAction(FCreateWidgetForActionData* const InCreateData, TSharedPtr<SDataflowGraphEditor> Editor) const override
 		{
 			return SNew(SDataflowVariablePaletteItem, InCreateData, Editor);
@@ -529,6 +535,35 @@ void SDataflowMembersWidget::OnSubGraphsChanged(const UDataflow* InDataflowAsset
 	{
 		Refresh();
 	}
+	
+	// handle SubGraphs tabs
+	if (TSharedPtr<FDataflowEditorToolkit> Toolkit = EditorToolkitWeakPtr.Pin())
+	{
+		if (InDataflowAsset)
+		{
+			if (const UDataflowSubGraph* SubGraph = InDataflowAsset->FindSubGraphByGuid(InSubGraphGuid))
+			{
+				switch (InReason)
+				{
+				case UE::Dataflow::ESubGraphChangedReason::Created:
+					Toolkit->OpenSubGraphTab(SubGraph);
+					break;
+
+				case UE::Dataflow::ESubGraphChangedReason::Renamed:
+					Toolkit->ReOpenSubGraphTab(SubGraph);
+					break;
+
+				case UE::Dataflow::ESubGraphChangedReason::Deleting:
+					Toolkit->CloseSubGraphTab(SubGraph);
+					break;
+
+				case UE::Dataflow::ESubGraphChangedReason::Deleted:
+					// nothing to do 
+					break;
+				}
+			}
+		}
+	}
 }
 
 void SDataflowMembersWidget::Refresh()
diff --git a/Engine/Plugins/Experimental/Dataflow/Source/DataflowEditor/Public/Dataflow/DataflowEditorToolkit.h b/Engine/Plugins/Experimental/Dataflow/Source/DataflowEditor/Public/Dataflow/DataflowEditorToolkit.h
--- a/Engine/Plugins/Experimental/Dataflow/Source/DataflowEditor/Public/Dataflow/DataflowEditorToolkit.h
+++ b/Engine/Plugins/Experimental/Dataflow/Source/DataflowEditor/Public/Dataflow/DataflowEditorToolkit.h
@@ -87,7 +87,11 @@ public:
 	// End of FSerializableObject interface
 
 	void OpenSubGraphTab(FName SubGraphName);
+	void OpenSubGraphTab(const UDataflowSubGraph* SubGraph);
+	void CloseSubGraphTab(const UDataflowSubGraph* SubGraph);
+	void ReOpenSubGraphTab(const UDataflowSubGraph* SubGraph);
 	void SetSubGraphTabActiveState(TSharedPtr<SDataflowGraphEditor> SubGraphEditor, bool bActive);
+	UDataflowSubGraph* GetSubGraph(FName SubGraphName) const;
 
 protected:
 
diff --git a/Engine/Plugins/Experimental/Dataflow/Source/DataflowEditor/Public/Dataflow/DataflowMembersWidget.h b/Engine/Plugins/Experimental/Dataflow/Source/DataflowEditor/Public/Dataflow/DataflowMembersWidget.h
--- a/Engine/Plugins/Experimental/Dataflow/Source/DataflowEditor/Public/Dataflow/DataflowMembersWidget.h
+++ b/Engine/Plugins/Experimental/Dataflow/Source/DataflowEditor/Public/Dataflow/DataflowMembersWidget.h
@@ -51,6 +51,8 @@ public:
 
 	struct ISection
 	{
+		virtual ~ISection() {};
+
 		virtual const FText& GetTitle() const = 0;
 		virtual const FButton* GetAddButton() const = 0;
 
diff --git a/Engine/Source/Runtime/Experimental/Dataflow/Engine/Private/Dataflow/DataflowSubGraphNodes.cpp b/Engine/Source/Runtime/Experimental/Dataflow/Engine/Private/Dataflow/DataflowSubGraphNodes.cpp
--- a/Engine/Source/Runtime/Experimental/Dataflow/Engine/Private/Dataflow/DataflowSubGraphNodes.cpp
+++ b/Engine/Source/Runtime/Experimental/Dataflow/Engine/Private/Dataflow/DataflowSubGraphNodes.cpp
@@ -527,9 +527,12 @@ void FDataflowCallSubGraphNode::OnSomeSubGraphsChanged(const UDataflow* InDatafl
 			// Todo(Dataflow) : we do not care about the name right now but we may if the name is displayed 
 			break;
 
-		case UE::Dataflow::ESubGraphChangedReason::Deleted:
+		case UE::Dataflow::ESubGraphChangedReason::Deleting:
 			SetSubGraphGuid(FGuid());
 			break;
+
+		case UE::Dataflow::ESubGraphChangedReason::Deleted:
+			break;
 		}
 	}
 }
diff --git a/Engine/Source/Runtime/Experimental/Dataflow/Engine/Public/Dataflow/DataflowObject.h b/Engine/Source/Runtime/Experimental/Dataflow/Engine/Public/Dataflow/DataflowObject.h
--- a/Engine/Source/Runtime/Experimental/Dataflow/Engine/Public/Dataflow/DataflowObject.h
+++ b/Engine/Source/Runtime/Experimental/Dataflow/Engine/Public/Dataflow/DataflowObject.h
@@ -34,7 +34,8 @@ namespace UE::Dataflow
 	{
 		Created,
 		Renamed,
-		Deleted
+		Deleting,
+		Deleted,
 	};
 }
 
