Dataflow : Automatically create Input and Output nodes when creating a subgraph

diff --git a/Engine/Plugins/Experimental/Dataflow/Source/DataflowEditor/Private/Dataflow/DataflowGraphEditor.cpp b/Engine/Plugins/Experimental/Dataflow/Source/DataflowEditor/Private/Dataflow/DataflowGraphEditor.cpp
--- a/Engine/Plugins/Experimental/Dataflow/Source/DataflowEditor/Private/Dataflow/DataflowGraphEditor.cpp
+++ b/Engine/Plugins/Experimental/Dataflow/Source/DataflowEditor/Private/Dataflow/DataflowGraphEditor.cpp
@@ -14,7 +14,9 @@
 #include "Dataflow/DataflowXml.h"
 #include "Dataflow/DataflowGraphSchemaAction.h"
 #include "Dataflow/DataflowSelectionNodes.h"
+#include "Dataflow/DataflowSubGraphNodes.h"
 #include "Dataflow/DataflowSchema.h"
+#include "Dataflow/DataflowSubGraph.h"
 #include "HAL/PlatformApplicationMisc.h"
 #include "ScopedTransaction.h"
 #include "IStructureDetailsView.h"
@@ -358,7 +360,17 @@ void SDataflowGraphEditor::AddNewSubGraph() const
 {
 	if (UDataflow* DataflowAsset = GetDataflowAsset())
 	{
-		UE::Dataflow::FEditAssetUtils::AddNewSubGraph(DataflowAsset);
+		const FName NewSubGraphName = UE::Dataflow::FEditAssetUtils::AddNewSubGraph(DataflowAsset);
+		if (UDataflowSubGraph* SubGraph = DataflowAsset->FindSubGraphByName(NewSubGraphName))
+		{
+			static const FName InputNodeName = "Input";
+			static const FVector2D InputNodePos { -100, 0 };
+			UE::Dataflow::FEditAssetUtils::AddNewNode(SubGraph, InputNodePos, InputNodeName, FDataflowSubGraphInputNode::StaticType(), nullptr);
+
+			static const FName OutputNodeName = "Output";
+			static const FVector2D OutputNodePos{ +100, 0 };
+			UE::Dataflow::FEditAssetUtils::AddNewNode(SubGraph, OutputNodePos, OutputNodeName, FDataflowSubGraphOutputNode::StaticType(), nullptr);
+		}
 	}
 }
 
