Dataflow : add function icon in fron of SubGraph

diff --git a/Engine/Plugins/Experimental/Dataflow/Source/DataflowEditor/Private/Dataflow/DataflowSubGraphPaletteItem.cpp b/Engine/Plugins/Experimental/Dataflow/Source/DataflowEditor/Private/Dataflow/DataflowSubGraphPaletteItem.cpp
--- a/Engine/Plugins/Experimental/Dataflow/Source/DataflowEditor/Private/Dataflow/DataflowSubGraphPaletteItem.cpp
+++ b/Engine/Plugins/Experimental/Dataflow/Source/DataflowEditor/Private/Dataflow/DataflowSubGraphPaletteItem.cpp
@@ -25,7 +25,7 @@ TSharedRef<SWidget> SDataflowSubGraphPaletteItem::CreateTextSlotWidget(FCreateWi
 		return SNullWidget::NullWidget;
 	}
 
-	const FText VariableName = FText::FromString(SubGraphAction->GetSubGraphName().ToString());
+	const FText SubGraphName = FText::FromString(SubGraphAction->GetSubGraphName().ToString());
 
 	TSharedPtr<SInlineEditableTextBlock> EditableTextElement;
 
@@ -33,13 +33,19 @@ TSharedRef<SWidget> SDataflowSubGraphPaletteItem::CreateTextSlotWidget(FCreateWi
 	TSharedPtr<SWidget> Widget =
 		SNew(SHorizontalBox)
 
-		// search bar 
+		+SHorizontalBox::Slot()
+		.AutoWidth()
+		.VAlign(VAlign_Center)
+		[
+			SNew(SImage)
+				.Image(FAppStyle::GetBrush(TEXT("GraphEditor.Function_16x")))
+		]
 		+ SHorizontalBox::Slot()
 		.AutoWidth()
 		.HAlign(HAlign_Left)
 		[
 			SAssignNew(EditableTextElement, SInlineEditableTextBlock)
-				.Text(VariableName)
+				.Text(SubGraphName)
 				.OnTextCommitted(this, &SDataflowSubGraphPaletteItem::OnNameTextCommitted)
 				.OnVerifyTextChanged(this, &SDataflowSubGraphPaletteItem::OnNameTextVerifyChanged)
 		]
