Dataflow : Rename Subgraph call nodes when renaming a Subgraph and delete corresponding nodes when deleting a subgraph

diff --git a/Engine/Plugins/Experimental/Dataflow/Source/DataflowEditor/Private/Dataflow/DataflowAssetEditUtils.cpp b/Engine/Plugins/Experimental/Dataflow/Source/DataflowEditor/Private/Dataflow/DataflowAssetEditUtils.cpp
--- a/Engine/Plugins/Experimental/Dataflow/Source/DataflowEditor/Private/Dataflow/DataflowAssetEditUtils.cpp
+++ b/Engine/Plugins/Experimental/Dataflow/Source/DataflowEditor/Private/Dataflow/DataflowAssetEditUtils.cpp
@@ -10,6 +10,7 @@
 #include "Dataflow/DataflowEdNode.h"
 #include "Dataflow/DataflowSNode.h"
 #include "Dataflow/DataflowSubGraph.h"
+#include "Dataflow/DataflowSubGraphNodes.h"
 #include "Editor.h"
 #include "PropertyBagDetails.h"
 #include "ScopedTransaction.h"
@@ -245,6 +246,70 @@ namespace UE::Dataflow
 
 			return nullptr;
 		}
+
+		void RenameSubGraphCallNodes(UDataflow& DataflowAsset, UEdGraph& EdGraph, const FGuid& SubGraphGuid, FName NewSubGraphName)
+		{
+			for (UEdGraphNode* EdNode : EdGraph.Nodes)
+			{
+				if (UDataflowEdNode* DataflowEdNode = Cast<UDataflowEdNode>(EdNode))
+				{
+					if (TSharedPtr<FDataflowNode> DataflowNode = DataflowEdNode->GetDataflowNode())
+					{
+						if (FDataflowCallSubGraphNode* CallNode = DataflowNode->AsType<FDataflowCallSubGraphNode>())
+						{
+							if (CallNode->GetSubGraphGuid() == SubGraphGuid)
+							{
+								const FName UniqueName = GenerateUniqueObjectName(&DataflowAsset, NewSubGraphName);
+								DataflowEdNode->Rename(*UniqueName.ToString(), &EdGraph);
+								CallNode->SetName(UniqueName);
+							}
+						}
+					}
+				}
+			}
+		}
+
+		void RenameSubGraphCallNodes(UDataflow& DataflowAsset, const FGuid& SubGraphGuid, FName NewSubGraphName)
+		{
+			RenameSubGraphCallNodes(DataflowAsset, DataflowAsset, SubGraphGuid, NewSubGraphName);
+			for (UDataflowSubGraph* SubGraph : DataflowAsset.GetSubGraphs())
+			{
+				if (SubGraph)
+				{
+					RenameSubGraphCallNodes(DataflowAsset, *SubGraph, SubGraphGuid, NewSubGraphName);
+				}
+			}
+		}
+
+		void DeleteNodesNoTransaction(UEdGraph* EdGraph, const TArray<UEdGraphNode*>& NodesToDelete)
+		{
+			if (EdGraph && NodesToDelete.Num())
+			{
+				for (UEdGraphNode* EdNode : NodesToDelete)
+				{
+					EdNode->Modify();
+					if (UDataflowEdNode* DataflowEdNode = Cast<UDataflowEdNode>(EdNode))
+					{
+						if (const TSharedPtr<UE::Dataflow::FGraph> DataflowGraph = DataflowEdNode->GetDataflowGraph())
+						{
+							EdGraph->RemoveNode(EdNode);
+							if (TSharedPtr<FDataflowNode> DataflowNode = DataflowEdNode->GetDataflowNode())
+							{
+								DataflowGraph->RemoveNode(DataflowNode);
+							}
+						}
+					}
+					else if (UEdGraphNode_Comment* CommentNode = Cast<UEdGraphNode_Comment>(EdNode))
+					{
+						EdGraph->RemoveNode(CommentNode);
+					}
+
+					// Auto-rename node so that its current name is made available until it is garbage collected
+					EdNode->Rename();
+				}
+			}
+		}
+
 	}
 
 	bool FEditAssetUtils::IsUniqueDataflowSubObjectName(UDataflow* DataflowAsset, FName SubObjectName)
@@ -318,6 +383,27 @@ namespace UE::Dataflow
 		return EdNodeToReturn;
 	}
 	
+	//////////////////////////////////////////////////////////////////////////////////////////////////////////
+	void FEditAssetUtils::DeleteNodes(UEdGraph* EdGraph, const TArray<UEdGraphNode*>& NodesToDelete)
+	{
+		using namespace EditAssetUtils::Private;
+
+		const FText TransactionName{ LOCTEXT("DeleteDataflowNodes", "Delete Dataflow Nodes") };
+
+		if (UDataflow* Dataflow = UDataflow::GetDataflowAssetFromEdGraph(EdGraph))
+		{
+			auto DeleteNodesInternal =
+				[&EdGraph, &NodesToDelete]
+				(UDataflow& DataflowAsset) -> EChangeResult
+				{
+					DeleteNodesNoTransaction(EdGraph, NodesToDelete);
+					return EChangeResult::Changed;
+				};
+
+			ChangeDataflowAssetWithTransaction(Dataflow, TransactionName, DeleteNodesInternal, NAME_None);
+		}
+	}
+
 	//////////////////////////////////////////////////////////////////////////////////////////////////////////
 	void FEditAssetUtils::DuplicateNodes(UEdGraph* EdGraph, const TArray<UEdGraphNode*>& EdNodesToDuplicate, const FVector2D& Location, TArray<UEdGraphNode*>& OutDuplicatedNodes)
 	{
@@ -931,6 +1017,9 @@ namespace UE::Dataflow
 					{
 						if (SubGraphToRename->Rename(*NewSubGraphName.ToString()))
 						{
+							// rename the call nodes using it 
+							RenameSubGraphCallNodes(DataflowAsset, SubGraphToRename->GetSubGraphGuid(), NewSubGraphName);
+
 							FDataflowAssetDelegates::OnSubGraphsChanged.Broadcast(&DataflowAsset, SubGraphToRename->GetSubGraphGuid(), UE::Dataflow::ESubGraphChangedReason::Renamed);
 							return EChangeResult::Changed;
 						}
@@ -957,8 +1046,13 @@ namespace UE::Dataflow
 					{
 						FDataflowAssetDelegates::OnSubGraphsChanged.Broadcast(&DataflowAsset, SubGraphGuid, UE::Dataflow::ESubGraphChangedReason::Deleting);
 
+						// make a copy so that we don't mofdy the array while iterating through it 
+						const TArray<UEdGraphNode*> NodesToDelete = SubGraphToDelete->Nodes;
+						DeleteNodesNoTransaction(SubGraphToDelete, NodesToDelete);
+
+						// delete the Subgraph 
 						DataflowAsset.RemoveSubGraph(SubGraphToDelete);
-						
+
 						FDataflowAssetDelegates::OnSubGraphsChanged.Broadcast(&DataflowAsset, SubGraphGuid, UE::Dataflow::ESubGraphChangedReason::Deleted);
 						return EChangeResult::Changed;
 					};
diff --git a/Engine/Plugins/Experimental/Dataflow/Source/DataflowEditor/Private/Dataflow/DataflowAssetEditUtils.h b/Engine/Plugins/Experimental/Dataflow/Source/DataflowEditor/Private/Dataflow/DataflowAssetEditUtils.h
--- a/Engine/Plugins/Experimental/Dataflow/Source/DataflowEditor/Private/Dataflow/DataflowAssetEditUtils.h
+++ b/Engine/Plugins/Experimental/Dataflow/Source/DataflowEditor/Private/Dataflow/DataflowAssetEditUtils.h
@@ -59,6 +59,13 @@ namespace UE::Dataflow
 		*/
 		static UEdGraphNode* AddNewComment(UEdGraph* EdGraph, const FVector2D& Location, const FVector2D& Size);
 
+		/**
+		* Delete some nodes from a graph 
+		* @param EdGraph			Main Graph or Subgraph to duplicate the node in
+		* @param NodesToDelete		List of node to delete 
+		*/
+		static void DeleteNodes(UEdGraph* EdGraph, const TArray<UEdGraphNode*>& NodesToDelete);
+
 		/**
 		* Duplicate a node
 		* @param EdGraph			Main Graph or Subgraph to duplicate the node in 
diff --git a/Engine/Plugins/Experimental/Dataflow/Source/DataflowEditor/Private/Dataflow/DataflowEditorCommands.cpp b/Engine/Plugins/Experimental/Dataflow/Source/DataflowEditor/Private/Dataflow/DataflowEditorCommands.cpp
--- a/Engine/Plugins/Experimental/Dataflow/Source/DataflowEditor/Private/Dataflow/DataflowEditorCommands.cpp
+++ b/Engine/Plugins/Experimental/Dataflow/Source/DataflowEditor/Private/Dataflow/DataflowEditorCommands.cpp
@@ -394,30 +394,16 @@ void FDataflowEditorCommands::UnfreezeNodes(UE::Dataflow::FContext& Context, con
 
 void FDataflowEditorCommands::DeleteNodes(UEdGraph* EdGraph, const FGraphPanelSelectionSet& SelectedNodes)
 {
-	if (ensureMsgf(EdGraph != nullptr, TEXT("Warning : Failed to find valid graph.")))
+	TArray<UEdGraphNode*> NodesToDelete;
+	NodesToDelete.Reserve(SelectedNodes.Num());
+	for (UObject* Node : SelectedNodes)
 	{
-		for (UObject* Node : SelectedNodes)
+		if (UEdGraphNode* EdNode = Cast<UEdGraphNode>(Node))
 		{
-			if (UDataflowEdNode* EdNode = Cast<UDataflowEdNode>(Node))
-			{
-				if (const TSharedPtr<UE::Dataflow::FGraph> DataflowGraph = EdNode->GetDataflowGraph())
-				{
-					EdGraph->RemoveNode(EdNode);
-					if (TSharedPtr<FDataflowNode> DataflowNode = DataflowGraph->FindBaseNode(EdNode->GetDataflowNodeGuid()))
-					{
-						DataflowGraph->RemoveNode(DataflowNode);
-					}
-				}
-			}
-			else if (UEdGraphNode_Comment* CommentNode = Cast<UEdGraphNode_Comment>(Node))
-			{
-				EdGraph->RemoveNode(CommentNode);
-			}
-
-			// Auto-rename node so that its current name is made available until it is garbage collected
-			Node->Rename();
+			NodesToDelete.Add(EdNode);
 		}
 	}
+	UE::Dataflow::FEditAssetUtils::DeleteNodes(EdGraph, NodesToDelete);
 }
 
 void FDataflowEditorCommands::RenameNode(const TSharedPtr<SDataflowGraphEditor>& DataflowGraphEditor, UEdGraphNode* EdNode)
diff --git a/Engine/Plugins/Experimental/Dataflow/Source/DataflowEditor/Private/Dataflow/DataflowGraphEditor.cpp b/Engine/Plugins/Experimental/Dataflow/Source/DataflowEditor/Private/Dataflow/DataflowGraphEditor.cpp
--- a/Engine/Plugins/Experimental/Dataflow/Source/DataflowEditor/Private/Dataflow/DataflowGraphEditor.cpp
+++ b/Engine/Plugins/Experimental/Dataflow/Source/DataflowEditor/Private/Dataflow/DataflowGraphEditor.cpp
@@ -286,22 +286,8 @@ void SDataflowGraphEditor::DeleteNode()
 		const FGraphPanelSelectionSet& SelectedNodes = GetSelectedNodes();
 		if (SelectedNodes.Num() > 0)
 		{
-			if (UDataflow* Dataflow = UDataflow::GetDataflowAssetFromEdGraph(EdGraph.Get()))
-			{
-				const FScopedTransaction Transaction(LOCTEXT("DeleteSelectedNodes", "Delete selected nodes"));
-
-				Dataflow->Modify();
-				EdGraph->Modify();
-
-				for (FGraphPanelSelectionSet::TConstIterator It(SelectedNodes); It; ++It)
-				{
-					(*It)->Modify();
-				}
-
-				FDataflowEditorCommands::DeleteNodes(EdGraph.Get(), SelectedNodes);
-
-				OnNodeDeletedMulticast.Broadcast(SelectedNodes);
-			}
+			FDataflowEditorCommands::DeleteNodes(EdGraph.Get(), SelectedNodes);
+			OnNodeDeletedMulticast.Broadcast(SelectedNodes);
 		}
 	}
 }
diff --git a/Engine/Source/Runtime/Experimental/Dataflow/Engine/Public/Dataflow/DataflowSubGraphNodes.h b/Engine/Source/Runtime/Experimental/Dataflow/Engine/Public/Dataflow/DataflowSubGraphNodes.h
--- a/Engine/Source/Runtime/Experimental/Dataflow/Engine/Public/Dataflow/DataflowSubGraphNodes.h
+++ b/Engine/Source/Runtime/Experimental/Dataflow/Engine/Public/Dataflow/DataflowSubGraphNodes.h
@@ -147,6 +147,7 @@ public:
 	DATAFLOWENGINE_API virtual ~FDataflowCallSubGraphNode() override;
 
 	DATAFLOWENGINE_API void SetSubGraphGuid(const FGuid& Guid);
+	DATAFLOWENGINE_API const FGuid& GetSubGraphGuid() const { return SubGraphGuid; }
 
 private:
 	virtual void Evaluate(UE::Dataflow::FContext& Context, const FDataflowOutput* Out) const override;
