Dataflow : add subgraph and variables to the datafloiw graph context menu

diff --git a/Engine/Plugins/Experimental/Dataflow/Source/DataflowEditor/Private/Dataflow/DataflowGraphSchemaAction.cpp b/Engine/Plugins/Experimental/Dataflow/Source/DataflowEditor/Private/Dataflow/DataflowGraphSchemaAction.cpp
--- a/Engine/Plugins/Experimental/Dataflow/Source/DataflowEditor/Private/Dataflow/DataflowGraphSchemaAction.cpp
+++ b/Engine/Plugins/Experimental/Dataflow/Source/DataflowEditor/Private/Dataflow/DataflowGraphSchemaAction.cpp
@@ -74,8 +74,8 @@ FEdGraphSchemaAction_DataflowVariable::FEdGraphSchemaAction_DataflowVariable()
 FEdGraphSchemaAction_DataflowVariable::FEdGraphSchemaAction_DataflowVariable(UDataflow* InDataflowAsset, const FPropertyBagPropertyDesc& PropertyDesc)
 	: FEdGraphSchemaAction(
 		FText::FromString(CategoryFromFullName(PropertyDesc.Name)),
-		LOCTEXT("Description", "Description"),
-		LOCTEXT("Tooltip", "Tooltip"),
+		FText::Format(LOCTEXT("DataflowVariableActionDescription", "Variable {0}"), FText::FromName(PropertyDesc.Name)),
+		FText::Format(LOCTEXT("DataflowVariableActionTooltip", "Variable {0}"), FText::FromName(PropertyDesc.Name)),
 		/*InGrouping*/0,
 		/*Keywords*/FText::GetEmpty(),
 		(int32)UE::Dataflow::ESchemaActionSectionID::VARIABLES)
@@ -229,15 +229,24 @@ FEdGraphSchemaAction_DataflowSubGraph::FEdGraphSchemaAction_DataflowSubGraph()
 
 FEdGraphSchemaAction_DataflowSubGraph::FEdGraphSchemaAction_DataflowSubGraph(UDataflow* InDataflowAsset, const FGuid& InSubGraphGuid)
 	: FEdGraphSchemaAction(
-		FText::GetEmpty(),
-		LOCTEXT("Description", "Description"),
-		LOCTEXT("Tooltip", "Tooltip"),
+		FText::GetEmpty(), // menu Category : defined in the ctor body
+		FText::GetEmpty(), // menu description : defined in the ctor body
+		FText::GetEmpty(), // menu tooltip : defined in the ctor body
 		/*InGrouping*/0,
 		/*Keywords*/FText::GetEmpty(),
 		(int32)UE::Dataflow::ESchemaActionSectionID::SUBGRAPHS)
 	, SubGraphGuid(InSubGraphGuid)
 	, DataflowAssetWeakPtr(InDataflowAsset)
 {
+	if (TStrongObjectPtr<const UDataflow> DataflowAsset = DataflowAssetWeakPtr.Pin())
+	{
+		if (const UDataflowSubGraph* SubGraph = DataflowAsset->FindSubGraphByGuid(SubGraphGuid))
+		{
+			const FText NewMenuDescription = FText::Format(LOCTEXT("DataflowSubGraphActionDescription", "Function {0}"), FText::FromName(SubGraph->GetFName()));
+			const FText NewToolTipDescription = FText::Format(LOCTEXT("DataflowSubGraphActionTooltip", "Function {0}"), FText::FromName(SubGraph->GetFName()));
+			UpdateSearchData(NewMenuDescription, NewToolTipDescription, /*NewCategory*/FText::GetEmpty(), /*NewKeywords*/FText::GetEmpty());
+		}
+	}
 }
 
 const FName FEdGraphSchemaAction_DataflowSubGraph::GetSubGraphName() const
diff --git a/Engine/Plugins/Experimental/Dataflow/Source/DataflowEditor/Private/Dataflow/DataflowSchema.cpp b/Engine/Plugins/Experimental/Dataflow/Source/DataflowEditor/Private/Dataflow/DataflowSchema.cpp
--- a/Engine/Plugins/Experimental/Dataflow/Source/DataflowEditor/Private/Dataflow/DataflowSchema.cpp
+++ b/Engine/Plugins/Experimental/Dataflow/Source/DataflowEditor/Private/Dataflow/DataflowSchema.cpp
@@ -16,6 +16,8 @@
 #include "ToolMenu.h"
 #include "GraphEditorActions.h"
 #include "Dataflow/DataflowSettings.h"
+#include "Dataflow/DataflowObject.h"
+#include "Dataflow/DataflowSubGraph.h"
 #include "ScopedTransaction.h"
 
 #include UE_INLINE_GENERATED_CPP_BY_NAME(DataflowSchema)
@@ -147,6 +149,42 @@ bool UDataflowSchema::CanPinBeConnectedToNode(const UEdGraphPin* Pin, const UE::
 
 void UDataflowSchema::GetGraphContextActions(FGraphContextMenuBuilder& ContextMenuBuilder) const
 {
+	if (UDataflow* DataflowAsset = UDataflow::GetDataflowAssetFromEdGraph(const_cast<UEdGraph*>(ContextMenuBuilder.CurrentGraph)))
+	{
+		const bool bDisplayVariable = (ContextMenuBuilder.FromPin == nullptr || ContextMenuBuilder.FromPin->Direction == EEdGraphPinDirection::EGPD_Input);
+		if (bDisplayVariable)
+		{
+			// Variables
+			const FText VariablesCategory = LOCTEXT("DataflowContextActionVariablesCategory", "Variables");
+			if (const UPropertyBag* PropertyBag = DataflowAsset->Variables.GetPropertyBagStruct())
+			{
+				for (const FPropertyBagPropertyDesc& PropertyDesc : PropertyBag->GetPropertyDescs())
+				{
+					// Todo : check if the type of the variable is compatible with the FromPin types
+					TSharedPtr<FEdGraphSchemaAction_DataflowVariable> VariableAction =
+						MakeShareable(new FEdGraphSchemaAction_DataflowVariable(DataflowAsset, PropertyDesc));
+					if (VariableAction)
+					{
+						VariableAction->CosmeticUpdateRootCategory(VariablesCategory);
+						ContextMenuBuilder.AddAction(VariableAction);
+					}
+				}
+			}
+		}
+		// SubGraph/functions
+		const FText SubGraphsCategory = LOCTEXT("DataflowContextActionSubGraphsCategory", "SubGraphs");
+		for (UDataflowSubGraph* SubGraph : DataflowAsset->GetSubGraphs())
+		{
+			// Todo : check if the input / output node of the subgraph are compatible with the FromPin types
+			TSharedPtr<FEdGraphSchemaAction_DataflowSubGraph> SubGraphAction =
+				MakeShareable(new FEdGraphSchemaAction_DataflowSubGraph(DataflowAsset, SubGraph->GetSubGraphGuid()));
+			if (SubGraphAction)
+			{
+				SubGraphAction->CosmeticUpdateRootCategory(SubGraphsCategory);
+				ContextMenuBuilder.AddAction(SubGraphAction);
+			}
+		}
+	}
 	if (UE::Dataflow::FNodeFactory* Factory = UE::Dataflow::FNodeFactory::GetInstance())
 	{
 		for (UE::Dataflow::FFactoryParameters NodeParameters : Factory->RegisteredParameters())
