Dataflow : Double click on a subgraph call node now opens the subgraph tab

diff --git a/Engine/Plugins/Experimental/Dataflow/Source/DataflowEditor/Private/Dataflow/DataflowEditorToolkit.cpp b/Engine/Plugins/Experimental/Dataflow/Source/DataflowEditor/Private/Dataflow/DataflowEditorToolkit.cpp
--- a/Engine/Plugins/Experimental/Dataflow/Source/DataflowEditor/Private/Dataflow/DataflowEditorToolkit.cpp
+++ b/Engine/Plugins/Experimental/Dataflow/Source/DataflowEditor/Private/Dataflow/DataflowEditorToolkit.cpp
@@ -40,6 +40,7 @@
 #include "Dataflow/DataflowSchema.h"
 #include "Dataflow/DataflowSkeletonView.h"
 #include "Dataflow/DataflowSimulationViewportClient.h"
+#include "Dataflow/DataflowSubGraphNodes.h"
 #include "Dataflow/DataflowMembersWidget.h"
 #include "DynamicMeshBuilder.h"
 #include "EditorModeManager.h"
@@ -797,7 +798,23 @@ void FDataflowEditorToolkit::OnNodeTitleCommitted(const FText& InNewText, ETextC
 	FDataflowEditorCommands::OnNodeTitleCommitted(InNewText, InCommitType, GraphNode);
 }
 
-
+void FDataflowEditorToolkit::OnNodeDoubleClicked(UEdGraphNode* ClickedNode)
+{
+	// if the node is a call to subgraph , open the subgraph tab 
+	if (UDataflowEdNode* DataflowEdNode = Cast<UDataflowEdNode>(ClickedNode))
+	{
+		if (TSharedPtr<FDataflowNode> DataflowNode = DataflowEdNode->GetDataflowNode())
+		{
+			if (FDataflowCallSubGraphNode* CallSubGraphNode = DataflowNode->AsType<FDataflowCallSubGraphNode>())
+			{
+				if (UDataflowSubGraph* Subgraph = GetSubGraph(CallSubGraphNode->GetSubGraphGuid()))
+				{
+					OpenSubGraphTab(Subgraph);
+				}
+			}
+		}
+	}
+}
 
 void FDataflowEditorToolkit::OnNodeSelectionChanged(const TSet<UObject*>& InNewSelection)
 {
@@ -1238,6 +1255,8 @@ TSharedRef<SDataflowGraphEditor> FDataflowEditorToolkit::CreateGraphEditorWidget
 	InEvents.OnVerifyTextCommit = FOnNodeVerifyTextCommit::CreateSP(this, &FDataflowEditorToolkit::OnNodeVerifyTitleCommit);
 	InEvents.OnTextCommitted = FOnNodeTextCommitted::CreateSP(this, &FDataflowEditorToolkit::OnNodeTitleCommitted);
 	InEvents.OnNodeSingleClicked = SGraphEditor::FOnNodeSingleClicked::CreateSP(this, &FDataflowEditorToolkit::OnNodeSingleClicked);
+	InEvents.OnNodeDoubleClicked = FSingleNodeEvent::CreateSP(this, &FDataflowEditorToolkit::OnNodeDoubleClicked);
+
 
 	UDataflow* DataflowAsset = UDataflow::GetDataflowAssetFromEdGraph(GraphToEdit);
 	check(DataflowAsset);
@@ -1824,6 +1843,19 @@ UDataflowSubGraph* FDataflowEditorToolkit::GetSubGraph(FName SubGraphName) const
 	return nullptr;
 }
 
+UDataflowSubGraph* FDataflowEditorToolkit::GetSubGraph(const FGuid& SubGraphGuid) const
+{
+	if (const TObjectPtr<UDataflowBaseContent>& EditorContent = GetEditorContent())
+	{
+		ensure(EditorContent);
+		if (UDataflow* DataflowAsset = EditorContent->GetDataflowAsset())
+		{
+			return DataflowAsset->FindSubGraphByGuid(SubGraphGuid);
+		}
+	}
+	return nullptr;
+}
+
 void FDataflowEditorToolkit::OpenSubGraphTab(FName SubGraphName)
 {
 	OpenSubGraphTab(GetSubGraph(SubGraphName));
diff --git a/Engine/Plugins/Experimental/Dataflow/Source/DataflowEditor/Public/Dataflow/DataflowEditorToolkit.h b/Engine/Plugins/Experimental/Dataflow/Source/DataflowEditor/Public/Dataflow/DataflowEditorToolkit.h
--- a/Engine/Plugins/Experimental/Dataflow/Source/DataflowEditor/Public/Dataflow/DataflowEditorToolkit.h
+++ b/Engine/Plugins/Experimental/Dataflow/Source/DataflowEditor/Public/Dataflow/DataflowEditorToolkit.h
@@ -91,6 +91,7 @@ public:
 	void CloseSubGraphTab(const UDataflowSubGraph* SubGraph);
 	void ReOpenSubGraphTab(const UDataflowSubGraph* SubGraph);
 	void SetSubGraphTabActiveState(TSharedPtr<SDataflowGraphEditor> SubGraphEditor, bool bActive);
+	UDataflowSubGraph* GetSubGraph(const FGuid& SubGraphGuid) const;
 	UDataflowSubGraph* GetSubGraph(FName SubGraphName) const;
 
 protected:
@@ -101,6 +102,7 @@ protected:
 	void OnPropertyValueChanged(const FPropertyChangedEvent& PropertyChangedEvent);
 	bool OnNodeVerifyTitleCommit(const FText& NewText, UEdGraphNode* GraphNode, FText& OutErrorMessage);
 	void OnNodeTitleCommitted(const FText& InNewText, ETextCommit::Type InCommitType, UEdGraphNode* GraphNode);
+	void OnNodeDoubleClicked(UEdGraphNode* ClickedNode);
 	void OnNodeSelectionChanged(const TSet<UObject*>& NewSelection);
 	void OnNodeDeleted(const TSet<UObject*>& NewSelection);
 	void OnNodeSingleClicked(UObject* ClickedNode) const;
