Enabling SubGraph  in Dataflow
( Keeping the cvar for now - will remove later )

diff --git a/Engine/Plugins/Experimental/Dataflow/Source/DataflowEditor/Private/Dataflow/DataflowMembersWidget.cpp b/Engine/Plugins/Experimental/Dataflow/Source/DataflowEditor/Private/Dataflow/DataflowMembersWidget.cpp
--- a/Engine/Plugins/Experimental/Dataflow/Source/DataflowEditor/Private/Dataflow/DataflowMembersWidget.cpp
+++ b/Engine/Plugins/Experimental/Dataflow/Source/DataflowEditor/Private/Dataflow/DataflowMembersWidget.cpp
@@ -31,7 +31,7 @@
 /////////////////////////////////////////////////////////////////////////////////////////////////////////////////////////
 namespace UE::Dataflow::DataflowMembersWidget::Private
 {
-	bool bEnableSubGraphs = false;
+	bool bEnableSubGraphs = true;
 	FAutoConsoleVariableRef CVarEnableSubGraphs(
 		TEXT("p.Dataflow.Editor.EnableSubgraphs"),
 		bEnableSubGraphs,
