Fix subgraph list not refreshing in Dataflow editor bound to an asset

diff --git a/Engine/Plugins/Experimental/Dataflow/Source/DataflowEditor/Private/Dataflow/DataflowMembersWidget.cpp b/Engine/Plugins/Experimental/Dataflow/Source/DataflowEditor/Private/Dataflow/DataflowMembersWidget.cpp
--- a/Engine/Plugins/Experimental/Dataflow/Source/DataflowEditor/Private/Dataflow/DataflowMembersWidget.cpp
+++ b/Engine/Plugins/Experimental/Dataflow/Source/DataflowEditor/Private/Dataflow/DataflowMembersWidget.cpp
@@ -531,7 +531,7 @@ void SDataflowMembersWidget::OnObjectPropertyChanged(UObject* InObject, FPropert
 
 void SDataflowMembersWidget::OnSubGraphsChanged(const UDataflow* InDataflowAsset, const FGuid& InSubGraphGuid, UE::Dataflow::ESubGraphChangedReason InReason)
 {
-	if (InDataflowAsset == EditedAssetWeakPtr)
+	if (InDataflowAsset == DataflowAssetWeakPtr)
 	{
 		Refresh();
 	}
