Dataflow : fix comment so they can be placed when no node are selected

diff --git a/Engine/Plugins/Experimental/Dataflow/Source/DataflowEditor/Private/Dataflow/DataflowSCommentNode.cpp b/Engine/Plugins/Experimental/Dataflow/Source/DataflowEditor/Private/Dataflow/DataflowSCommentNode.cpp
--- a/Engine/Plugins/Experimental/Dataflow/Source/DataflowEditor/Private/Dataflow/DataflowSCommentNode.cpp
+++ b/Engine/Plugins/Experimental/Dataflow/Source/DataflowEditor/Private/Dataflow/DataflowSCommentNode.cpp
@@ -28,12 +28,20 @@ TSharedPtr<FAssetSchemaAction_Dataflow_CreateCommentNode_DataflowEdNode> FAssetS
 //
 UEdGraphNode* FAssetSchemaAction_Dataflow_CreateCommentNode_DataflowEdNode::PerformAction(class UEdGraph* ParentGraph, UEdGraphPin* FromPin, const FVector2D Location, bool bSelectNewNode)
 {
-	FSlateRect Bounds;
-	if (GraphEditor->GetBoundsForSelectedNodes(Bounds, 50.0f))
+	constexpr float MinSizeX = 500.f;
+	constexpr float MinSizeY = 250.f;
+
+	// set default bounds 
+	FSlateRect Bounds
 	{
-		return UE::Dataflow::FEditAssetUtils::AddNewComment(ParentGraph, Bounds.GetTopLeft(), Bounds.GetSize());
-	}
-	return nullptr;
+		(float)Location.X,
+		(float)Location.Y,
+		(float)Location.X + MinSizeX,
+		(float)Location.Y + MinSizeY
+	};
+	GraphEditor->GetBoundsForSelectedNodes(Bounds, 50.0f);
+
+	return UE::Dataflow::FEditAssetUtils::AddNewComment(ParentGraph, Bounds.GetTopLeft(), Bounds.GetSize());
 }
 
 #undef LOCTEXT_NAMESPACE
