Make the Dataflow plugin dependency on MeshModelingToolsetExp Editor-only

diff --git a/Engine/Plugins/Experimental/Dataflow/Dataflow.uplugin b/Engine/Plugins/Experimental/Dataflow/Dataflow.uplugin
--- a/Engine/Plugins/Experimental/Dataflow/Dataflow.uplugin
+++ b/Engine/Plugins/Experimental/Dataflow/Dataflow.uplugin
@@ -52,7 +52,8 @@
 		},
 		{
 			"Name": "MeshModelingToolsetExp",
-			"Enabled": true
+			"Enabled": true,
+			"TargetAllowList": [ "Editor" ]
 		},
 		{
 			"Name": "ChaosCaching",
