Dataflow:
- Added OutputLog window
- Feedback in OutputLog for evaluation of graph
- Added FContext::Info(), Warning() and Error() to set warnings and errors during node evaluation
- Added SetError() to set match SetValue()
- Sets warning if a node's DATAFLOW_INTRINSIC input is not connected
- Double click on TokenizedMessage in OutputLog selects the node and homes on it
- After graph evaluation it displays evaluation time in OutputLog
- Bug fixes

diff --git a/Engine/Plugins/Experimental/Dataflow/Source/DataflowEditor/DataflowEditor.Build.cs b/Engine/Plugins/Experimental/Dataflow/Source/DataflowEditor/DataflowEditor.Build.cs
--- a/Engine/Plugins/Experimental/Dataflow/Source/DataflowEditor/DataflowEditor.Build.cs
+++ b/Engine/Plugins/Experimental/Dataflow/Source/DataflowEditor/DataflowEditor.Build.cs
@@ -74,7 +74,8 @@ namespace UnrealBuildTool.Rules
 					"KismetWidgets",      // SScrubControlPanel
 					"AnimGraph", 
 					"ChaosCaching", // UAnimSingleNodeInstance
-					"StructUtilsEditor"
+					"StructUtilsEditor",
+					"MessageLog"
 				}
 			);
 		}
diff --git a/Engine/Plugins/Experimental/Dataflow/Source/DataflowEditor/Private/Dataflow/DataflowEditorToolkit.cpp b/Engine/Plugins/Experimental/Dataflow/Source/DataflowEditor/Private/Dataflow/DataflowEditorToolkit.cpp
--- a/Engine/Plugins/Experimental/Dataflow/Source/DataflowEditor/Private/Dataflow/DataflowEditorToolkit.cpp
+++ b/Engine/Plugins/Experimental/Dataflow/Source/DataflowEditor/Private/Dataflow/DataflowEditorToolkit.cpp
@@ -66,6 +66,11 @@
 #include "GeometryCache.h"
 #include "ToolMenus.h"
 #include "WorkflowOrientedApp/WorkflowTabManager.h"
+#include "Dataflow/DataflowNodeParameters.h"
+#include "Framework/Notifications/NotificationManager.h"
+#include "Widgets/Notifications/SNotificationList.h"
+#include "Dataflow/DataflowOutputLog.h"
+#include "Dataflow/DataflowPath.h"
 
 #define LOCTEXT_NAMESPACE "DataflowEditorToolkit"
 
@@ -91,6 +96,7 @@ const FName FDataflowEditorToolkit::CollectionSpreadSheetTabId_4(TEXT("DataflowE
 const FName FDataflowEditorToolkit::SimulationViewportTabId(TEXT("DataflowEditor_SimulationViewport"));
 const FName FDataflowEditorToolkit::SimulationVisualizationTabId(TEXT("DataflowEditor_SimulationVisualizationTab"));
 const FName FDataflowEditorToolkit::MembersWidgetTabId(TEXT("DataflowEditor_MembersWidgetTab"));
+const FName FDataflowEditorToolkit::OutputLogTabId(TEXT("DataflowEditor_OutputLog"));
 
 FDataflowEditorToolkit::FDataflowEditorToolkit(UAssetEditor* InOwningAssetEditor)
 	: FBaseCharacterFXEditorToolkit(InOwningAssetEditor, FName("DataflowEditor"))
@@ -279,7 +285,7 @@ FDataflowEditorToolkit::FDataflowEditorToolkit(UAssetEditor* InOwningAssetEditor
 
 	if(TObjectPtr<UDataflowBaseContent>& EditorContent = DataflowEditor->GetEditorContent())
 	{
-		if(EditorContent->GetDataflowAsset() && EditorContent->GetDataflowAsset()->Type == EDataflowType::Simulation)
+		if (EditorContent->GetDataflowAsset() && EditorContent->GetDataflowAsset()->Type == EDataflowType::Simulation)
 		{
 			StandaloneDefaultLayout = SimulationDefaultLayout;
 			bForceViewportTab = false;
@@ -341,6 +347,19 @@ FDataflowEditorToolkit::~FDataflowEditorToolkit()
 		GraphEditor->OnNodeDeletedMulticast.Remove(OnNodeDeletedMulticastDelegateHandle);
 	}
 
+	const TObjectPtr<UDataflowBaseContent>& EditorContent = GetEditorContent();
+	check(EditorContent);
+	
+	if (TSharedPtr<UE::Dataflow::FEngineContext> DataflowContext = EditorContent->GetDataflowContext())
+	{
+		DataflowContext->OnNodeBeginEvaluateMulticast.Remove(OnNodeBeginEvaluateMulticastDelegateHandle);
+		DataflowContext->OnNodeFinishEvaluateMulticast.Remove(OnNodeFinishEvaluateMulticastDelegateHandle);
+
+		DataflowContext->OnContextHasInfo.Remove(OnContextHasInfoDelegateHandle);
+		DataflowContext->OnContextHasWarning.Remove(OnContextHasWarningDelegateHandle);
+		DataflowContext->OnContextHasError.Remove(OnContextHasErrorDelegateHandle);
+	}
+
 	if (NodeDetailsEditor)
 	{
 		NodeDetailsEditor->GetOnFinishedChangingPropertiesDelegate().Remove(OnFinishedChangingPropertiesDelegateHandle);
@@ -351,6 +370,11 @@ FDataflowEditorToolkit::~FDataflowEditorToolkit()
 		AssetDetailsEditor->OnFinishedChangingProperties().Remove(OnFinishedChangingAssetPropertiesDelegateHandle);
 	}
 
+	if (DataflowOutputLog)
+	{
+		DataflowOutputLog->GetOnOutputLogMessageTokenClickedDelegate().Remove(OnOutputLogMessageTokenClickedDelegateHandle);
+	}
+
 	// We need to force the dataflow editor mode deletion now because otherwise the preview and rest-space worlds
 	// will end up getting destroyed before the mode's Exit() function gets to run, and we'll get some
 	// warnings when we destroy any mode actors.
@@ -406,6 +430,28 @@ const UDataflow* FDataflowEditorToolkit::GetDataflowAsset(const UObject* ObjectT
 	return UE::Dataflow::InstanceUtils::GetDataflowAssetFromObject(ObjectToEdit);
 }
 
+bool FDataflowEditorToolkit::IsSimulationDataflowAsset() const
+{
+	if (const TObjectPtr<UDataflowBaseContent>& EditorContent = GetEditorContent())
+	{
+		return (EditorContent->GetDataflowAsset() && EditorContent->GetDataflowAsset()->Type == EDataflowType::Simulation);
+	}
+	return false;
+}
+
+FName FDataflowEditorToolkit::GetGraphLogName() const
+{
+	static FName CSimulationGraphName = TEXT("Simulation");
+	static FName CConstructionGraphName = TEXT("Construction");
+
+	if (IsSimulationDataflowAsset())
+	{
+		return CSimulationGraphName;
+	}
+
+	return CConstructionGraphName;
+}
+
 //~ Begin FBaseCharacterFXEditorToolkit overrides
 
 FEditorModeID FDataflowEditorToolkit::GetEditorModeId() const
@@ -762,6 +808,23 @@ void FDataflowEditorToolkit::CreateSimulationViewportClient()
 
 //~ End FBaseAssetToolkit overrides
 
+namespace UE::Dataflow::Private
+{
+	static void ShowNotificationMessage(const FText& Message, const SNotificationItem::ECompletionState CompletionState)
+	{
+		FNotificationInfo Info(Message);
+		Info.ExpireDuration = 5.0f;
+		Info.bUseLargeFont = false;
+		Info.bUseThrobber = false;
+		Info.bUseSuccessFailIcons = false;
+		TSharedPtr<SNotificationItem> Notification = FSlateNotificationManager::Get().AddNotification(Info);
+		if (Notification.IsValid())
+		{
+			Notification->SetCompletionState(CompletionState);
+		}
+	}
+}
+
 void FDataflowEditorToolkit::OnPropertyValueChanged(const FPropertyChangedEvent& PropertyChangedEvent)
 {
 	if (const TObjectPtr<UDataflowBaseContent>& EditorContent = GetEditorContent())
@@ -776,6 +839,9 @@ void FDataflowEditorToolkit::OnPropertyValueChanged(const FPropertyChangedEvent&
 
 			EditorContent->SetDataflowContext(DataflowContext);
 			EditorContent->SetLastModifiedTimestamp(LastNodeTimestamp);
+
+			// Refresh graph display to display nodes with warning/error
+			GraphEditor->NotifyGraphChanged();
 		}
 	}
 }
@@ -1127,12 +1193,154 @@ void FDataflowEditorToolkit::OnConstructionViewSelectionChanged(const TArray<UPr
 	}
 }
 
+void FDataflowEditorToolkit::OnNodeBeginEvaluate(const FDataflowNode* Node, const FDataflowOutput* Output)
+{
+	if (const TObjectPtr<UDataflowBaseContent>& EditorContent = GetEditorContent())
+	{
+		ensure(EditorContent);
+		if (TSharedPtr<UE::Dataflow::FEngineContext> DataflowContext = EditorContent->GetDataflowContext())
+		{
+			DataflowContext->Info(TEXT("Begin Evaluate"), Node, Output);
+		}
+	}
+}
+
+void FDataflowEditorToolkit::OnNodeFinishEvaluate(const FDataflowNode* Node, const FDataflowOutput* Output)
+{
+	if (const TObjectPtr<UDataflowBaseContent>& EditorContent = GetEditorContent())
+	{
+		ensure(EditorContent);
+		if (TSharedPtr<UE::Dataflow::FEngineContext> DataflowContext = EditorContent->GetDataflowContext())
+		{
+			DataflowContext->Info(TEXT("End Evaluate"), Node, Output);
+		}
+	}
+}
+
+void FDataflowEditorToolkit::LogMessage(const EMessageSeverity::Type Severity, const FDataflowNode* Node, const FDataflowOutput* Output, const FString& Message) const
+{
+	if (DataflowOutputLog)
+	{
+		const FString NodeName = Node ? Node->GetName().ToString() : FString();
+		const FString OutputName = Output ? Output->GetName().ToString() : FString();
+
+		FDataflowPath Path(GetGraphLogName().ToString(), NodeName, FString(), OutputName);
+		DataflowOutputLog->AddMessage(Severity, Message, Path);
+	}
+}
+
+void FDataflowEditorToolkit::OnContextHasInfo(const FDataflowNode* Node, const FDataflowOutput* Output, const FString& Info)
+{
+	LogMessage(EMessageSeverity::Info, Node, Output, Info);
+}
+
+void FDataflowEditorToolkit::OnContextHasWarning(const FDataflowNode* Node, const FDataflowOutput* Output, const FString& Warning)
+{
+	LogMessage(EMessageSeverity::Warning, Node, Output, Warning);
+}
+
+void FDataflowEditorToolkit::OnContextHasError(const FDataflowNode* Node, const FDataflowOutput* Output, const FString& Error)
+{
+	LogMessage(EMessageSeverity::Error, Node, Output, Error);
+}
+
+void FDataflowEditorToolkit::OnBeginEvaluate()
+{
+	GraphEvaluationBegin = FDateTime::Now();
+
+	if (DataflowOutputLog)
+	{
+		DataflowOutputLog->ClearMessageLog();
+	}
+
+	if (const TObjectPtr<UDataflowBaseContent>& EditorContent = GetEditorContent())
+	{
+		ensure(EditorContent);
+		if (TSharedPtr<UE::Dataflow::FEngineContext> DataflowContext = EditorContent->GetDataflowContext())
+		{
+			DataflowContext->ClearNodesData();
+		}
+	}
+
+	// Refresh graph display to display nodes with warning/error
+	GraphEditor->NotifyGraphChanged();
+}
+
+namespace UE::Dataflow::Private
+{
+	static int32 GetMilliseconds(const FTimespan& InTimespan)
+	{
+		const int64 Ticks = InTimespan.GetTicks();
+		return (int32)((Ticks / ETimespan::TicksPerMillisecond) % 1000);
+	}
+
+	static FString GetElapsedTimeString(const FDateTime& InGraphEvaluationFinished, const FDateTime& InGraphEvaluationBegin)
+	{
+		const FTimespan ElapsedTimespan = InGraphEvaluationFinished - InGraphEvaluationBegin;
+		const FString ElapsedTimeString = FString::Printf(TEXT("%02dm%02ds%02dms"), ElapsedTimespan.GetMinutes(), ElapsedTimespan.GetSeconds(), GetMilliseconds(ElapsedTimespan));
+
+		return ElapsedTimeString;
+	}
+}
+
 void FDataflowEditorToolkit::OnFinishEvaluate()
 {
+	// Display message stating that evaluation finished
+	const TObjectPtr<UDataflowBaseContent>& EditorContent = GetEditorContent();
+	ensure(EditorContent);
+
+	if (TSharedPtr<UE::Dataflow::FEngineContext> DataflowContext = EditorContent->GetDataflowContext())
+	{
+		const int32 NumWarnings = DataflowContext->GetNumWarnings();
+		const int32 NumErrors = DataflowContext->GetNumErrors();
+
+		if (NumWarnings > 0 || NumErrors > 0)
+		{
+			const FText MessageFormat = LOCTEXT("OnFinishEvaluate", "Finished graph evaluation\nWarning(s): {0} Error(s): {1}");
+			UE::Dataflow::Private::ShowNotificationMessage(FText::Format(MessageFormat, NumWarnings, NumErrors), SNotificationItem::CS_Fail);
+		}
+
+		GraphEvaluationFinished = FDateTime::Now();
+
+		const FString ElapsedTimeString = UE::Dataflow::Private::GetElapsedTimeString(GraphEvaluationFinished, GraphEvaluationBegin);
+		DataflowContext->Info(TEXT("Evaluation time: ") + ElapsedTimeString);
+	}
+
 	// Refresh graph display to update node output pin display (invalid or valid)
 	GraphEditor->NotifyGraphChanged();
 }
 
+void FDataflowEditorToolkit::OnOutputLogMessageTokenClicked(const FString TokenString)
+{
+	FDataflowPath Path;
+	Path.DecodePath(TokenString);
+	const FString NodeName = Path.GetNode();
+
+	FString Output;
+	if (!NodeName.IsEmpty() && Path.PathHasOutput())
+	{
+		Output = Path.GetOutput();
+	}
+
+	// Clear node selection
+	GraphEditor->ClearSelectionSet();
+
+	// Select node from TokenString
+	if (UDataflow* DataflowAsset = GetEditorContent()->GetDataflowAsset())
+	{
+		if (TSharedPtr<UE::Dataflow::FGraph> DataflowGraph = DataflowAsset->GetDataflow())
+		{
+			if (TSharedPtr<FDataflowNode> DataflowNode = DataflowGraph->FindBaseNode(FName(*NodeName)))
+			{
+				if (TObjectPtr<UDataflowEdNode> DataflowUEdNode = DataflowAsset->FindEdNodeByDataflowNodeGuid(DataflowNode->GetGuid()))
+				{
+					GraphEditor->JumpToNode(DataflowUEdNode);
+				}
+			}
+		}
+	}
+}
+
 void FDataflowEditorToolkit::Tick(float DeltaTime)
 {
 	if (const TObjectPtr<UDataflowBaseContent> EditorContent = GetEditorContent())
@@ -1234,6 +1442,11 @@ TSharedRef<SDataflowGraphEditor> FDataflowEditorToolkit::CreateGraphEditorWidget
 	const FDataflowEditorCommands::FGraphEvaluationCallback Evaluate =
 		[this](const FDataflowNode* Node, const FDataflowOutput* Output)
 		{
+			//
+			// Graph evaluation begin
+			//
+			OnBeginEvaluate();
+
 			if (const TObjectPtr<UDataflowBaseContent>& EditorContent = GetEditorContent())
 			{
 				UE::Dataflow::FTimestamp LastNodeTimestamp = EditorContent->GetLastModifiedTimestamp();
@@ -1270,6 +1483,19 @@ TSharedRef<SDataflowGraphEditor> FDataflowEditorToolkit::CreateGraphEditorWidget
 
 	OnSelectionChangedMulticastDelegateHandle = NewGraphEditor->OnSelectionChangedMulticast.AddSP(this, &FDataflowEditorToolkit::OnNodeSelectionChanged);
 	OnNodeDeletedMulticastDelegateHandle = NewGraphEditor->OnNodeDeletedMulticast.AddSP(this, &FDataflowEditorToolkit::OnNodeDeleted);
+	
+	const TObjectPtr<UDataflowBaseContent>& EditorContent = GetEditorContent();
+	check(EditorContent);
+
+	TSharedPtr<FEngineContext>& Context = EditorContent->GetDataflowContext();
+	check(Context);
+
+	OnNodeBeginEvaluateMulticastDelegateHandle = Context->OnNodeBeginEvaluateMulticast.AddSP(this, &FDataflowEditorToolkit::OnNodeBeginEvaluate);
+	OnNodeFinishEvaluateMulticastDelegateHandle = Context->OnNodeFinishEvaluateMulticast.AddSP(this, &FDataflowEditorToolkit::OnNodeFinishEvaluate);
+
+	OnContextHasInfoDelegateHandle = Context->OnContextHasInfo.AddSP(this, &FDataflowEditorToolkit::OnContextHasInfo);
+	OnContextHasWarningDelegateHandle = Context->OnContextHasWarning.AddSP(this, &FDataflowEditorToolkit::OnContextHasWarning);
+	OnContextHasErrorDelegateHandle = Context->OnContextHasError.AddSP(this, &FDataflowEditorToolkit::OnContextHasError);
 
 	return NewGraphEditor;
 }
@@ -1455,7 +1681,7 @@ TSharedRef<SDockTab> FDataflowEditorToolkit::SpawnTab_SelectionView(const FSpawn
 
 	if (Args.GetTabId() == SelectionViewTabId_1)
 	{
-		DataflowSelectionView_1 = MakeShared<FDataflowSelectionView>(FDataflowSelectionView(DataflowEditor->GetEditorContent()));
+		DataflowSelectionView_1 = MakeShared<FDataflowSelectionView>(DataflowEditor->GetEditorContent());
 		if (DataflowSelectionView_1.IsValid())
 		{
 			ViewListeners.Add(DataflowSelectionView_1.Get());
@@ -1463,7 +1689,7 @@ TSharedRef<SDockTab> FDataflowEditorToolkit::SpawnTab_SelectionView(const FSpawn
 	}
 	else if (Args.GetTabId() == SelectionViewTabId_2)
 	{
-		DataflowSelectionView_2 = MakeShared<FDataflowSelectionView>(FDataflowSelectionView(DataflowEditor->GetEditorContent()));
+		DataflowSelectionView_2 = MakeShared<FDataflowSelectionView>(DataflowEditor->GetEditorContent());
 		if (DataflowSelectionView_2.IsValid())
 		{
 			ViewListeners.Add(DataflowSelectionView_2.Get());
@@ -1471,7 +1697,7 @@ TSharedRef<SDockTab> FDataflowEditorToolkit::SpawnTab_SelectionView(const FSpawn
 	}
 	else if (Args.GetTabId() == SelectionViewTabId_3)
 	{
-		DataflowSelectionView_3 = MakeShared<FDataflowSelectionView>(FDataflowSelectionView(DataflowEditor->GetEditorContent()));
+		DataflowSelectionView_3 = MakeShared<FDataflowSelectionView>(DataflowEditor->GetEditorContent());
 		if (DataflowSelectionView_3.IsValid())
 		{
 			ViewListeners.Add(DataflowSelectionView_3.Get());
@@ -1479,7 +1705,7 @@ TSharedRef<SDockTab> FDataflowEditorToolkit::SpawnTab_SelectionView(const FSpawn
 	}
 	else if (Args.GetTabId() == SelectionViewTabId_4)
 	{
-		DataflowSelectionView_4 = MakeShared<FDataflowSelectionView>(FDataflowSelectionView(DataflowEditor->GetEditorContent()));
+		DataflowSelectionView_4 = MakeShared<FDataflowSelectionView>(DataflowEditor->GetEditorContent());
 		if (DataflowSelectionView_4.IsValid())
 		{
 			ViewListeners.Add(DataflowSelectionView_4.Get());
@@ -1528,7 +1754,7 @@ TSharedRef<SDockTab> FDataflowEditorToolkit::SpawnTab_CollectionSpreadSheet(cons
 
 	if (Args.GetTabId() == CollectionSpreadSheetTabId_1)
 	{
-		DataflowCollectionSpreadSheet_1 = MakeShared<FDataflowCollectionSpreadSheet>(FDataflowCollectionSpreadSheet(DataflowEditor->GetEditorContent()));
+		DataflowCollectionSpreadSheet_1 = MakeShared<FDataflowCollectionSpreadSheet>(DataflowEditor->GetEditorContent());
 		if (DataflowCollectionSpreadSheet_1.IsValid())
 		{
 			ViewListeners.Add(DataflowCollectionSpreadSheet_1.Get());
@@ -1536,7 +1762,7 @@ TSharedRef<SDockTab> FDataflowEditorToolkit::SpawnTab_CollectionSpreadSheet(cons
 	}
 	else if (Args.GetTabId() == CollectionSpreadSheetTabId_2)
 	{
-		DataflowCollectionSpreadSheet_2 = MakeShared<FDataflowCollectionSpreadSheet>(FDataflowCollectionSpreadSheet(DataflowEditor->GetEditorContent()));
+		DataflowCollectionSpreadSheet_2 = MakeShared<FDataflowCollectionSpreadSheet>(DataflowEditor->GetEditorContent());
 		if (DataflowCollectionSpreadSheet_2.IsValid())
 		{
 			ViewListeners.Add(DataflowCollectionSpreadSheet_2.Get());
@@ -1544,7 +1770,7 @@ TSharedRef<SDockTab> FDataflowEditorToolkit::SpawnTab_CollectionSpreadSheet(cons
 	}
 	else if (Args.GetTabId() == CollectionSpreadSheetTabId_3)
 	{
-		DataflowCollectionSpreadSheet_3 = MakeShared<FDataflowCollectionSpreadSheet>(FDataflowCollectionSpreadSheet(DataflowEditor->GetEditorContent()));
+		DataflowCollectionSpreadSheet_3 = MakeShared<FDataflowCollectionSpreadSheet>(DataflowEditor->GetEditorContent());
 		if (DataflowCollectionSpreadSheet_3.IsValid())
 		{
 			ViewListeners.Add(DataflowCollectionSpreadSheet_3.Get());
@@ -1552,7 +1778,7 @@ TSharedRef<SDockTab> FDataflowEditorToolkit::SpawnTab_CollectionSpreadSheet(cons
 	}
 	else if (Args.GetTabId() == CollectionSpreadSheetTabId_4)
 	{
-		DataflowCollectionSpreadSheet_4 = MakeShared<FDataflowCollectionSpreadSheet>(FDataflowCollectionSpreadSheet(DataflowEditor->GetEditorContent()));
+		DataflowCollectionSpreadSheet_4 = MakeShared<FDataflowCollectionSpreadSheet>(DataflowEditor->GetEditorContent());
 		if (DataflowCollectionSpreadSheet_4.IsValid())
 		{
 			ViewListeners.Add(DataflowCollectionSpreadSheet_4.Get());
@@ -1636,6 +1862,33 @@ TSharedRef<SDockTab> FDataflowEditorToolkit::SpawnTab_MembersWidget(const FSpawn
 	return DataflowMembersTab;
 }
 
+TSharedRef<SDockTab> FDataflowEditorToolkit::SpawnTab_OutputLog(const FSpawnTabArgs& Args)
+{
+	check(DataflowEditor);
+	check(DataflowEditor->GetEditorContent());
+
+	if (Args.GetTabId() == OutputLogTabId)
+	{
+		DataflowOutputLog = MakeShared<FDataflowOutputLog>(DataflowEditor->GetEditorContent());
+	}
+
+	using namespace UE::Dataflow;
+
+	TSharedRef<SDockTab> DockableTab = SNew(SDockTab)
+	[
+		DataflowOutputLog->GetOutputLogWidget().ToSharedRef()
+	];
+
+	if (DataflowOutputLog->GetOutputLogWidget())
+	{
+		OnOutputLogMessageTokenClickedDelegateHandle = DataflowOutputLog->GetOnOutputLogMessageTokenClickedDelegate().AddRaw(this, &FDataflowEditorToolkit::OnOutputLogMessageTokenClicked);
+	}
+
+	DockableTab->SetOnTabClosed(SDockTab::FOnTabClosedCallback::CreateRaw(this, &FDataflowEditorToolkit::OnTabClosed));
+
+	return DockableTab;
+}
+
 void FDataflowEditorToolkit::RegisterTabSpawners(const TSharedRef<FTabManager>& InTabManager)
 {
 	FAssetEditorToolkit::RegisterTabSpawners(InTabManager);
@@ -1732,6 +1985,7 @@ void FDataflowEditorToolkit::RegisterTabSpawners(const TSharedRef<FTabManager>&
 		.SetDisplayName(LOCTEXT("SimulationVisualizationTabDisplayName", "Simulation Visualization"))
 		.SetGroup(AssetEditorTabsCategory.ToSharedRef());
 
+	
 	InTabManager->RegisterTabSpawner(MembersWidgetTabId, FOnSpawnTab::CreateSP(this, &FDataflowEditorToolkit::SpawnTab_MembersWidget))
 		.SetDisplayName(LOCTEXT("MembersWidgetTabDisplayName", "Dataflow Members"))
 		.SetGroup(AssetEditorTabsCategory.ToSharedRef());
@@ -1748,6 +2002,11 @@ void FDataflowEditorToolkit::RegisterTabSpawners(const TSharedRef<FTabManager>&
 			FDataflowEditorSubGraphTabSummoner::FOnCreateGraphEditorWidget::CreateSP(this, &FDataflowEditorToolkit::CreateSubGraphEditorWidget)
 		));
 	DocumentManager->RegisterDocumentFactory(SubGraphTabSummoner);
+	
+	InTabManager->RegisterTabSpawner(OutputLogTabId, FOnSpawnTab::CreateSP(this, &FDataflowEditorToolkit::SpawnTab_OutputLog))
+		.SetDisplayName(LOCTEXT("OutputLogTab", "Output Log"))
+		.SetGroup(EditorMenuCategory.ToSharedRef())
+		.SetIcon(FSlateIcon(FAppStyle::GetAppStyleSetName(), "LevelEditor.Tabs.Outliner"));
 }
 
 void FDataflowEditorToolkit::UnregisterTabSpawners(const TSharedRef<class FTabManager>& InTabManager)
@@ -1772,6 +2031,7 @@ void FDataflowEditorToolkit::UnregisterTabSpawners(const TSharedRef<class FTabMa
 	InTabManager->UnregisterTabSpawner(CollectionSpreadSheetTabId_4);
 	InTabManager->UnregisterTabSpawner(SimulationVisualizationTabId);
 	InTabManager->UnregisterTabSpawner(MembersWidgetTabId);
+	InTabManager->UnregisterTabSpawner(OutputLogTabId);
 }
 
 void FDataflowEditorToolkit::OnTabClosed(TSharedRef<SDockTab> Tab)
diff --git a/Engine/Plugins/Experimental/Dataflow/Source/DataflowEditor/Private/Dataflow/DataflowOutputLog.cpp b/Engine/Plugins/Experimental/Dataflow/Source/DataflowEditor/Private/Dataflow/DataflowOutputLog.cpp
new file mode 100644
--- /dev/null
+++ b/Engine/Plugins/Experimental/Dataflow/Source/DataflowEditor/Private/Dataflow/DataflowOutputLog.cpp
@@ -0,0 +1,131 @@
+// Copyright Epic Games, Inc. All Rights Reserved.
+
+#include "Dataflow/DataflowOutputLog.h"
+#include "IMessageLogListing.h"
+#include "MessageLogModule.h"
+#include "Dataflow/DataflowPath.h"
+
+DEFINE_LOG_CATEGORY_STATIC(LogDataflowContextOutput, Warning, All);
+
+FDataflowOutputLog::FDataflowOutputLog(TObjectPtr<UDataflowBaseContent> InContent)
+	: FDataflowNodeView(InContent)
+{
+	CreateMessageLog();
+	CreateMessageLogWidget();
+}
+
+FDataflowOutputLog::~FDataflowOutputLog()
+{
+	if (MessageLogListing)
+	{
+		MessageLogListing->OnMessageTokenClicked().RemoveAll(this);
+	}
+}
+
+void FDataflowOutputLog::CreateMessageLog()
+{
+	FMessageLogModule& MessageLogModule = FModuleManager::LoadModuleChecked<FMessageLogModule>("MessageLog");
+
+	const FString NameString = FString::Printf(TEXT("LogDataflow%d"), UE::Dataflow::FTimestamp::Current());
+	const FName MessageLogName = FName(*NameString);
+
+	if (MessageLogModule.IsRegisteredLogListing(MessageLogName))
+	{
+		MessageLogListing = MessageLogModule.GetLogListing(MessageLogName);
+	}
+	else
+	{
+		FMessageLogInitializationOptions LogOptions;
+		LogOptions.bShowPages = false;
+		LogOptions.bShowFilters = true;
+		LogOptions.bAllowClear = true;
+		LogOptions.MaxPageCount = 1;
+		LogOptions.bShowInLogWindow = false;
+
+		MessageLogListing = MessageLogModule.CreateLogListing(MessageLogName, LogOptions);
+	}
+
+	if (MessageLogListing.IsValid())
+	{
+		MessageLogListing->OnMessageTokenClicked().AddRaw(this, &FDataflowOutputLog::OnMessageTokenClicked);
+		MessageLogListing->ClearMessages();
+	}
+}
+
+TSharedRef<IMessageLogListing> FDataflowOutputLog::GetMessageLog() const
+{
+	return MessageLogListing.ToSharedRef();
+}
+
+void FDataflowOutputLog::ClearMessageLog()
+{
+	if (MessageLogListing.IsValid())
+	{
+		MessageLogListing->ClearMessages();
+	}
+}
+
+namespace UE::Dataflow::Private
+{
+	TSharedRef<FTokenizedMessage> BuildTokenizedMessage(const EMessageSeverity::Type InSeverity, const FString& InMessage, const FDataflowPath& InPath)
+	{
+		// Token0 - Severity
+		// Token1 - Time code
+		// Token2 - Path ()
+		// Token3 - Message
+
+		const FDateTime TimeNow = FDateTime::Now();
+		const FString TimeString = FString::Printf(TEXT("[%02d:%02d:%02d:%02d]"), TimeNow.GetHour(), TimeNow.GetMinute(), TimeNow.GetSecond(), TimeNow.GetMillisecond());
+
+		TSharedRef<FTokenizedMessage> TokenizedMessage = FTokenizedMessage::Create(InSeverity);
+		TSharedRef<IMessageToken> TokenTimecode = FTextToken::Create(FText::FromString(TimeString));
+		TokenizedMessage->AddToken(TokenTimecode);
+		TSharedRef<IMessageToken> TokenPath = FTextToken::Create(FText::FromString("[" + InPath.ToString() + "]"));
+		TokenizedMessage->AddToken(TokenPath);
+		TokenizedMessage->SetMessageLink(TokenPath);
+		TSharedRef<IMessageToken> TokenMessage = FTextToken::Create(FText::FromString(InMessage));
+		TokenizedMessage->AddToken(TokenMessage);
+
+		return TokenizedMessage;
+	}
+
+	FString GetTokenString(const TSharedRef<IMessageToken>& InMessageToken)
+	{
+		// Remove [] from InToken
+		FString TokenString = InMessageToken->ToText().ToString();
+		TokenString = TokenString.TrimChar(TEXT('['));
+		TokenString = TokenString.TrimChar(TEXT(']'));
+
+		return TokenString;
+	}
+}
+
+void FDataflowOutputLog::AddMessage(const EMessageSeverity::Type InSeverity, const FString& InMessage, const FDataflowPath& InPath)
+{
+	TSharedRef<FTokenizedMessage> TokenizedMessage = UE::Dataflow::Private::BuildTokenizedMessage(InSeverity, InMessage, InPath);
+
+	if (MessageLogListing.IsValid())
+	{
+		MessageLogListing->AddMessage(TokenizedMessage, false);
+	}
+
+	UE_LOG(LogDataflowContextOutput, Warning, TEXT("[%s][%s]"), *InPath.ToString(), *InMessage);
+}
+
+void FDataflowOutputLog::CreateMessageLogWidget()
+{
+	FMessageLogModule& MessageLogModule = FModuleManager::LoadModuleChecked<FMessageLogModule>("MessageLog");
+
+	OutputLogWidget = MessageLogModule.CreateLogListingWidget(MessageLogListing.ToSharedRef());
+}
+
+void FDataflowOutputLog::OnMessageTokenClicked(const TSharedRef<IMessageToken>& InMessageToken)
+{
+	FString TokenString = UE::Dataflow::Private::GetTokenString(InMessageToken);
+
+	if (OnOutputLogMessageTokenClickedDelegate.IsBound())
+	{
+		OnOutputLogMessageTokenClickedDelegate.Broadcast(TokenString);
+	}
+}
+
diff --git a/Engine/Plugins/Experimental/Dataflow/Source/DataflowEditor/Private/Dataflow/DataflowSNode.cpp b/Engine/Plugins/Experimental/Dataflow/Source/DataflowEditor/Private/Dataflow/DataflowSNode.cpp
--- a/Engine/Plugins/Experimental/Dataflow/Source/DataflowEditor/Private/Dataflow/DataflowSNode.cpp
+++ b/Engine/Plugins/Experimental/Dataflow/Source/DataflowEditor/Private/Dataflow/DataflowSNode.cpp
@@ -241,6 +241,10 @@ TArray<FOverlayWidgetInfo> SDataflowEdNode::GetOverlayWidgets(bool bSelected, co
 	return Widgets;
 }
 
+const FSlateColor CExperimentalColor(FLinearColor(26.f/255.f, 188.f/255.f, 156.f/255.f, 255.f)); // Turquoise
+const FSlateColor CDeprecatedColor(FLinearColor(46.f / 255.f, 204.f / 255.f, 113.f / 255.f, 255.f)); // Emerald
+const FSlateColor CFailedColor(FLinearColor(243.f / 255.f, 156.f / 255.f, 18.f / 255.f, 255.f)); // Orange
+
 void SDataflowEdNode::UpdateErrorInfo()
 {
 	if (DataflowGraphNode)
@@ -250,12 +254,35 @@ void SDataflowEdNode::UpdateErrorInfo()
 			if (UE::Dataflow::FNodeFactory::IsNodeExperimental(DataflowNode->GetType()))
 			{
 				ErrorMsg = FString(TEXT("Experimental"));
-				ErrorColor = FAppStyle::GetColor("ErrorReporting.WarningBackgroundColor");
-			} 
+				ErrorColor = CExperimentalColor;
+			}
 			if (UE::Dataflow::FNodeFactory::IsNodeDeprecated(DataflowNode->GetType()))
 			{
 				ErrorMsg = FString(TEXT("Deprecated"));
-				ErrorColor = FAppStyle::GetColor("ErrorReporting.WarningBackgroundColor");
+				ErrorColor = CDeprecatedColor;
+			}
+
+			if (DataflowInterface)
+			{
+				if (const TSharedPtr<UE::Dataflow::FContext> Context = DataflowInterface->GetDataflowContext())
+				{
+					if (Context->NodeHasWarning(DataflowNode.Get()))
+					{
+						ErrorMsg = FString(TEXT("WARNING"));
+						ErrorColor = FAppStyle::GetColor("ErrorReporting.WarningBackgroundColor");
+					}
+					else if (Context->NodeHasError(DataflowNode.Get()))
+					{
+						ErrorMsg = FString(TEXT("ERROR"));
+						const FLinearColor ErrorBackgroundColor = FLinearColor(1.f, 0.02f, 0.003f);
+						ErrorColor = ErrorBackgroundColor;
+					}
+					else if (Context->NodeFailed(DataflowNode.Get()))
+					{
+						ErrorMsg = FString(TEXT("FAILED"));
+						ErrorColor = CFailedColor;
+					}
+				}
 			}
 		}
 	}
diff --git a/Engine/Plugins/Experimental/Dataflow/Source/DataflowEditor/Public/Dataflow/DataflowEditorToolkit.h b/Engine/Plugins/Experimental/Dataflow/Source/DataflowEditor/Public/Dataflow/DataflowEditorToolkit.h
--- a/Engine/Plugins/Experimental/Dataflow/Source/DataflowEditor/Public/Dataflow/DataflowEditorToolkit.h
+++ b/Engine/Plugins/Experimental/Dataflow/Source/DataflowEditor/Public/Dataflow/DataflowEditorToolkit.h
@@ -32,11 +32,12 @@ class FDataflowSkeletonView;
 class UDataflowEditor;
 class SDataflowMembersWidget;
 class FDataflowSimulationSceneProfileIndexStorage;
+class FDataflowOutputLog;
 namespace UE::Dataflow
 {
 	class FDataflowNodeDetailExtensionHandler;
 }
-
+namespace EMessageSeverity { enum Type : int; }
 
 class DATAFLOWEDITOR_API FDataflowEditorToolkit final : public FBaseCharacterFXEditorToolkit, public FTickableEditorObject, public FNotifyHook, public FGCObject
 {
@@ -51,7 +52,10 @@ public:
 	static bool HasDataflowAsset(UObject* ObjectToEdit);
 	static UDataflow* GetDataflowAsset(UObject* ObjectToEdit);
 	static const UDataflow* GetDataflowAsset(const UObject* ObjectToEdit);
-	
+	bool IsSimulationDataflowAsset() const;
+	FName GetGraphLogName() const;
+	void LogMessage(const EMessageSeverity::Type Severity, const FDataflowNode* Node, const FDataflowOutput* Output, const FString& Message) const;
+
 	/** Editor dataflow content accessors */
 	const TObjectPtr<UDataflowBaseContent>& GetEditorContent() const;
 	TObjectPtr<UDataflowBaseContent>& GetEditorContent();
@@ -113,7 +117,14 @@ protected:
 
 	// Node evaluation
 	void EvaluateNode(const FDataflowNode* Node, const FDataflowOutput* Output, UE::Dataflow::FTimestamp& InOutTimestamp);
+	void OnNodeBeginEvaluate(const FDataflowNode* Node, const FDataflowOutput* Output);
+	void OnNodeFinishEvaluate(const FDataflowNode* Node, const FDataflowOutput* Output);
+	void OnBeginEvaluate();
 	void OnFinishEvaluate();
+	void OnOutputLogMessageTokenClicked(const FString TokenString);
+	void OnContextHasInfo(const FDataflowNode* Node, const FDataflowOutput* Output, const FString& Info);
+	void OnContextHasWarning(const FDataflowNode* Node, const FDataflowOutput* Output, const FString& Warning);
+	void OnContextHasError(const FDataflowNode* Node, const FDataflowOutput* Output, const FString& Error);
 
 private:
 	
@@ -129,7 +140,8 @@ private:
 	TSharedRef<SDockTab> SpawnTab_PreviewScene(const FSpawnTabArgs& Args);
 	TSharedRef<SDockTab> SpawnTab_SimulationVisualization(const FSpawnTabArgs& Args);
 	TSharedRef<SDockTab> SpawnTab_MembersWidget(const FSpawnTabArgs& Args);
-	
+	TSharedRef<SDockTab> SpawnTab_OutputLog(const FSpawnTabArgs& Args);
+
 	// FTickableEditorObject interface
 	virtual void Tick(float DeltaTime) override;
 	virtual bool IsTickable() const override { return true; }
@@ -171,6 +183,7 @@ private:
 	static const FName PreviewSceneTabId;
 	static const FName SimulationVisualizationTabId;
 	static const FName MembersWidgetTabId;
+	static const FName OutputLogTabId;
 
 	// List of all the widgets shared ptr that will be built in the editor
 	TSharedPtr<SDataflowConstructionViewport> DataflowConstructionViewport;
@@ -191,6 +204,7 @@ private:
 	TSharedPtr<FDataflowCollectionSpreadSheet> DataflowCollectionSpreadSheet_4;
 	TSharedPtr<SWidget> AdvancedPreviewSettingsWidget;
 	TSharedPtr<SWidget> SimulationVisualizationWidget;
+	TSharedPtr<FDataflowOutputLog> DataflowOutputLog;
 
 	/** Customize preview scene with editor/terminal contents */
 	TSharedRef<class IDetailCustomization> CustomizePreviewSceneDescription() const;
@@ -222,6 +236,15 @@ private:
 	FDelegateHandle OnConstructionSelectionChangedDelegateHandle;
 	FDelegateHandle OnSimulationSceneChangedDelegateHandle;
 
+	// Delegates to communicate with Context
+	FDelegateHandle OnNodeBeginEvaluateMulticastDelegateHandle;
+	FDelegateHandle OnNodeFinishEvaluateMulticastDelegateHandle;
+	FDelegateHandle OnOutputLogMessageTokenClickedDelegateHandle;
+
+	FDelegateHandle OnContextHasInfoDelegateHandle;
+	FDelegateHandle OnContextHasWarningDelegateHandle;
+	FDelegateHandle OnContextHasErrorDelegateHandle;
+
 	// The currently selected set of dataflow nodes. 
 	UPROPERTY()
 	TSet< TObjectPtr<UObject> > SelectedDataflowNodes;
@@ -266,4 +289,7 @@ private:
 
 	/** Document tracker for dynamic tabs ( like subgraphs ) */
 	TSharedPtr<FDocumentTracker> DocumentManager;
+
+	FDateTime GraphEvaluationBegin;
+	FDateTime GraphEvaluationFinished;
 };
diff --git a/Engine/Plugins/Experimental/Dataflow/Source/DataflowEditor/Public/Dataflow/DataflowOutputLog.h b/Engine/Plugins/Experimental/Dataflow/Source/DataflowEditor/Public/Dataflow/DataflowOutputLog.h
new file mode 100644
--- /dev/null
+++ b/Engine/Plugins/Experimental/Dataflow/Source/DataflowEditor/Public/Dataflow/DataflowOutputLog.h
@@ -0,0 +1,50 @@
+// Copyright Epic Games, Inc. All Rights Reserved.
+
+#pragma once
+
+#include "Dataflow/DataflowView.h"
+
+namespace UE::Dataflow
+{
+	class SOutputLogWidget;
+}
+
+struct FDataflowPath;
+class IMessageLogListing;
+
+/**
+*
+* Class to handle the OutputLog widget
+*
+*/
+class FDataflowOutputLog final : public FDataflowNodeView
+{
+public:
+	explicit FDataflowOutputLog(TObjectPtr<UDataflowBaseContent> InContent = nullptr);
+	virtual ~FDataflowOutputLog();
+
+	TSharedPtr<SWidget> GetOutputLogWidget() { return OutputLogWidget; }
+
+	TSharedRef<IMessageLogListing> GetMessageLog() const;
+	void ClearMessageLog();
+	void AddMessage(const EMessageSeverity::Type InSeverity, const FString& InMessage, const FDataflowPath& InPath);
+
+	DECLARE_MULTICAST_DELEGATE_OneParam(FOnOutputLogMessageTokenClicked, const FString);
+	FOnOutputLogMessageTokenClicked& GetOnOutputLogMessageTokenClickedDelegate() { return OnOutputLogMessageTokenClickedDelegate; }
+
+private:
+	virtual void SetSupportedOutputTypes() override {};
+	virtual void UpdateViewData() override {};
+	virtual void ConstructionViewSelectionChanged(const TArray<UPrimitiveComponent*>& InSelectedComponents) override {};
+
+	void OnMessageTokenClicked(const TSharedRef<IMessageToken>& InMessageToken);
+
+	void CreateMessageLog();
+	void CreateMessageLogWidget();
+
+	TSharedPtr<SWidget> OutputLogWidget;
+	TSharedPtr<IMessageLogListing> MessageLogListing;
+
+	FOnOutputLogMessageTokenClicked OnOutputLogMessageTokenClickedDelegate;
+};
+
diff --git a/Engine/Plugins/Experimental/GeometryCollectionPlugin/Source/GeometryCollectionNodes/Private/Dataflow/GeometryCollectionMathNodes.cpp b/Engine/Plugins/Experimental/GeometryCollectionPlugin/Source/GeometryCollectionNodes/Private/Dataflow/GeometryCollectionMathNodes.cpp
--- a/Engine/Plugins/Experimental/GeometryCollectionPlugin/Source/GeometryCollectionNodes/Private/Dataflow/GeometryCollectionMathNodes.cpp
+++ b/Engine/Plugins/Experimental/GeometryCollectionPlugin/Source/GeometryCollectionNodes/Private/Dataflow/GeometryCollectionMathNodes.cpp
@@ -4,7 +4,6 @@
 #include "Math/BasicMathExpressionEvaluator.h"
 #include "Dataflow/DataflowCore.h"
 
-
 #include UE_INLINE_GENERATED_CPP_BY_NAME(GeometryCollectionMathNodes)
 
 namespace UE::Dataflow
@@ -12,8 +11,6 @@ namespace UE::Dataflow
 
 	void GeometryCollectionMathNodes()
 	{
-		static const FLinearColor CDefaultNodeBodyTintColor = FLinearColor(0.f, 0.f, 0.f, 0.5f);
-
 		DATAFLOW_NODE_REGISTER_CREATION_FACTORY(FAddDataflowNode);
 		DATAFLOW_NODE_REGISTER_CREATION_FACTORY(FSubtractDataflowNode);
 		DATAFLOW_NODE_REGISTER_CREATION_FACTORY(FMultiplyDataflowNode);
@@ -71,9 +68,6 @@ namespace UE::Dataflow
 		DATAFLOW_NODE_REGISTER_CREATION_FACTORY(FOneMinusDataflowNode);
 		DATAFLOW_NODE_REGISTER_CREATION_FACTORY(FFloatMathExpressionDataflowNode);
 		DATAFLOW_NODE_REGISTER_CREATION_FACTORY(FMathExpressionDataflowNode);
-
-		// Math
-		DATAFLOW_NODE_REGISTER_CREATION_FACTORY_NODE_COLORS_BY_CATEGORY("Math", FLinearColor(0.f, 0.4f, 0.8f), CDefaultNodeBodyTintColor);
 	}
 }
 
@@ -124,8 +118,12 @@ void FSafeDivideDataflowNode::Evaluate(UE::Dataflow::FContext& Context, const FD
 		if (InFloatB != 0.f)
 		{
 			Result = InFloatA / InFloatB;
+			SetValue(Context, Result, &ReturnValue);
+		}
+		else
+		{
+			SetError(Context, &ReturnValue, TEXT("Division by zero error"));
 		}
-		SetValue(Context, Result, &ReturnValue);
 	}
 }
 
diff --git a/Engine/Source/Runtime/Experimental/Dataflow/Core/DataflowCore.Build.cs b/Engine/Source/Runtime/Experimental/Dataflow/Core/DataflowCore.Build.cs
--- a/Engine/Source/Runtime/Experimental/Dataflow/Core/DataflowCore.Build.cs
+++ b/Engine/Source/Runtime/Experimental/Dataflow/Core/DataflowCore.Build.cs
@@ -15,7 +15,8 @@ namespace UnrealBuildTool.Rules
 					"DeveloperSettings",
 					"InputCore",
 					"ImageCore",
-					"ApplicationCore"
+					"ApplicationCore",
+					"MessageLog"
 				}
 			);
 		}
diff --git a/Engine/Source/Runtime/Experimental/Dataflow/Core/Private/Dataflow/DataflowInputOutput.cpp b/Engine/Source/Runtime/Experimental/Dataflow/Core/Private/Dataflow/DataflowInputOutput.cpp
--- a/Engine/Source/Runtime/Experimental/Dataflow/Core/Private/Dataflow/DataflowInputOutput.cpp
+++ b/Engine/Source/Runtime/Experimental/Dataflow/Core/Private/Dataflow/DataflowInputOutput.cpp
@@ -290,11 +290,19 @@ bool FDataflowOutput::EvaluateImpl(UE::Dataflow::FContext& Context) const
 		// if not, evaluate
 		UE_LOG(LogChaosDataflow, Verbose, TEXT("FDataflowNode::Evaluate(): Node [%s], Output [%s], NodeTimestamp [%lu]"), *GetOwningNode()->GetName().ToString(), *GetName().ToString(), OwningNode->GetTimestamp().Value);
 		OwningNode->Evaluate(Context, this);
+		
 		// Validation
-		if (!Context.HasData(CacheKey()))
+		if (!Context.NodeHasError(OwningNode) && !Context.NodeFailed(OwningNode))
 		{
-			ensureMsgf(false, TEXT("Failed to evaluate output (%s:%s)"), *OwningNode->GetName().ToString(), *GetName().ToString());
-			return false;
+			if (!Context.HasData(CacheKey()))
+			{
+				ensureMsgf(false, TEXT("Failed to evaluate output (%s:%s)"), *OwningNode->GetName().ToString(), *GetName().ToString());
+				return false;
+			}
+		}
+		else
+		{
+			this->SetNullValue(Context);
 		}
 	}
 	return true;
@@ -447,3 +455,8 @@ void* FDataflowArrayOutput::RealAddress() const
 	}
 	return nullptr;
 }
+
+bool FDataflowOutput::HasNodeFailedOrErrored(UE::Dataflow::FContext& Context) const
+{
+	return Context.NodeHasError(GetOwningNode()) || Context.NodeFailed(GetOwningNode());
+}
diff --git a/Engine/Source/Runtime/Experimental/Dataflow/Core/Private/Dataflow/DataflowNode.cpp b/Engine/Source/Runtime/Experimental/Dataflow/Core/Private/Dataflow/DataflowNode.cpp
--- a/Engine/Source/Runtime/Experimental/Dataflow/Core/Private/Dataflow/DataflowNode.cpp
+++ b/Engine/Source/Runtime/Experimental/Dataflow/Core/Private/Dataflow/DataflowNode.cpp
@@ -903,6 +903,13 @@ FDataflowInput& FDataflowNode::RegisterInputConnectionInternal(const UE::Dataflo
 	AddInput(Input);
 	check(FindInput(Reference) == Input);
 
+#if WITH_EDITORONLY_DATA
+	if (Input->GetProperty()->HasMetaData("DataflowIntrinsic"))
+	{
+		Input->SetIsRequired(true);
+	}
+#endif
+
 	return *Input;
 }
 
diff --git a/Engine/Source/Runtime/Experimental/Dataflow/Core/Private/Dataflow/DataflowNodeParameters.cpp b/Engine/Source/Runtime/Experimental/Dataflow/Core/Private/Dataflow/DataflowNodeParameters.cpp
--- a/Engine/Source/Runtime/Experimental/Dataflow/Core/Private/Dataflow/DataflowNodeParameters.cpp
+++ b/Engine/Source/Runtime/Experimental/Dataflow/Core/Private/Dataflow/DataflowNodeParameters.cpp
@@ -8,6 +8,11 @@
 #include "Dataflow/DataflowNode.h"
 #include "Serialization/Archive.h"
 #include "GeometryCollection/ManagedArrayCollection.h"
+#include "IMessageLogListing.h"
+#include "MessageLogModule.h"
+#include "Modules/ModuleManager.h"
+#include "Misc/UObjectToken.h"
+#include "Dataflow/DataflowPath.h"
 
 namespace UE::Dataflow
 {
@@ -66,6 +71,81 @@ namespace UE::Dataflow
 		return false;
 	}
 
+	void FContext::Info(const FString& InInfo, const FDataflowNode* InNode, const FDataflowOutput* InOutput)
+	{
+		if (OnContextHasInfo.IsBound())
+		{
+			OnContextHasInfo.Broadcast(InNode, InOutput, InInfo);
+		}
+
+		NodesWithInfo.Add(InNode);
+	}
+
+	int32 FContext::GetNumInfo() const
+	{
+		return NodesWithInfo.Num();
+	}
+
+	void FContext::Warning(const FString& InWarning, const FDataflowNode* InNode, const FDataflowOutput* InOutput)
+	{
+		if (OnContextHasWarning.IsBound())
+		{
+			OnContextHasWarning.Broadcast(InNode, InOutput, InWarning);
+		}
+
+		NodesWithWarning.Add(InNode);
+	}
+
+	int32 FContext::GetNumWarnings() const
+	{
+		return NodesWithWarning.Num();
+	}
+
+	void FContext::Error(const FString& InError, const FDataflowNode* InNode, const FDataflowOutput* InOutput)
+	{
+		if (OnContextHasError.IsBound())
+		{
+			OnContextHasError.Broadcast(InNode, InOutput, InError);
+		}
+	
+		NodesWithError.Add(InNode);
+
+		//
+		// Add all the connected nodes to NodesFailed set
+		//
+#if DATAFLOW_EDITOR_EVALUATION
+		const int32 NumNodesInCallStack = Callstack.Num();
+		if (NumNodesInCallStack > 1)
+		{
+			for (int32 Idx = NumNodesInCallStack - 2; Idx >= 0; --Idx)
+			{
+				const FDataflowNode* OwningNode = Callstack[Idx]->GetOwningNode();
+				const FString WarningStr = FString(TEXT("Evaluation failed"));
+
+				if (OnContextHasError.IsBound())
+				{
+					OnContextHasWarning.Broadcast(OwningNode, InOutput, WarningStr);
+				}
+
+				NodesFailed.Add(OwningNode);
+			}
+		}
+#endif
+	}
+
+	int32 FContext::GetNumErrors() const
+	{
+		return NodesWithError.Num();
+	}
+
+	void FContext::ClearNodesData()
+	{
+		NodesWithInfo.Empty();
+		NodesWithWarning.Empty();
+		NodesWithError.Empty();
+		NodesFailed.Empty();
+	}
+
 	FContextScopedCallstack::FContextScopedCallstack(FContext& InContext, const FDataflowConnection* InConnection)
 		: Context(InContext)
 		, Connection(InConnection)
@@ -102,8 +182,19 @@ namespace UE::Dataflow
 				// TODO: When no outputs are specified, this call to Evaluate should really be removed.
 				//       The purpose of the node evaluation function is to evaluate outputs.
 				//       Therefore if a node has no outputs, then it shouldn't need any evaluation.
+
+				if (this->OnNodeBeginEvaluateMulticast.IsBound())
+				{
+					this->OnNodeBeginEvaluateMulticast.Broadcast(Node, Output);
+				}
+
 				UE_LOG(LogChaosDataflow, Verbose, TEXT("FDataflowNode::Evaluate(): Node [%s], Output [nullptr], NodeTimestamp [%lu]"), *Node->GetName().ToString(), Node->GetTimestamp().Value);
 				Node->Evaluate(*this, nullptr);
+
+				if (this->OnNodeFinishEvaluateMulticast.IsBound())
+				{
+					this->OnNodeFinishEvaluateMulticast.Broadcast(Node, Output);
+				}
 			}
 		}
 		else
@@ -119,11 +210,39 @@ namespace UE::Dataflow
 
 	bool FContextSingle::Evaluate(const FDataflowOutput& Connection)
 	{
+#if WITH_EDITOR
+		//
+		// Check if all DataflowIntrinsic input(s) connected
+		// 
+		const FDataflowNode* Node = Connection.GetOwningNode();
+		TArray<FDataflowInput*> Inputs = Node->GetInputs();
+		for (FDataflowInput* Input : Inputs)
+		{
+			if (Input->IsRequired() && Input->GetConnection() == nullptr)
+			{
+				const FString InputNameStr = Input->GetName().ToString();
+				const FString WarningMsg = FString::Printf(TEXT("Missing connection to required input : %s"), *InputNameStr);
+				this->Warning(WarningMsg, Node);
+			}
+		}
+
 		UE_LOG(LogChaosDataflow, VeryVerbose, TEXT("FContextSingle::Evaluate(): Node [%s], Output [%s]"), *Connection.GetOwningNode()->GetName().ToString(), *Connection.GetName().ToString());
-		return Connection.EvaluateImpl(*this);
-	}
+#endif // EDITOR_ONLY
 
+		if (OnNodeBeginEvaluateMulticast.IsBound())
+		{
+			this->OnNodeBeginEvaluateMulticast.Broadcast(Connection.GetOwningNode(), &Connection);
+		}
 
+		const bool RetValue = Connection.EvaluateImpl(*this);
+
+		if (OnNodeFinishEvaluateMulticast.IsBound())
+		{
+			this->OnNodeFinishEvaluateMulticast.Broadcast(Connection.GetOwningNode(), &Connection);
+		}
+
+		return !GetNumErrors() && RetValue;
+	}
 
 	void FContextThreaded::Evaluate(const FDataflowNode* Node, const FDataflowOutput* Output)
 	{
@@ -137,10 +256,8 @@ namespace UE::Dataflow
 		return Connection.EvaluateImpl(*this);
 	}
 
-
 	void FContextCache::Serialize(FArchive& Ar)
 	{
-
 		if (Ar.IsSaving())
 		{
 
diff --git a/Engine/Source/Runtime/Experimental/Dataflow/Core/Private/Dataflow/DataflowPath.cpp b/Engine/Source/Runtime/Experimental/Dataflow/Core/Private/Dataflow/DataflowPath.cpp
new file mode 100644
--- /dev/null
+++ b/Engine/Source/Runtime/Experimental/Dataflow/Core/Private/Dataflow/DataflowPath.cpp
@@ -0,0 +1,84 @@
+// Copyright Epic Games, Inc. All Rights Reserved.
+
+#include "Dataflow/DataflowPath.h"
+
+#include UE_INLINE_GENERATED_CPP_BY_NAME(DataflowPath)
+
+FString FDataflowPath::ToString() const
+{
+	FString OutString;
+
+	if (Node.IsEmpty())
+	{
+		OutString = FString::Printf(TEXT("Graph=%s"), *Graph);
+	}
+	else
+	{
+		if (Input.IsEmpty() && Output.IsEmpty())
+		{
+			OutString = FString::Printf(TEXT("Graph=%s|Node=%s"), *Graph, *Node);
+		}
+		else if (Output.IsEmpty())
+		{
+			OutString = FString::Printf(TEXT("Graph=%s|Node=%s|Input=%s"), *Graph, *Node, *Input);
+		}
+		else if (Input.IsEmpty())
+		{
+			OutString = FString::Printf(TEXT("Graph=%s|Node=%s|Output=%s"), *Graph, *Node, *Output);
+		}
+	}
+
+	return OutString;
+}
+
+bool FDataflowPath::PathHasInput() const
+{
+	return !Input.IsEmpty();
+}
+
+bool FDataflowPath::PathHasOutput() const
+{
+	return !Output.IsEmpty();
+}
+
+void FDataflowPath::DecodePath(const FString InPath)
+{
+	TArray<FString> Tokens;
+	int32 NumElems = InPath.ParseIntoArray(Tokens, TEXT("|"));
+
+	FString StringA, StringB;
+
+	Graph = "";
+	Node = "";
+	Input = "";
+	Output = "";
+
+	// Graph
+	if (Tokens.Num() > 0)
+	{
+		Tokens[0].Split(TEXT("="), &StringA, &StringB);
+		Graph = StringB;
+	}
+
+	// Node
+	if (Tokens.Num() > 1)
+	{
+		Tokens[1].Split(TEXT("="), &StringA, &StringB);
+		Node = StringB;
+	}
+
+	// Input
+	if (Tokens.Num() > 1 && InPath.Contains("Input", ESearchCase::CaseSensitive, ESearchDir::FromStart))
+	{
+		Tokens[2].Split(TEXT("="), &StringA, &StringB);
+		Input = StringB;
+	}
+
+	// Output
+	if (Tokens.Num() > 1 && InPath.Contains("Output", ESearchCase::CaseSensitive, ESearchDir::FromStart))
+	{
+		Tokens[2].Split(TEXT("="), &StringA, &StringB);
+		Output = StringB;
+	}
+}
+
diff --git a/Engine/Source/Runtime/Experimental/Dataflow/Core/Public/Dataflow/DataflowInputOutput.h b/Engine/Source/Runtime/Experimental/Dataflow/Core/Public/Dataflow/DataflowInputOutput.h
--- a/Engine/Source/Runtime/Experimental/Dataflow/Core/Public/Dataflow/DataflowInputOutput.h
+++ b/Engine/Source/Runtime/Experimental/Dataflow/Core/Public/Dataflow/DataflowInputOutput.h
@@ -51,6 +51,10 @@ protected:
 	friend struct FDataflowOutput;
 	virtual void FixAndPropagateType(FName InType) override;
 
+private:
+	// if the Input has DATAFLOW_INTRINSIC metatag specified
+	bool bIsRequired = false;
+
 public:
 	UE_DEPRECATED(5.5, "Deprecated constructor : Guid is now passed through FInputParameters")
 	DATAFLOWCORE_API FDataflowInput(const UE::Dataflow::FInputParameters& Param, FGuid InGuid);
@@ -90,6 +94,9 @@ public:
 	TFuture<const T&> GetValueParallel(UE::Dataflow::FContext& Context, const T& Default) const;
 
 	virtual void Invalidate(const UE::Dataflow::FTimestamp& ModifiedTimestamp = UE::Dataflow::FTimestamp::Current()) override;
+
+	bool IsRequired() const { return bIsRequired; }
+	void SetIsRequired(const bool bInIsRequired) { bIsRequired = bInIsRequired; }
 };
 
 USTRUCT()
@@ -221,6 +228,7 @@ public:
 				using FSingleType = typename decltype(SingleTypePolicy)::FType;
 				FSingleType ValueToSet{};
 				FDataflowConverter<typename TAnyType::FStorageType>::To(InVal, ValueToSet);
+
 				Context.SetData(CacheKey(), GetProperty(), Forward<FSingleType>(ValueToSet), GetOwningNodeGuid(), GetOwningNodeValueHash(), GetOwningNodeTimestamp());
 			});
 	}
@@ -278,6 +286,9 @@ public:
 
 	DATAFLOWCORE_API virtual void Invalidate(const UE::Dataflow::FTimestamp& ModifiedTimestamp = UE::Dataflow::FTimestamp::Current()) override;
 
+	// Check if OwningNode has an error or it has failed, and if it does set the timestamp on the output(s) to FTimestamp::Invalid
+	DATAFLOWCORE_API bool HasNodeFailedOrErrored(UE::Dataflow::FContext& Context) const;
+
 private:
 	DATAFLOWCORE_API const uint8* GetFrozenPropertyValue(const uint8* Default) const;
 };
@@ -329,7 +340,7 @@ const T& FDataflowInput::GetValue(UE::Dataflow::FContext& Context, const T& Defa
 		}
 		if (!ConnectionOut->Evaluate(Context))
 		{
-			Context.SetData(ConnectionOut->CacheKey(), Property, Default, GetOwningNodeGuid(), GetOwningNodeValueHash(), GetOwningNodeTimestamp());
+			Context.SetData(ConnectionOut->CacheKey(), Property, Default, GetOwningNodeGuid(), GetOwningNodeValueHash(), ConnectionOut->GetOwningNodeTimestamp());
 		}
 		if (Context.HasData(ConnectionOut->CacheKey()))
 		{
diff --git a/Engine/Source/Runtime/Experimental/Dataflow/Core/Public/Dataflow/DataflowNode.h b/Engine/Source/Runtime/Experimental/Dataflow/Core/Public/Dataflow/DataflowNode.h
--- a/Engine/Source/Runtime/Experimental/Dataflow/Core/Public/Dataflow/DataflowNode.h
+++ b/Engine/Source/Runtime/Experimental/Dataflow/Core/Public/Dataflow/DataflowNode.h
@@ -510,6 +510,11 @@ PRAGMA_ENABLE_DEPRECATION_WARNINGS
 		}
 	}
 
+	void SetError(UE::Dataflow::FContext& Context, const void* Reference, const FString& InError) const
+	{
+		Context.Error(InError, this, FindOutput(Reference));
+	}
+
 	/**
 	*   ForwardInput(...)
 	*
diff --git a/Engine/Source/Runtime/Experimental/Dataflow/Core/Public/Dataflow/DataflowNodeParameters.h b/Engine/Source/Runtime/Experimental/Dataflow/Core/Public/Dataflow/DataflowNodeParameters.h
--- a/Engine/Source/Runtime/Experimental/Dataflow/Core/Public/Dataflow/DataflowNodeParameters.h
+++ b/Engine/Source/Runtime/Experimental/Dataflow/Core/Public/Dataflow/DataflowNodeParameters.h
@@ -40,6 +40,7 @@ namespace UE::Dataflow
 		bool operator<(const FTimestamp& InTimestamp) const { return Value < InTimestamp.Value; }
 		bool operator==(const FTimestamp& InTimestamp) const { return Value == InTimestamp.Value; }
 		bool IsInvalid() { return Value == Invalid; }
+		bool IsInvalid() const { return Value == Invalid; }
 
 		static DATAFLOWCORE_API Type Current();
 		static DATAFLOWCORE_API Type Invalid; // 0
@@ -177,8 +178,12 @@ namespace UE::Dataflow
 	//--------------------------------------------------------------------
 	struct FContextCacheElementNull : public FContextCacheElementBase
 	{
+		//
+		// IMPORTANT: 
+		// Timestamp must be set to (Timestamp.Value - 1) to make sure that this type of entry is always invalid
+		//
 		FContextCacheElementNull(FGuid InNodeGuid, const FProperty* InProperty, FContextCacheKey InDataKey, uint32 InNodeHash, FTimestamp Timestamp)
-			: FContextCacheElementBase(EType::CacheElementNull, InNodeGuid, InProperty, InNodeHash, Timestamp)
+			: FContextCacheElementBase(EType::CacheElementNull, InNodeGuid, InProperty, InNodeHash, FTimestamp(Timestamp.Value - 1))
 		{}
 
 		virtual const void* GetUntypedData(const FContext& Context, const FProperty* PropertyIn) const
@@ -418,6 +423,7 @@ namespace UE::Dataflow
 			if (const TUniquePtr<FContextCacheElementBase>* CacheElement = GetDataImpl(ReferenceKey))
 			{
 				TUniquePtr<FContextCacheElementBase> CacheReferenceElement = (*CacheElement)->CreateReference(ReferenceKey);
+
 				SetDataImpl(Key, MoveTemp(CacheReferenceElement));
 			}
 			else
@@ -489,6 +495,34 @@ namespace UE::Dataflow
 
 		DATAFLOWCORE_API bool IsCacheEntryAfterTimestamp(FContextCacheKey InKey, const FTimestamp InTimestamp) const;
 
+		DATAFLOWCORE_API void Info(const FString& InInfo, const FDataflowNode* InNode = nullptr, const FDataflowOutput* InOutput = nullptr);
+		DATAFLOWCORE_API int32 GetNumInfo() const;
+		DATAFLOWCORE_API void Warning(const FString& InWarning, const FDataflowNode* InNode = nullptr, const FDataflowOutput* InOutput = nullptr);
+		DATAFLOWCORE_API int32 GetNumWarnings() const;
+		DATAFLOWCORE_API void Error(const FString& InError, const FDataflowNode* InNode = nullptr, const FDataflowOutput* InOutput = nullptr);
+		DATAFLOWCORE_API int32 GetNumErrors() const;
+
+		DATAFLOWCORE_API bool NodeHasWarning(const FDataflowNode* InNode) { return NodesWithWarning.Contains(InNode); };
+		DATAFLOWCORE_API bool NodeHasError(const FDataflowNode* InNode) { return NodesWithError.Contains(InNode); };
+		DATAFLOWCORE_API bool NodeFailed(const FDataflowNode* InNode) { return NodesFailed.Contains(InNode); };
+
+		DATAFLOWCORE_API void ClearNodesData();
+
+		DECLARE_MULTICAST_DELEGATE_TwoParams(FOnNodeBeginEvaluateMulticast, const FDataflowNode* Node, const FDataflowOutput* Output);
+		FOnNodeBeginEvaluateMulticast OnNodeBeginEvaluateMulticast;
+
+		DECLARE_MULTICAST_DELEGATE_TwoParams(FOnNodeFinishEvaluateMulticast, const FDataflowNode* Node, const FDataflowOutput* Output);
+		FOnNodeFinishEvaluateMulticast OnNodeFinishEvaluateMulticast;
+
+		DECLARE_MULTICAST_DELEGATE_ThreeParams(FOnContextHasInfoMulticast, const FDataflowNode* Node, const FDataflowOutput* Output, const FString& Info);
+		FOnContextHasInfoMulticast OnContextHasInfo;
+
+		DECLARE_MULTICAST_DELEGATE_ThreeParams(FOnContextHasWarningMulticast, const FDataflowNode* Node, const FDataflowOutput* Output, const FString& Info);
+		FOnContextHasWarningMulticast OnContextHasWarning;
+
+		DECLARE_MULTICAST_DELEGATE_ThreeParams(FOnContextHasErrorMulticast, const FDataflowNode* Node, const FDataflowOutput* Output, const FString& Info);
+		FOnContextHasErrorMulticast OnContextHasError;
+
 	protected:
 		DATAFLOWCORE_API void BeginContextEvaluation(const FDataflowNode* Node, const FDataflowOutput* Output);
 
@@ -496,6 +530,10 @@ namespace UE::Dataflow
 #if DATAFLOW_EDITOR_EVALUATION
 		TArray<const FDataflowConnection*> Callstack;
 #endif
+		TSet<const FDataflowNode*> NodesWithInfo;
+		TSet<const FDataflowNode*> NodesWithWarning;
+		TSet<const FDataflowNode*> NodesWithError;
+		TSet<const FDataflowNode*> NodesFailed;
 	};
 
 	struct FContextScopedCallstack
@@ -538,7 +576,7 @@ namespace UE::Dataflow
 
 		virtual bool HasDataImpl(FContextCacheKey Key, FTimestamp InTimestamp = FTimestamp::Invalid) const override
 		{
-			return DataStore.Contains(Key) && DataStore[Key]->GetTimestamp() >= InTimestamp;
+			return DataStore.Contains(Key) && (DataStore[Key]->GetTimestamp() >= InTimestamp);
 		}
 
 		virtual bool IsEmptyImpl() const override
diff --git a/Engine/Source/Runtime/Experimental/Dataflow/Core/Public/Dataflow/DataflowPath.h b/Engine/Source/Runtime/Experimental/Dataflow/Core/Public/Dataflow/DataflowPath.h
new file mode 100644
--- /dev/null
+++ b/Engine/Source/Runtime/Experimental/Dataflow/Core/Public/Dataflow/DataflowPath.h
@@ -0,0 +1,39 @@
+// Copyright Epic Games, Inc. All Rights Reserved.
+
+#pragma once
+
+#include "DataflowPath.generated.h"
+
+USTRUCT()
+struct FDataflowPath final
+{
+	GENERATED_USTRUCT_BODY()
+
+	DATAFLOWCORE_API FDataflowPath() = default;
+
+	DATAFLOWCORE_API explicit FDataflowPath(const FString& InGraph, const FString& InNode, const FString& InInput = TEXT(""), const FString& InOutput = TEXT(""))
+		: Graph(InGraph)
+		, Node(InNode)
+		, Input(InInput)
+		, Output(InOutput)
+	{}
+
+	DATAFLOWCORE_API const FString& GetGraph() const { return Graph; }
+	DATAFLOWCORE_API const FString& GetNode() const { return Node; }
+	DATAFLOWCORE_API const FString& GetInput() const { return Input; }
+	DATAFLOWCORE_API const FString& GetOutput() const { return Output; }
+	DATAFLOWCORE_API void SetGraph(const FString InGraph) { Graph = InGraph; }
+	DATAFLOWCORE_API void SetNode(const FString InNode) { Node = InNode; }
+	DATAFLOWCORE_API void SetInput(const FString InInput) { Input = InInput; }
+	DATAFLOWCORE_API void SetOutput(const FString InOutput) { Output = InOutput; }
+	DATAFLOWCORE_API FString ToString() const;
+	DATAFLOWCORE_API void DecodePath(const FString InPath);
+	DATAFLOWCORE_API bool PathHasInput() const;
+	DATAFLOWCORE_API bool PathHasOutput() const;
+
+private:
+	FString Graph;
+	FString Node;
+	FString Input;
+	FString Output;
+};
diff --git a/Engine/Source/Runtime/Experimental/Dataflow/Engine/Private/Dataflow/DataflowEdNode.cpp b/Engine/Source/Runtime/Experimental/Dataflow/Engine/Private/Dataflow/DataflowEdNode.cpp
--- a/Engine/Source/Runtime/Experimental/Dataflow/Engine/Private/Dataflow/DataflowEdNode.cpp
+++ b/Engine/Source/Runtime/Experimental/Dataflow/Engine/Private/Dataflow/DataflowEdNode.cpp
@@ -970,6 +970,38 @@ TArray<UE::Dataflow::FRenderingParameter> UDataflowEdNode::GetRenderParameters()
 	return 	TArray<UE::Dataflow::FRenderingParameter>();
 }
 
+void UDataflowEdNode::RegisterDelegateHandle()
+{
+#if WITH_EDITOR
+	if (TSharedPtr<FDataflowNode> DataflowNode = GetDataflowNode())
+	{
+		OnNodeInvalidatedDelegateHandle = DataflowNode->GetOnNodeInvalidatedDelegate().AddWeakLambda(this, [this](FDataflowNode* InDataflowNode)
+			{
+				GetGraph()->NotifyNodeChanged(this);
+			});
+	}
+#endif
+}
+
+void UDataflowEdNode::UnRegisterDelegateHandle()
+{
+	if (TSharedPtr<FDataflowNode> DataflowNode = GetDataflowNode())
+	{
+		if (DataflowNode->GetOnNodeInvalidatedDelegate().IsBound() && OnNodeInvalidatedDelegateHandle.IsValid())
+		{
+			DataflowNode->GetOnNodeInvalidatedDelegate().Remove(OnNodeInvalidatedDelegateHandle);
+		}
+	}
+}
+
+void UDataflowEdNode::SetDataflowNodeGuid(FGuid InGuid)
+{
+	UnRegisterDelegateHandle();
+
+	DataflowNodeGuid = InGuid;
+
+	RegisterDelegateHandle();
+}
 
 #undef LOCTEXT_NAMESPACE
 
diff --git a/Engine/Source/Runtime/Experimental/Dataflow/Engine/Public/Dataflow/DataflowEdNode.h b/Engine/Source/Runtime/Experimental/Dataflow/Engine/Public/Dataflow/DataflowEdNode.h
--- a/Engine/Source/Runtime/Experimental/Dataflow/Engine/Public/Dataflow/DataflowEdNode.h
+++ b/Engine/Source/Runtime/Experimental/Dataflow/Engine/Public/Dataflow/DataflowEdNode.h
@@ -31,6 +31,7 @@ class UDataflowEdNode : public UEdGraphNode
 	TSharedPtr<UE::Dataflow::FGraph> DataflowGraph;
 
 public:
+	virtual ~UDataflowEdNode() { UnRegisterDelegateHandle(); };
 
 	// UEdGraphNode interface
 	DATAFLOWENGINE_API virtual void AllocateDefaultPins();
@@ -67,7 +68,7 @@ public:
 	DATAFLOWENGINE_API void UpdatePinsConnectionsFromDataflowNode();
 
 	FGuid GetDataflowNodeGuid() const { return DataflowNodeGuid; }
-	void SetDataflowNodeGuid(FGuid InGuid) { DataflowNodeGuid = InGuid; }
+	DATAFLOWENGINE_API void SetDataflowNodeGuid(FGuid InGuid);
 
 	DATAFLOWENGINE_API TSharedPtr<FDataflowNode> GetDataflowNode();
 	DATAFLOWENGINE_API TSharedPtr<const FDataflowNode> GetDataflowNode() const;
@@ -102,6 +103,9 @@ public:
 
 	DATAFLOWENGINE_API TArray<UE::Dataflow::FRenderingParameter> GetRenderParameters() const;
 
+	DATAFLOWENGINE_API void RegisterDelegateHandle();
+	DATAFLOWENGINE_API void UnRegisterDelegateHandle();
+
 private:
 
 	UPROPERTY(Transient)
@@ -114,5 +118,7 @@ private:
 	bool bCanEnableRenderWireframe = true;
 
 	void RemoveSpacesInAllPinTypes();
+
+	FDelegateHandle OnNodeInvalidatedDelegateHandle;
 };
 
