Dataflow - Removed (Experimental) from the default node names of experimental names.

#rnx

diff --git a/Engine/Plugins/Experimental/Dataflow/Source/DataflowEditor/Private/Dataflow/DataflowGraphSchemaAction.cpp b/Engine/Plugins/Experimental/Dataflow/Source/DataflowEditor/Private/Dataflow/DataflowGraphSchemaAction.cpp
--- a/Engine/Plugins/Experimental/Dataflow/Source/DataflowEditor/Private/Dataflow/DataflowGraphSchemaAction.cpp
+++ b/Engine/Plugins/Experimental/Dataflow/Source/DataflowEditor/Private/Dataflow/DataflowGraphSchemaAction.cpp
@@ -34,16 +34,15 @@ TSharedPtr<FAssetSchemaAction_Dataflow_CreateNode_DataflowEdNode> FAssetSchemaAc
 				if ((bIsSimulationGraph && bIsSimulationNode) || (!bIsSimulationGraph && !bIsSimulationNode))
 				{
 					const FText ToolTip = FText::FromString(Param.ToolTip.IsEmpty() ? FString("Add a Dataflow node.") : Param.ToolTip);
-					FText NodeName = FText::FromString(Param.DisplayName.ToString());
-					if (!InOverrideNodeName.IsNone())
-					{
-						NodeName = FText::FromName(InOverrideNodeName);
-					}
+					const FName NodeName = !InOverrideNodeName.IsNone() ? InOverrideNodeName : Param.DisplayName;
+					const FText MenuDesc = Param.bIsExperimental ?
+						FText::FromString(NodeName.ToString() + TEXT(" (Experimental)")) :
+						FText::FromName(NodeName);
 
 					const FText Category = FText::FromString(Param.Category.ToString().IsEmpty() ? FString("Dataflow") : Param.Category.ToString());
 					const FText Tags = FText::FromString(Param.Tags);
 					TSharedPtr<FAssetSchemaAction_Dataflow_CreateNode_DataflowEdNode> NewNodeAction(
-						new FAssetSchemaAction_Dataflow_CreateNode_DataflowEdNode(InNodeTypeName, Category, NodeName, ToolTip, Tags));
+						new FAssetSchemaAction_Dataflow_CreateNode_DataflowEdNode(NodeName, InNodeTypeName, Category, MenuDesc, ToolTip, Tags));
 					return NewNodeAction;
 				}
 			}
@@ -55,8 +54,7 @@ TSharedPtr<FAssetSchemaAction_Dataflow_CreateNode_DataflowEdNode> FAssetSchemaAc
 UEdGraphNode* FAssetSchemaAction_Dataflow_CreateNode_DataflowEdNode::PerformAction(class UEdGraph* ParentGraph, UEdGraphPin* FromPin, const FVector2D Location, bool bSelectNewNode)
 {
 	// by default use the type name 
-	const FString NodeBaseName = GetMenuDescription().ToString();
-	UEdGraphNode* NewNode = UE::Dataflow::FEditAssetUtils::AddNewNode(ParentGraph, Location, FName(NodeBaseName), NodeTypeName, FromPin);
+	UEdGraphNode* NewNode = UE::Dataflow::FEditAssetUtils::AddNewNode(ParentGraph, Location, NodeName, NodeTypeName, FromPin);
 
 	if (NewNode && bSelectNewNode)
 	{
diff --git a/Engine/Plugins/Experimental/Dataflow/Source/DataflowEditor/Public/Dataflow/DataflowGraphSchemaAction.h b/Engine/Plugins/Experimental/Dataflow/Source/DataflowEditor/Public/Dataflow/DataflowGraphSchemaAction.h
--- a/Engine/Plugins/Experimental/Dataflow/Source/DataflowEditor/Public/Dataflow/DataflowGraphSchemaAction.h
+++ b/Engine/Plugins/Experimental/Dataflow/Source/DataflowEditor/Public/Dataflow/DataflowGraphSchemaAction.h
@@ -38,13 +38,14 @@ struct DATAFLOWEDITOR_API FAssetSchemaAction_Dataflow_CreateNode_DataflowEdNode
 public:
 	FAssetSchemaAction_Dataflow_CreateNode_DataflowEdNode() : FEdGraphSchemaAction() {}
 
-	FAssetSchemaAction_Dataflow_CreateNode_DataflowEdNode(const FName& InType, FText InNodeCategory, FText InMenuDesc, FText InToolTip, FText InKeywords)
-		: FEdGraphSchemaAction(InNodeCategory, InMenuDesc, InToolTip, 0, InKeywords), NodeTypeName(InType) {}
+	FAssetSchemaAction_Dataflow_CreateNode_DataflowEdNode(const FName& InName, const FName& InType, FText InNodeCategory, FText InMenuDesc, FText InToolTip, FText InKeywords)
+		: FEdGraphSchemaAction(InNodeCategory, InMenuDesc, InToolTip, 0, InKeywords), NodeName(InName), NodeTypeName(InType) {}
 
 	static TSharedPtr<FAssetSchemaAction_Dataflow_CreateNode_DataflowEdNode> CreateAction(const UEdGraph* ParentGraph, const FName& NodeTypeName, const FName& InOverrideNodeName = NAME_None);
 
 	virtual UEdGraphNode* PerformAction(class UEdGraph* ParentGraph, UEdGraphPin* FromPin, const FVector2D Location, bool bSelectNewNode = true) override;
 
+	FName NodeName;
 	FName NodeTypeName;
 };
 
diff --git a/Engine/Source/Runtime/Experimental/Dataflow/Core/Private/Dataflow/DataflowNodeFactory.cpp b/Engine/Source/Runtime/Experimental/Dataflow/Core/Private/Dataflow/DataflowNodeFactory.cpp
--- a/Engine/Source/Runtime/Experimental/Dataflow/Core/Private/Dataflow/DataflowNodeFactory.cpp
+++ b/Engine/Source/Runtime/Experimental/Dataflow/Core/Private/Dataflow/DataflowNodeFactory.cpp
@@ -340,13 +340,6 @@ namespace UE::Dataflow
 				// Do not display version in DisplayName
 				//
 				ParametersArray[0].DisplayName = GetDisplayNameNoVersion(ParametersArray[0].DisplayName);
-				
-				if (ParametersArray[0].bIsExperimental)
-				{
-					FString NewDisplayName = ParametersArray[0].DisplayName.ToString() + " (Experimental)";
-					ParametersArray[0].DisplayName = FName(*NewDisplayName);
-				}
-
 				RetVal.Add(ParametersArray[0]);
 			}
 			else if (ParametersArray.Num() > 1)
@@ -354,18 +347,13 @@ namespace UE::Dataflow
 				// 
 				// There are multiple versions of the node
 				// DisplayName: DisplayName (v2)
-				//
+				// TODO: Should this be removed? 
+				//       Adding (v2) to the name instead of the menu text could cause mismatched name issues.
+				//       Previous versions should already be deprecated anyway (which seems to be the standard practice).
 				for (FFactoryParameters& FactoryParameters : ParametersArray)
 				{
-					FString NewDisplayName = GetDisplayNameNoVersion(FactoryParameters.DisplayName).ToString() + " (v" + FString::FromInt(GetNumVersionFromVersion(FactoryParameters.NodeVersion)) + ")";
+					const FString NewDisplayName = GetDisplayNameNoVersion(FactoryParameters.DisplayName).ToString() + " (v" + FString::FromInt(GetNumVersionFromVersion(FactoryParameters.NodeVersion)) + ")";
 					FactoryParameters.DisplayName = FName(*NewDisplayName);
-
-					if (FactoryParameters.bIsExperimental)
-					{
-						NewDisplayName = FactoryParameters.DisplayName.ToString() + " (Experimental)";
-						FactoryParameters.DisplayName = FName(*NewDisplayName);
-					}
-
 					RetVal.Add(FactoryParameters);
 				}
 			}
