Dataflow Editor: fix crash when changing selected nodes with the paint tool active. Calling ToolsContext->EndTool schedules closing the tool for the next tick, whereas calling ToolManager->PostActiveToolShutdownRequest requeste the tool close immediately.

diff --git a/Engine/Plugins/Experimental/Dataflow/Source/DataflowEditor/Private/Dataflow/DataflowEditorToolkit.cpp b/Engine/Plugins/Experimental/Dataflow/Source/DataflowEditor/Private/Dataflow/DataflowEditorToolkit.cpp
--- a/Engine/Plugins/Experimental/Dataflow/Source/DataflowEditor/Private/Dataflow/DataflowEditorToolkit.cpp
+++ b/Engine/Plugins/Experimental/Dataflow/Source/DataflowEditor/Private/Dataflow/DataflowEditorToolkit.cpp
@@ -1068,11 +1068,16 @@ void FDataflowEditorToolkit::OnNodeSelectionChanged(const TSet<UObject*>& InNewS
 			if( UDataflowEditorMode* const DataflowMode = Cast<UDataflowEditorMode>(EditorModeManager->GetActiveScriptableMode(UDataflowEditorMode::EM_DataflowEditorModeId)) )
 			{
 				// Close any running tool. OnNodeSingleClicked() will start a new tool if a new node was clicked.
+
 				UEditorInteractiveToolsContext* const ToolsContext = DataflowMode->GetInteractiveToolsContext();
 				checkf(ToolsContext, TEXT("No valid ToolsContext found for FDataflowEditorToolkit"));
-				if (ToolsContext->HasActiveTool())
+
+				UInteractiveToolManager* const ToolManager = ToolsContext->ToolManager;
+				checkf(ToolManager, TEXT("ToolsContext doesn't have a valid ToolManager"));
+
+				if (UInteractiveTool* const ActiveTool = ToolManager->GetActiveTool(EToolSide::Left))
 				{
-					ToolsContext->EndTool(EToolShutdownType::Completed);
+					ToolManager->PostActiveToolShutdownRequest(ActiveTool, EToolShutdownType::Accept);
 				}
 
 				EditorContent->SetSelectedNode(PrimarySelection);
