Dataflow : SubGraphs evaluations now properly report error, warning and info
The entries of the log are now showing the name of the subgraph is needed
Clicking on the log also properly opens subgraph to foxus on the relevant node

diff --git a/Engine/Plugins/Experimental/Dataflow/Source/DataflowEditor/Private/Dataflow/DataflowEditorToolkit.cpp b/Engine/Plugins/Experimental/Dataflow/Source/DataflowEditor/Private/Dataflow/DataflowEditorToolkit.cpp
--- a/Engine/Plugins/Experimental/Dataflow/Source/DataflowEditor/Private/Dataflow/DataflowEditorToolkit.cpp
+++ b/Engine/Plugins/Experimental/Dataflow/Source/DataflowEditor/Private/Dataflow/DataflowEditorToolkit.cpp
@@ -347,18 +347,7 @@ FDataflowEditorToolkit::~FDataflowEditorToolkit()
 		GraphEditor->OnNodeDeletedMulticast.Remove(OnNodeDeletedMulticastDelegateHandle);
 	}
 
-	const TObjectPtr<UDataflowBaseContent>& EditorContent = GetEditorContent();
-	check(EditorContent);
-	
-	if (TSharedPtr<UE::Dataflow::FEngineContext> DataflowContext = EditorContent->GetDataflowContext())
-	{
-		DataflowContext->OnNodeBeginEvaluateMulticast.Remove(OnNodeBeginEvaluateMulticastDelegateHandle);
-		DataflowContext->OnNodeFinishEvaluateMulticast.Remove(OnNodeFinishEvaluateMulticastDelegateHandle);
-
-		DataflowContext->OnContextHasInfo.Remove(OnContextHasInfoDelegateHandle);
-		DataflowContext->OnContextHasWarning.Remove(OnContextHasWarningDelegateHandle);
-		DataflowContext->OnContextHasError.Remove(OnContextHasErrorDelegateHandle);
-	}
+	UnregisterContextHandlers();
 
 	if (NodeDetailsEditor)
 	{
@@ -382,6 +371,38 @@ FDataflowEditorToolkit::~FDataflowEditorToolkit()
 	SimulationModeManager->DestroyMode(UDataflowEditorMode::EM_DataflowEditorModeId);
 }
 
+void FDataflowEditorToolkit::RegisterContextHandlers()
+{
+	const TObjectPtr<UDataflowBaseContent>& EditorContent = GetEditorContent();
+	check(EditorContent);
+
+	if (TSharedPtr<UE::Dataflow::FEngineContext> DataflowContext = EditorContent->GetDataflowContext())
+	{
+		OnNodeBeginEvaluateMulticastDelegateHandle = DataflowContext->OnNodeBeginEvaluateMulticast.AddSP(this, &FDataflowEditorToolkit::OnNodeBeginEvaluate);
+		OnNodeFinishEvaluateMulticastDelegateHandle = DataflowContext->OnNodeFinishEvaluateMulticast.AddSP(this, &FDataflowEditorToolkit::OnNodeFinishEvaluate);
+
+		OnContextHasInfoDelegateHandle = DataflowContext->OnContextHasInfo.AddSP(this, &FDataflowEditorToolkit::OnContextHasInfo);
+		OnContextHasWarningDelegateHandle = DataflowContext->OnContextHasWarning.AddSP(this, &FDataflowEditorToolkit::OnContextHasWarning);
+		OnContextHasErrorDelegateHandle = DataflowContext->OnContextHasError.AddSP(this, &FDataflowEditorToolkit::OnContextHasError);
+	}
+}
+
+void FDataflowEditorToolkit::UnregisterContextHandlers()
+{
+	const TObjectPtr<UDataflowBaseContent>& EditorContent = GetEditorContent();
+	check(EditorContent);
+
+	if (TSharedPtr<UE::Dataflow::FEngineContext> DataflowContext = EditorContent->GetDataflowContext())
+	{
+		DataflowContext->OnNodeBeginEvaluateMulticast.Remove(OnNodeBeginEvaluateMulticastDelegateHandle);
+		DataflowContext->OnNodeFinishEvaluateMulticast.Remove(OnNodeFinishEvaluateMulticastDelegateHandle);
+
+		DataflowContext->OnContextHasInfo.Remove(OnContextHasInfoDelegateHandle);
+		DataflowContext->OnContextHasWarning.Remove(OnContextHasWarningDelegateHandle);
+		DataflowContext->OnContextHasError.Remove(OnContextHasErrorDelegateHandle);
+	}
+}
+
 void FDataflowEditorToolkit::CreateEditorModeManager()
 {
 	// Setup the construction manager / scene
@@ -1222,14 +1243,36 @@ void FDataflowEditorToolkit::OnNodeFinishEvaluate(const FDataflowNode* Node, con
 	}
 }
 
+void FDataflowEditorToolkit::SetDataflowPathFromNodeAndOutput(const FDataflowNode* Node, const FDataflowOutput* Output, FDataflowPath& OutPath) const
+{
+	if (const TObjectPtr<UDataflowBaseContent>& EditorContent = GetEditorContent())
+	{
+		if (const TObjectPtr<UDataflow> DataflowAsset = EditorContent->GetDataflowAsset())
+		{
+			FName GraphName = GetGraphLogName();
+			if (Node)
+			{
+				if (const UDataflowEdNode* EdNode = DataflowAsset->FindEdNodeByDataflowNodeGuid(Node->GetGuid()))
+				{
+					if (const UDataflowSubGraph* SubGraph = Cast<UDataflowSubGraph>(EdNode->GetGraph()))
+					{
+						GraphName = SubGraph->GetFName();
+					}
+				}
+			}
+			OutPath.SetGraph(GraphName.ToString());
+			OutPath.SetNode(Node ? Node->GetName().ToString() : FString());
+			OutPath.SetOutput(Output ? Output->GetName().ToString() : FString());
+		}
+	}
+}
+
 void FDataflowEditorToolkit::LogMessage(const EMessageSeverity::Type Severity, const FDataflowNode* Node, const FDataflowOutput* Output, const FString& Message) const
 {
 	if (DataflowOutputLog)
 	{
-		const FString NodeName = Node ? Node->GetName().ToString() : FString();
-		const FString OutputName = Output ? Output->GetName().ToString() : FString();
-
-		FDataflowPath Path(GetGraphLogName().ToString(), NodeName, FString(), OutputName);
+		FDataflowPath Path;
+		SetDataflowPathFromNodeAndOutput(Node, Output, Path);
 		DataflowOutputLog->AddMessage(Severity, Message, Path);
 	}
 }
@@ -1327,19 +1370,41 @@ void FDataflowEditorToolkit::OnOutputLogMessageTokenClicked(const FString TokenS
 		Output = Path.GetOutput();
 	}
 
-	// Clear node selection
-	GraphEditor->ClearSelectionSet();
+	TSharedPtr<SDataflowGraphEditor> ActiveGraphEditor = GraphEditor;
+
 
 	// Select node from TokenString
 	if (UDataflow* DataflowAsset = GetEditorContent()->GetDataflowAsset())
 	{
+		const FName GraphName(Path.GetGraph());
+		if (UDataflowSubGraph* SubGraph = DataflowAsset->FindSubGraphByName(GraphName))
+		{
+			OpenSubGraphTab(GraphName);
+			ActiveGraphEditor = ActiveSubGraphEditorWeakPtr.Pin();
+		}
+		else
+		{
+			// TODO : try to find a way to focus back on the main graph
+			//			current issue seem to be related to the tabspawner beeing unknown when calling thos method
+			//			maybe a better way would be to leverage the subgraph spawner and show the main graph in the subgraph list 
+			//			so that all graphs are spawned and treated the same way ( but this may be an issue to keep the pointer of the main graph )
+			//if (TSharedPtr<SDockTab> MainGraphTab = FGlobalTabmanager::Get()->TryInvokeTab(GraphCanvasTabId))
+			//{
+			//	MainGraphTab->ActivateInParent(SetDirectly);
+			//}
+		}
+
+		// Clear node selection
+		ActiveGraphEditor->ClearSelectionSet();
+
+		// now jump to the relevant node 
 		if (TSharedPtr<UE::Dataflow::FGraph> DataflowGraph = DataflowAsset->GetDataflow())
 		{
 			if (TSharedPtr<FDataflowNode> DataflowNode = DataflowGraph->FindBaseNode(FName(*NodeName)))
 			{
 				if (TObjectPtr<UDataflowEdNode> DataflowUEdNode = DataflowAsset->FindEdNodeByDataflowNodeGuid(DataflowNode->GetGuid()))
 				{
-					GraphEditor->JumpToNode(DataflowUEdNode);
+					ActiveGraphEditor->JumpToNode(DataflowUEdNode);
 				}
 			}
 		}
@@ -1489,18 +1554,8 @@ TSharedRef<SDataflowGraphEditor> FDataflowEditorToolkit::CreateGraphEditorWidget
 	OnSelectionChangedMulticastDelegateHandle = NewGraphEditor->OnSelectionChangedMulticast.AddSP(this, &FDataflowEditorToolkit::OnNodeSelectionChanged);
 	OnNodeDeletedMulticastDelegateHandle = NewGraphEditor->OnNodeDeletedMulticast.AddSP(this, &FDataflowEditorToolkit::OnNodeDeleted);
 	
-	const TObjectPtr<UDataflowBaseContent>& EditorContent = GetEditorContent();
-	check(EditorContent);
-
-	TSharedPtr<FEngineContext>& Context = EditorContent->GetDataflowContext();
-	check(Context);
-
-	OnNodeBeginEvaluateMulticastDelegateHandle = Context->OnNodeBeginEvaluateMulticast.AddSP(this, &FDataflowEditorToolkit::OnNodeBeginEvaluate);
-	OnNodeFinishEvaluateMulticastDelegateHandle = Context->OnNodeFinishEvaluateMulticast.AddSP(this, &FDataflowEditorToolkit::OnNodeFinishEvaluate);
-
-	OnContextHasInfoDelegateHandle = Context->OnContextHasInfo.AddSP(this, &FDataflowEditorToolkit::OnContextHasInfo);
-	OnContextHasWarningDelegateHandle = Context->OnContextHasWarning.AddSP(this, &FDataflowEditorToolkit::OnContextHasWarning);
-	OnContextHasErrorDelegateHandle = Context->OnContextHasError.AddSP(this, &FDataflowEditorToolkit::OnContextHasError);
+	UnregisterContextHandlers();
+	RegisterContextHandlers();
 
 	return NewGraphEditor;
 }
diff --git a/Engine/Plugins/Experimental/Dataflow/Source/DataflowEditor/Public/Dataflow/DataflowEditorToolkit.h b/Engine/Plugins/Experimental/Dataflow/Source/DataflowEditor/Public/Dataflow/DataflowEditorToolkit.h
--- a/Engine/Plugins/Experimental/Dataflow/Source/DataflowEditor/Public/Dataflow/DataflowEditorToolkit.h
+++ b/Engine/Plugins/Experimental/Dataflow/Source/DataflowEditor/Public/Dataflow/DataflowEditorToolkit.h
@@ -33,6 +33,8 @@ class UDataflowEditor;
 class SDataflowMembersWidget;
 class FDataflowSimulationSceneProfileIndexStorage;
 class FDataflowOutputLog;
+struct FDataflowPath;
+
 namespace UE::Dataflow
 {
 	class FDataflowNodeDetailExtensionHandler;
@@ -222,6 +224,11 @@ private:
 	/** Create the simulation viewport client */
 	void CreateSimulationViewportClient();
 
+	void SetDataflowPathFromNodeAndOutput(const FDataflowNode* Node, const FDataflowOutput* Output, FDataflowPath& OutPath) const;
+
+	void RegisterContextHandlers();
+	void UnregisterContextHandlers();
+
 	// List of editor commands used  for the dataflow asset
 	TSharedPtr<FUICommandList> GraphEditorCommands;
 
diff --git a/Engine/Source/Runtime/Experimental/Dataflow/Engine/Private/Dataflow/DataflowSubGraphNodes.cpp b/Engine/Source/Runtime/Experimental/Dataflow/Engine/Private/Dataflow/DataflowSubGraphNodes.cpp
--- a/Engine/Source/Runtime/Experimental/Dataflow/Engine/Private/Dataflow/DataflowSubGraphNodes.cpp
+++ b/Engine/Source/Runtime/Experimental/Dataflow/Engine/Private/Dataflow/DataflowSubGraphNodes.cpp
@@ -51,8 +51,28 @@ namespace UE::Dataflow::SubGraphNodes::Private
 			return Connection.EvaluateImpl(*this);
 		}
 
-		bool EvaluateSubGraph(const FDataflowSubGraphOutputNode& SubGraphOutputNode, const FDataflowOutput& CallerOutput)
+		bool EvaluateSubGraph(const FDataflowSubGraphOutputNode& SubGraphOutputNode, const FDataflowNode& CallerNode, const FDataflowOutput& CallerOutput)
 		{
+			auto OnInfo = [this](const FDataflowNode* Node, const FDataflowOutput* Output, const FString& Info)
+				{
+					ParentContext.Info(Info, Node, Output);
+				};
+
+			bool bHasWarning = false;
+			auto OnWarning = [this, &bHasWarning](const FDataflowNode* Node, const FDataflowOutput* Output, const FString& Warning)
+				{
+					ParentContext.Warning(Warning, Node, Output);
+					bHasWarning = true;
+				};
+
+			bool bHasError = false;
+			auto OnError = [this, &bHasError](const FDataflowNode* Node, const FDataflowOutput* Output, const FString& Error)
+				{
+					ParentContext.Error(Error, Node, Output);
+					bHasError = true;
+				};
+
+			bool bEvaluationSuccess = false;
 			if (const FDataflowInput* InputToPull = SubGraphOutputNode.FindInput(CallerOutput.GetName()))
 			{
 				if (const FDataflowOutput* OutputToEvaluate = InputToPull->GetConnection())
@@ -60,6 +80,10 @@ namespace UE::Dataflow::SubGraphNodes::Private
 					// we need to make sure the graph output and caller output are of compatible type
 					if (UE::Dataflow::AreTypesCompatible(OutputToEvaluate->GetType(), CallerOutput.GetType()))
 					{
+						const FDelegateHandle OnInfoHandle = OnContextHasInfo.AddLambda(OnInfo);
+						const FDelegateHandle OnWarningHandle = OnContextHasWarning.AddLambda(OnWarning);
+						const FDelegateHandle OnErrorHandle = OnContextHasError.AddLambda(OnError);
+
 						if (Evaluate(*OutputToEvaluate))
 						{
 							// copy the cache value from this cachestore to the parent one 
@@ -68,13 +92,28 @@ namespace UE::Dataflow::SubGraphNodes::Private
 							{
 								TUniquePtr<FContextCacheElementBase> ClonedCacheEntry = (*CacheEntryToClone)->Clone(*this);
 								ParentContext.SetDataImpl(CallerOutput.CacheKey(), MoveTemp(ClonedCacheEntry));
-								return true;
+								bEvaluationSuccess = true;
 							}
 						}
+
+						if (bHasError)
+						{
+							const FString ErrorMessage = TEXT("Subgraph call evaluation returned errors : see the details above");
+							ParentContext.Error(ErrorMessage, &CallerNode, &CallerOutput);
+						}
+						if (bHasWarning)
+						{
+							const FString ErrorMessage = TEXT("Subgraph call evaluation returned warning errors : see the details above");
+							ParentContext.Warning(ErrorMessage, &CallerNode, &CallerOutput);
+						}
+
+						OnContextHasInfo.Remove(OnInfoHandle);
+						OnContextHasWarning.Remove(OnWarningHandle);
+						OnContextHasError.Remove(OnErrorHandle);
 					}
 				}
 			}
-			return false;
+			return bEvaluationSuccess;
 		}
 
 		void EvaluateParentContext(const FDataflowOutput& SubGraphOutput)
@@ -668,7 +707,7 @@ void FDataflowCallSubGraphNode::Evaluate(UE::Dataflow::FContext& Context, const
 			{
 				// pulling the outputs and when the graph evaluate the remote left end of its graph OnInputValueRequested will be called
 				UE::Dataflow::SubGraphNodes::Private::FContextSubGraph SubGraphContext(Context, *this);
-				SubGraphContext.EvaluateSubGraph(*SubGraphOutputNode, *Out);
+				SubGraphContext.EvaluateSubGraph(*SubGraphOutputNode, *this, *Out);
 			}
 		}
 	}
