Dataflow : when variable are renamed , update the variable node  accordingly

diff --git a/Engine/Plugins/Experimental/Dataflow/Source/DataflowEditor/Private/Dataflow/DataflowAssetEditUtils.cpp b/Engine/Plugins/Experimental/Dataflow/Source/DataflowEditor/Private/Dataflow/DataflowAssetEditUtils.cpp
--- a/Engine/Plugins/Experimental/Dataflow/Source/DataflowEditor/Private/Dataflow/DataflowAssetEditUtils.cpp
+++ b/Engine/Plugins/Experimental/Dataflow/Source/DataflowEditor/Private/Dataflow/DataflowAssetEditUtils.cpp
@@ -11,6 +11,7 @@
 #include "Dataflow/DataflowSNode.h"
 #include "Dataflow/DataflowSubGraph.h"
 #include "Dataflow/DataflowSubGraphNodes.h"
+#include "Dataflow/DataflowVariableNodes.h"
 #include "Editor.h"
 #include "PropertyBagDetails.h"
 #include "ScopedTransaction.h"
@@ -281,6 +282,41 @@ namespace UE::Dataflow
 			}
 		}
 
+		void RenameVariableCallNodes(UDataflow& DataflowAsset, UEdGraph& EdGraph, FName VariableName, FName NewVariableName)
+		{
+			for (UEdGraphNode* EdNode : EdGraph.Nodes)
+			{
+				if (UDataflowEdNode* DataflowEdNode = Cast<UDataflowEdNode>(EdNode))
+				{
+					if (TSharedPtr<FDataflowNode> DataflowNode = DataflowEdNode->GetDataflowNode())
+					{
+						if (FGetDataflowVariableNode* VariableNode = DataflowNode->AsType<FGetDataflowVariableNode>())
+						{
+							if (VariableNode->GetVariableName() == VariableName)
+							{
+								const FName UniqueName = GenerateUniqueObjectName(&DataflowAsset, NewVariableName);
+								DataflowEdNode->Rename(*UniqueName.ToString(), &EdGraph);
+								VariableNode->SetName(UniqueName);
+								VariableNode->SetVariable(&DataflowAsset, NewVariableName);
+							}
+						}
+					}
+				}
+			}
+		}
+
+		void RenameVariableCallNodes(UDataflow& DataflowAsset, FName VariableName, FName NewVariableName)
+		{
+			RenameVariableCallNodes(DataflowAsset, DataflowAsset, VariableName, NewVariableName);
+			for (UDataflowSubGraph* SubGraph : DataflowAsset.GetSubGraphs())
+			{
+				if (SubGraph)
+				{
+					RenameVariableCallNodes(DataflowAsset, *SubGraph, VariableName, NewVariableName);
+				}
+			}
+		}
+
 		void DeleteNodesNoTransaction(UEdGraph* EdGraph, const TArray<UEdGraphNode*>& NodesToDelete)
 		{
 			if (EdGraph && NodesToDelete.Num())
@@ -810,8 +846,12 @@ namespace UE::Dataflow
 						PropertyDesc.Name = NewVariableName;
 					};
 
-				const bool bModified = ModifyVariable(DataflowAsset, OldVariableName, ChangePropertyNameLambda);
-				return (bModified) ? EChangeResult::Changed : EChangeResult::None;
+				if (ModifyVariable(DataflowAsset, OldVariableName, ChangePropertyNameLambda))
+				{
+					RenameVariableCallNodes(DataflowAsset, OldVariableName, NewVariableName);
+					return EChangeResult::Changed;
+				}
+				return EChangeResult::None;
 			};
 
 		ChangeDataflowAssetWithTransaction(DataflowAsset, TransactionName, SetVariableNameInternal, GET_MEMBER_NAME_CHECKED(UDataflow, Variables));
diff --git a/Engine/Source/Runtime/Experimental/Dataflow/Engine/Private/Dataflow/DataflowVariableNodes.cpp b/Engine/Source/Runtime/Experimental/Dataflow/Engine/Private/Dataflow/DataflowVariableNodes.cpp
--- a/Engine/Source/Runtime/Experimental/Dataflow/Engine/Private/Dataflow/DataflowVariableNodes.cpp
+++ b/Engine/Source/Runtime/Experimental/Dataflow/Engine/Private/Dataflow/DataflowVariableNodes.cpp
@@ -26,6 +26,18 @@ FGetDataflowVariableNode::FGetDataflowVariableNode(const UE::Dataflow::FNodePara
 }
 
 FGetDataflowVariableNode::~FGetDataflowVariableNode()
+{
+	UnregisterHandlers();
+}
+
+void FGetDataflowVariableNode::RegisterHandlers()
+{
+#if WITH_EDITOR
+	FCoreUObjectDelegates::OnObjectPropertyChanged.AddRaw(this, &FGetDataflowVariableNode::OnObjectPropertyChanged);
+#endif
+}
+
+void FGetDataflowVariableNode::UnregisterHandlers()
 {
 #if WITH_EDITOR
 	FCoreUObjectDelegates::OnObjectPropertyChanged.RemoveAll(this);
@@ -69,9 +81,7 @@ void FGetDataflowVariableNode::OnObjectPropertyChanged(UObject* InObject, FPrope
 
 void FGetDataflowVariableNode::SetVariable(UDataflow* Dataflow, FName InVariableName)
 {
-#if WITH_EDITOR
-	FCoreUObjectDelegates::OnObjectPropertyChanged.RemoveAll(this);
-#endif
+	UnregisterHandlers();
 
 	WeakDataflowPtr = Dataflow;
 	EvaluateFunction = nullptr;
@@ -91,9 +101,7 @@ void FGetDataflowVariableNode::SetVariable(UDataflow* Dataflow, FName InVariable
 			Invalidate();
 		}
 
-#if WITH_EDITOR
-		FCoreUObjectDelegates::OnObjectPropertyChanged.AddRaw(this, &FGetDataflowVariableNode::OnObjectPropertyChanged);
-#endif
+		RegisterHandlers();
 	}
 }
 
diff --git a/Engine/Source/Runtime/Experimental/Dataflow/Engine/Public/Dataflow/DataflowVariableNodes.h b/Engine/Source/Runtime/Experimental/Dataflow/Engine/Public/Dataflow/DataflowVariableNodes.h
--- a/Engine/Source/Runtime/Experimental/Dataflow/Engine/Public/Dataflow/DataflowVariableNodes.h
+++ b/Engine/Source/Runtime/Experimental/Dataflow/Engine/Public/Dataflow/DataflowVariableNodes.h
@@ -58,6 +58,9 @@ private:
 	UPROPERTY()
 	FName VariableName;
 
+	void RegisterHandlers();
+	void UnregisterHandlers();
+
 	using FStaticEvaluationFunctionPtr = void(*)(const FGetDataflowVariableNode&, const FInstancedPropertyBag&, UE::Dataflow::FContext&, const FDataflowOutput&);
 	FStaticEvaluationFunctionPtr EvaluateFunction = nullptr;
 
