Dataflow : add support for drag and dropping assets on the dataflow graph and nodes
- Added a way to register a getter node for a specific type of asset
- Added an API to the dataflow node to allow node to indicate they support setting some asset properties from a asset drag and drop on them
- Added support for creating getters for Static mesh , SkeletalMesh, Skeleton,  MaterialInterface

diff --git a/Engine/Plugins/Experimental/Dataflow/Source/DataflowEditor/Private/Dataflow/DataflowSchema.cpp b/Engine/Plugins/Experimental/Dataflow/Source/DataflowEditor/Private/Dataflow/DataflowSchema.cpp
--- a/Engine/Plugins/Experimental/Dataflow/Source/DataflowEditor/Private/Dataflow/DataflowSchema.cpp
+++ b/Engine/Plugins/Experimental/Dataflow/Source/DataflowEditor/Private/Dataflow/DataflowSchema.cpp
@@ -8,6 +8,7 @@
 #include "Dataflow/DataflowEdNode.h"
 #include "Dataflow/DataflowGraphSchemaAction.h"
 #include "Dataflow/DataflowNodeFactory.h"
+#include "Dataflow/DataflowAssetEditUtils.h"
 #include "EdGraph/EdGraphNode.h"
 #include "EdGraph/EdGraph.h"
 #include "Framework/Commands/GenericCommands.h"
@@ -351,6 +352,15 @@ TSharedPtr<FDataflowNode> UDataflowSchema::GetDataflowNodeFromEdNode(UEdGraphNod
 	return {};
 }
 
+TSharedPtr<const FDataflowNode> UDataflowSchema::GetDataflowNodeFromEdNode(const UEdGraphNode* EdNode)
+{
+	if (const UDataflowEdNode* DataflowEdNode = Cast<UDataflowEdNode>(EdNode))
+	{
+		return DataflowEdNode->GetDataflowNode();
+	}
+	return {};
+}
+
 FDataflowConnection* UDataflowSchema::GetConnectionFromPin(const UEdGraphPin* Pin)
 {
 	FDataflowConnection* DataflowConnection = nullptr;
@@ -384,6 +394,98 @@ UE::Dataflow::FPin::EDirection UDataflowSchema::GetDirectionFromPinDirection(EEd
 	}
 }
 
+void UDataflowSchema::GetAssetsGraphHoverMessage(const TArray<FAssetData>&Assets, const UEdGraph * HoverGraph, FString & OutTooltipText, bool& OutOkIcon) const
+{
+	using FDataflowNodeFactory = UE::Dataflow::FNodeFactory;
+
+	OutOkIcon = false;
+	OutTooltipText = FString(TEXT("No supported node found"));
+	for (const FAssetData& AssetData : Assets)
+	{
+		if (UObject* AssetObject = AssetData.GetAsset())
+		{
+			const FName GetterNodeType = FDataflowNodeFactory::GetInstance()->GetGetterNodeFromAssetClass(*AssetObject->GetClass());
+			if (!GetterNodeType.IsNone())
+			{
+				OutOkIcon = true;
+				OutTooltipText = FString(TEXT("Place as a getter node here."));
+				break; // at least one is require to be a positive drop
+			}
+		}
+	}
+}
+void UDataflowSchema::DroppedAssetsOnGraph(const TArray<FAssetData>&Assets, const FVector2D & GraphPosition, UEdGraph * Graph) const
+{
+	using FDataflowNodeFactory = UE::Dataflow::FNodeFactory;
+
+	const FVector2D NodeOffset{ 0,100 };
+	FVector2D NodePosition = GraphPosition;
+
+	for (const FAssetData& AssetData: Assets)
+	{
+		if (UObject* AssetObject = AssetData.GetAsset())
+		{
+			const FName GetterNodeType = FDataflowNodeFactory::GetInstance()->GetGetterNodeFromAssetClass(*AssetObject->GetClass());
+			if (!GetterNodeType.IsNone())
+			{
+				if (UDataflowEdNode* DataflowEdNode = UE::Dataflow::FEditAssetUtils::AddNewNode(Graph, NodePosition, AssetObject->GetFName(), GetterNodeType, nullptr))
+				{
+					NodePosition += NodeOffset;
+					if (TSharedPtr<FDataflowNode> DataflowNode = DataflowEdNode->GetDataflowNode())
+					{
+						DataflowNode->SetAssetProperty(AssetObject);
+					}
+				}
+			}
+		}
+	}
+}
+
+void UDataflowSchema::GetAssetsNodeHoverMessage(const TArray<struct FAssetData>& Assets, const UEdGraphNode* HoverNode, FString& OutTooltipText, bool& OutOkIcon) const
+{
+	OutOkIcon = false;
+	OutTooltipText = FString(TEXT("Asset type unsupported by this node"));
+	if (TSharedPtr<const FDataflowNode> DataflowNode = GetDataflowNodeFromEdNode(HoverNode))
+	{
+		bool bNeedInvalidation = false;
+		for (const FAssetData& AssetData : Assets)
+		{
+			if (UObject* AssetObject = AssetData.GetAsset())
+			{
+				if (DataflowNode->SupportsAssetProperty(AssetObject))
+				{
+					OutOkIcon = true;
+					OutTooltipText = FString(TEXT("Set asset property on this node."));
+					break; // at least one is require to be a positive drop
+				}
+			}
+		}
+	}
+}
+
+void UDataflowSchema::DroppedAssetsOnNode(const TArray<struct FAssetData>& Assets, const FVector2D& GraphPosition, UEdGraphNode* Node) const
+{
+	if (TSharedPtr<FDataflowNode> DataflowNode = GetDataflowNodeFromEdNode(Node))
+	{
+		bool bNeedInvalidation = false;
+		for (const FAssetData& AssetData : Assets)
+		{
+			if (UObject* AssetObject = AssetData.GetAsset())
+			{
+				if (DataflowNode->SupportsAssetProperty(AssetObject))
+				{
+					DataflowNode->SetAssetProperty(AssetObject);
+					bNeedInvalidation = true;
+				}
+			}
+		}
+		if (bNeedInvalidation)
+		{
+			DataflowNode->Invalidate();
+		}
+	}
+}
+
 bool UDataflowSchema::SupportsDropPinOnNode(UEdGraphNode* InTargetNode, const FEdGraphPinType& InSourcePinType, EEdGraphPinDirection InSourcePinDirection, FText& OutMessage) const
 {
 	if (TSharedPtr<FDataflowNode> TargetDataflowNode = GetDataflowNodeFromEdNode(InTargetNode))
diff --git a/Engine/Plugins/Experimental/Dataflow/Source/DataflowEditor/Public/Dataflow/DataflowSchema.h b/Engine/Plugins/Experimental/Dataflow/Source/DataflowEditor/Public/Dataflow/DataflowSchema.h
--- a/Engine/Plugins/Experimental/Dataflow/Source/DataflowEditor/Public/Dataflow/DataflowSchema.h
+++ b/Engine/Plugins/Experimental/Dataflow/Source/DataflowEditor/Public/Dataflow/DataflowSchema.h
@@ -33,6 +33,12 @@ public:
 	virtual void BreakPinLinks(UEdGraphPin& TargetPin, bool bSendsNodeNotifcation) const override;
 	virtual bool TryCreateConnection(UEdGraphPin* PinA, UEdGraphPin* PinB) const override;
 
+	virtual void GetAssetsGraphHoverMessage(const TArray<FAssetData>& Assets, const UEdGraph* HoverGraph, FString& OutTooltipText, bool& OutOkIcon) const override;
+	virtual void DroppedAssetsOnGraph(const TArray<FAssetData>& Assets, const FVector2D& GraphPosition, UEdGraph* Graph) const override;
+
+	virtual void GetAssetsNodeHoverMessage(const TArray<struct FAssetData>& Assets, const UEdGraphNode* HoverNode, FString& OutTooltipText, bool& OutOkIcon) const override;
+	virtual void DroppedAssetsOnNode(const TArray<struct FAssetData>& Assets, const FVector2D& GraphPosition, UEdGraphNode* Node) const override;
+
 	virtual void SetPinBeingDroppedOnNode(UEdGraphPin* InSourcePin) const override { PinBeingDropped = InSourcePin; }
 	virtual bool SupportsDropPinOnNode(UEdGraphNode* InTargetNode, const FEdGraphPinType& InSourcePinType, EEdGraphPinDirection InSourcePinDirection, FText& OutMessage) const override;
 	virtual UEdGraphPin* DropPinOnNode(UEdGraphNode* InTargetNode, const FName& InSourcePinName, const FEdGraphPinType& InSourcePinType, EEdGraphPinDirection InSourcePinDirection) const override;
@@ -47,6 +53,7 @@ private:
 	static bool CanConnectPins(const UEdGraphPin& OutputPin, const UEdGraphPin& InputPin);
 	
 	static TSharedPtr<FDataflowNode> GetDataflowNodeFromEdNode(UEdGraphNode* EdNode);
+	static TSharedPtr<const FDataflowNode> GetDataflowNodeFromEdNode(const UEdGraphNode* EdNode);
 	static FDataflowConnection* GetConnectionFromPin(const UEdGraphPin* Pin);
 	static UE::Dataflow::FPin::EDirection GetDirectionFromPinDirection(EEdGraphPinDirection InPinDirection);
 
diff --git a/Engine/Plugins/Experimental/Dataflow/Source/DataflowNodes/Private/Dataflow/DataflowSkeletalMeshNodes.cpp b/Engine/Plugins/Experimental/Dataflow/Source/DataflowNodes/Private/Dataflow/DataflowSkeletalMeshNodes.cpp
--- a/Engine/Plugins/Experimental/Dataflow/Source/DataflowNodes/Private/Dataflow/DataflowSkeletalMeshNodes.cpp
+++ b/Engine/Plugins/Experimental/Dataflow/Source/DataflowNodes/Private/Dataflow/DataflowSkeletalMeshNodes.cpp
@@ -19,9 +19,14 @@ namespace UE::Dataflow
 		DATAFLOW_NODE_REGISTER_CREATION_FACTORY(FGetSkeletonDataflowNode);
 		DATAFLOW_NODE_REGISTER_CREATION_FACTORY(FSkeletalMeshBoneDataflowNode);
 		DATAFLOW_NODE_REGISTER_CREATION_FACTORY(FSkeletalMeshReferenceTransformDataflowNode);
+
+		DATAFLOW_NODE_REGISTER_GETTER_FOR_ASSET(USkeletalMesh, FGetSkeletalMeshDataflowNode);
+		DATAFLOW_NODE_REGISTER_GETTER_FOR_ASSET(USkeleton, FGetSkeletonDataflowNode);
 	}
 }
 
+///////////////////////////////////////////////////////////////////////////////////////////////////////////////////////////////////
+
 void FGetSkeletalMeshDataflowNode::Evaluate(UE::Dataflow::FContext& Context, const FDataflowOutput* Out) const
 {
 	typedef TObjectPtr<const USkeletalMesh> DataType;
@@ -43,6 +48,20 @@ void FGetSkeletalMeshDataflowNode::Evaluate(UE::Dataflow::FContext& Context, con
 	}
 }
 
+bool FGetSkeletalMeshDataflowNode::SupportsAssetProperty(UObject* Asset) const
+{
+	return (Cast<USkeletalMesh>(Asset) != nullptr);
+}
+
+void FGetSkeletalMeshDataflowNode::SetAssetProperty(UObject* Asset)
+{
+	if (USkeletalMesh* SkeletalMeshAsset = Cast<USkeletalMesh>(Asset))
+	{
+		SkeletalMesh = SkeletalMeshAsset;
+	}
+}
+
+///////////////////////////////////////////////////////////////////////////////////////////////////////////////////////////////////
 
 void FGetSkeletonDataflowNode::Evaluate(UE::Dataflow::FContext& Context, const FDataflowOutput* Out) const
 {
@@ -65,6 +84,21 @@ void FGetSkeletonDataflowNode::Evaluate(UE::Dataflow::FContext& Context, const F
 	}
 }
 
+bool FGetSkeletonDataflowNode::SupportsAssetProperty(UObject* Asset) const
+{
+	return (Cast<USkeleton>(Asset) != nullptr);
+}
+
+void FGetSkeletonDataflowNode::SetAssetProperty(UObject* Asset)
+{
+	if (USkeleton* SkeletonAsset = Cast<USkeleton>(Asset))
+	{
+		Skeleton = SkeletonAsset;
+	}
+}
+
+///////////////////////////////////////////////////////////////////////////////////////////////////////////////////////////////////
+
 void FSkeletalMeshBoneDataflowNode::Evaluate(UE::Dataflow::FContext& Context, const FDataflowOutput* Out) const
 {
 	typedef TObjectPtr<const USkeletalMesh> InDataType;
diff --git a/Engine/Plugins/Experimental/Dataflow/Source/DataflowNodes/Private/Dataflow/DataflowStaticMeshNodes.cpp b/Engine/Plugins/Experimental/Dataflow/Source/DataflowNodes/Private/Dataflow/DataflowStaticMeshNodes.cpp
--- a/Engine/Plugins/Experimental/Dataflow/Source/DataflowNodes/Private/Dataflow/DataflowStaticMeshNodes.cpp
+++ b/Engine/Plugins/Experimental/Dataflow/Source/DataflowNodes/Private/Dataflow/DataflowStaticMeshNodes.cpp
@@ -10,6 +10,7 @@ namespace UE::Dataflow
 	void RegisterStaticMeshNodes()
 	{
 		DATAFLOW_NODE_REGISTER_CREATION_FACTORY(FGetStaticMeshDataflowNode);
+		DATAFLOW_NODE_REGISTER_GETTER_FOR_ASSET(UStaticMesh, FGetStaticMeshDataflowNode);
 	}
 }
 
@@ -36,4 +37,16 @@ void FGetStaticMeshDataflowNode::Evaluate(UE::Dataflow::FContext& Context, const
 	}
 }
 
+bool FGetStaticMeshDataflowNode::SupportsAssetProperty(UObject* Asset) const
+{
+	return (Cast<UStaticMesh>(Asset) != nullptr);
+}
+
+void FGetStaticMeshDataflowNode::SetAssetProperty(UObject* Asset)
+{
+	if (UStaticMesh* StaticMeshAsset = Cast<UStaticMesh>(Asset))
+	{
+		StaticMesh = StaticMeshAsset;
+	}
+}
 
diff --git a/Engine/Plugins/Experimental/Dataflow/Source/DataflowNodes/Public/Dataflow/DataflowSkeletalMeshNodes.h b/Engine/Plugins/Experimental/Dataflow/Source/DataflowNodes/Public/Dataflow/DataflowSkeletalMeshNodes.h
--- a/Engine/Plugins/Experimental/Dataflow/Source/DataflowNodes/Public/Dataflow/DataflowSkeletalMeshNodes.h
+++ b/Engine/Plugins/Experimental/Dataflow/Source/DataflowNodes/Public/Dataflow/DataflowSkeletalMeshNodes.h
@@ -32,8 +32,11 @@ public:
 		RegisterOutputConnection(&SkeletalMesh);
 	}
 
-
 	virtual void Evaluate(UE::Dataflow::FContext& Context, const FDataflowOutput* Out) const override;
+
+private:
+	virtual bool SupportsAssetProperty(UObject* Asset) const override;
+	virtual void SetAssetProperty(UObject* Asset) override;
 };
 
 USTRUCT()
@@ -56,8 +59,11 @@ public:
 		RegisterOutputConnection(&Skeleton);
 	}
 
-
 	virtual void Evaluate(UE::Dataflow::FContext& Context, const FDataflowOutput* Out) const override;
+
+private:
+	virtual bool SupportsAssetProperty(UObject* Asset) const override;
+	virtual void SetAssetProperty(UObject* Asset) override;
 };
 
 
diff --git a/Engine/Plugins/Experimental/Dataflow/Source/DataflowNodes/Public/Dataflow/DataflowStaticMeshNodes.h b/Engine/Plugins/Experimental/Dataflow/Source/DataflowNodes/Public/Dataflow/DataflowStaticMeshNodes.h
--- a/Engine/Plugins/Experimental/Dataflow/Source/DataflowNodes/Public/Dataflow/DataflowStaticMeshNodes.h
+++ b/Engine/Plugins/Experimental/Dataflow/Source/DataflowNodes/Public/Dataflow/DataflowStaticMeshNodes.h
@@ -32,8 +32,10 @@ public:
 		RegisterOutputConnection(&StaticMesh);
 	}
 
-
 	virtual void Evaluate(UE::Dataflow::FContext& Context, const FDataflowOutput* Out) const override;
+
+	virtual bool SupportsAssetProperty(UObject* Asset) const override;
+	virtual void SetAssetProperty(UObject* Asset) override;
 };
 
 
diff --git a/Engine/Plugins/Experimental/GeometryCollectionPlugin/Source/GeometryCollectionNodes/Private/Dataflow/GeometryCollectionMaterialInterfaceNodes.cpp b/Engine/Plugins/Experimental/GeometryCollectionPlugin/Source/GeometryCollectionNodes/Private/Dataflow/GeometryCollectionMaterialInterfaceNodes.cpp
--- a/Engine/Plugins/Experimental/GeometryCollectionPlugin/Source/GeometryCollectionNodes/Private/Dataflow/GeometryCollectionMaterialInterfaceNodes.cpp
+++ b/Engine/Plugins/Experimental/GeometryCollectionPlugin/Source/GeometryCollectionNodes/Private/Dataflow/GeometryCollectionMaterialInterfaceNodes.cpp
@@ -2,6 +2,7 @@
 
 #include "Dataflow/GeometryCollectionMaterialInterfaceNodes.h"
 #include "Dataflow/DataflowCore.h"
+#include "Dataflow/DataflowObject.h"
 #include "GeometryCollection/GeometryCollection.h"
 #include "Materials/MaterialInterface.h"
 #include "Materials/Material.h"
@@ -17,6 +18,8 @@ namespace UE::Dataflow
 		DATAFLOW_NODE_REGISTER_CREATION_FACTORY(FSetIntoMaterialInterfaceArrayDataflowNode);
 		DATAFLOW_NODE_REGISTER_CREATION_FACTORY(FAddToMaterialInterfaceArrayDataflowNode);
 		DATAFLOW_NODE_REGISTER_CREATION_FACTORY(FAssignMaterialInterfaceToCollectionDataflowNode);
+
+		DATAFLOW_NODE_REGISTER_GETTER_FOR_ASSET(UMaterialInterface, FGetMaterialInterfaceAssetDataflowNode);
 	}
 }
 
@@ -50,6 +53,19 @@ void FGetMaterialInterfaceAssetDataflowNode::Evaluate(UE::Dataflow::FContext& Co
 	}
 }
 
+bool FGetMaterialInterfaceAssetDataflowNode::SupportsAssetProperty(UObject* Asset) const
+{
+	return (Cast<UMaterialInterface>(Asset) != nullptr);
+}
+
+void FGetMaterialInterfaceAssetDataflowNode::SetAssetProperty(UObject* Asset)
+{
+	if (UMaterialInterface* MaterialAsset = Cast<UMaterialInterface>(Asset))
+	{
+		Material = MaterialAsset;
+	}
+}
+
 /////////////////////////////////////////////////////////////////////////////////////////////////////////////
 FGetMaterialInterfaceArraySizeDataflowNode::FGetMaterialInterfaceArraySizeDataflowNode(const UE::Dataflow::FNodeParameters& InParam, FGuid InGuid)
 	: FDataflowNode(InParam, InGuid)
diff --git a/Engine/Plugins/Experimental/GeometryCollectionPlugin/Source/GeometryCollectionNodes/Public/Dataflow/GeometryCollectionMaterialInterfaceNodes.h b/Engine/Plugins/Experimental/GeometryCollectionPlugin/Source/GeometryCollectionNodes/Public/Dataflow/GeometryCollectionMaterialInterfaceNodes.h
--- a/Engine/Plugins/Experimental/GeometryCollectionPlugin/Source/GeometryCollectionNodes/Public/Dataflow/GeometryCollectionMaterialInterfaceNodes.h
+++ b/Engine/Plugins/Experimental/GeometryCollectionPlugin/Source/GeometryCollectionNodes/Public/Dataflow/GeometryCollectionMaterialInterfaceNodes.h
@@ -70,6 +70,9 @@ private:
 
 	virtual void Evaluate(UE::Dataflow::FContext& Context, const FDataflowOutput* Out) const override;
 
+	virtual bool SupportsAssetProperty(UObject* Asset) const override;
+	virtual void SetAssetProperty(UObject* Asset) override;
+
 public:
 	FGetMaterialInterfaceAssetDataflowNode(const UE::Dataflow::FNodeParameters& InParam, FGuid InGuid = FGuid::NewGuid());
 };
diff --git a/Engine/Source/Runtime/Experimental/Dataflow/Core/Private/Dataflow/DataflowNodeFactory.cpp b/Engine/Source/Runtime/Experimental/Dataflow/Core/Private/Dataflow/DataflowNodeFactory.cpp
--- a/Engine/Source/Runtime/Experimental/Dataflow/Core/Private/Dataflow/DataflowNodeFactory.cpp
+++ b/Engine/Source/Runtime/Experimental/Dataflow/Core/Private/Dataflow/DataflowNodeFactory.cpp
@@ -102,6 +102,29 @@ namespace UE::Dataflow
 		}
 	}
 
+	void FNodeFactory::RegisterGetterNodeForAssetType(FName AssetTypeName, FName NodeTypeName)
+	{
+		if (ensure(!GetterNodesByAssetType.Contains(AssetTypeName)))
+		{
+			GetterNodesByAssetType.Emplace(AssetTypeName, NodeTypeName);
+		}
+
+	}
+
+	FName FNodeFactory::GetGetterNodeFromAssetClass(const UClass& AssetClass) const
+	{
+		if (const FName* GetterNodeTypeName = GetterNodesByAssetType.Find(AssetClass.GetFName()))
+		{
+			return *GetterNodeTypeName;
+		}
+		// serach for compatible types up in the class hierarchy
+		if (const UClass* AssetParentClass = AssetClass.GetSuperClass())
+		{
+			return GetGetterNodeFromAssetClass(*AssetParentClass);
+		}
+		return {};
+	}
+
 	FName FNodeFactory::GetVersionFromTypeName(const FName& TypeName)
 	{
 		const FString String = TypeName.ToString();
diff --git a/Engine/Source/Runtime/Experimental/Dataflow/Core/Public/Dataflow/DataflowNode.h b/Engine/Source/Runtime/Experimental/Dataflow/Core/Public/Dataflow/DataflowNode.h
--- a/Engine/Source/Runtime/Experimental/Dataflow/Core/Public/Dataflow/DataflowNode.h
+++ b/Engine/Source/Runtime/Experimental/Dataflow/Core/Public/Dataflow/DataflowNode.h
@@ -637,6 +637,12 @@ PRAGMA_ENABLE_DEPRECATION_WARNINGS
 
 	DATAFLOWCORE_API virtual void ValidateProperties();
 
+	// return true if the node supports a specific type of asset to be assign to it 
+	DATAFLOWCORE_API virtual bool SupportsAssetProperty(UObject* Asset) const { return false; }
+
+	// Set an asset property on this node ( called when SupportsAssetProperty returns true )
+	DATAFLOWCORE_API virtual void SetAssetProperty(UObject* Asset) {};
+
 	// Called when UPROPERTY members of the dataflow node have been changed in the editor
 	DATAFLOWCORE_API virtual void OnPropertyChanged(UE::Dataflow::FContext& Context, const FPropertyChangedEvent& InPropertyChangedEvent) {}
 
@@ -866,4 +872,7 @@ private:
 	::UE::Dataflow::FPinSettingsRegistry::Get().RegisterPinSettings(A, {C, T});		\
 }																					\
 
+#define DATAFLOW_NODE_REGISTER_GETTER_FOR_ASSET(ASSETTYPE, NODETYPE)	\
+	::UE::Dataflow::FNodeFactory::GetInstance()->RegisterGetterNodeForAssetType(ASSETTYPE::StaticClass()->GetFName(), #NODETYPE);
+
 }
diff --git a/Engine/Source/Runtime/Experimental/Dataflow/Core/Public/Dataflow/DataflowNodeFactory.h b/Engine/Source/Runtime/Experimental/Dataflow/Core/Public/Dataflow/DataflowNodeFactory.h
--- a/Engine/Source/Runtime/Experimental/Dataflow/Core/Public/Dataflow/DataflowNodeFactory.h
+++ b/Engine/Source/Runtime/Experimental/Dataflow/Core/Public/Dataflow/DataflowNodeFactory.h
@@ -57,13 +57,17 @@ namespace UE::Dataflow
 	//
 	class FNodeFactory
 	{
+	public:
 		typedef TFunction<TUniquePtr<FDataflowNode> (const FNewNodeParameters&)> FNewNodeFunction;
 
+	private:
 		// All Maps indexed by TypeName
 		TMap<FName, FNewNodeFunction > ClassMap;			// [TypeName] -> NewNodeFunction
 		TMap<FName, FFactoryParameters > ParametersMap;		// [TypeName] -> Parameters
 		TMap<FName, TArray<FName>> VersionMap;				// [TypeNameNoVersion] -> Array of TypeName(versions)
 
+		TMap<FName, FName> GetterNodesByAssetType;	// list of getter nodes that relate to a specific asset type
+
 		DATAFLOWCORE_API static FNodeFactory* Instance;
 		FNodeFactory() {};
 
@@ -114,6 +118,10 @@ namespace UE::Dataflow
 			GetInstance()->RegisterNode(FactoryParameters, CreationFunction);
 		}
 
+		DATAFLOWCORE_API void RegisterGetterNodeForAssetType(FName AssetTypeName, FName NodeTypeName);
+
+		DATAFLOWCORE_API FName GetGetterNodeFromAssetClass(const UClass& AssetClass) const;
+
 		DATAFLOWCORE_API const FFactoryParameters& GetParameters(FName InTypeName) const;
 
 		DATAFLOWCORE_API TSharedPtr<FDataflowNode> NewNodeFromRegisteredType(FGraph& Graph, const FNewNodeParameters& Param);
