Dataflow - Removed the redundant FDataflowXml class.

#rnx

diff --git a/Engine/Plugins/Experimental/Dataflow/Source/DataflowEditor/Private/Dataflow/DataflowGraphEditor.cpp b/Engine/Plugins/Experimental/Dataflow/Source/DataflowEditor/Private/Dataflow/DataflowGraphEditor.cpp
--- a/Engine/Plugins/Experimental/Dataflow/Source/DataflowEditor/Private/Dataflow/DataflowGraphEditor.cpp
+++ b/Engine/Plugins/Experimental/Dataflow/Source/DataflowEditor/Private/Dataflow/DataflowGraphEditor.cpp
@@ -11,7 +11,6 @@
 #include "Dataflow/DataflowSNodeFactories.h"
 #include "Dataflow/DataflowSCommentNode.h"
 #include "Dataflow/DataflowNodeParameters.h"
-#include "Dataflow/DataflowXml.h"
 #include "Dataflow/DataflowGraphSchemaAction.h"
 #include "Dataflow/DataflowSelectionNodes.h"
 #include "Dataflow/DataflowSubGraphNodes.h"
@@ -425,23 +424,6 @@ bool SDataflowGraphEditor::IsAltDown() const
 
 FReply SDataflowGraphEditor::OnKeyUp(const FGeometry& MyGeometry, const FKeyEvent& InKeyEvent)
 {
-	// process paste
-	if ( (VKeyDown && LeftControlKeyDown) || (VKeyDown && RightControlKeyDown) )
-	{
-		if (InKeyEvent.GetKey() == EKeys::LeftControl
-			|| InKeyEvent.GetKey() == EKeys::RightControl
-			|| InKeyEvent.GetKey() == EKeys::V)
-		{
-			FString XmlBuffer;
-			FPlatformApplicationMisc::ClipboardPaste(XmlBuffer);
-			FDataflowXmlRead Xml(this); 
-			if (Xml.LoadFromBuffer(XmlBuffer))
-			{
-				Xml.ParseXmlFile();
-			}
-		}
-	}
-
 	if (InKeyEvent.GetKey() == EKeys::LeftControl)
 	{
 		LeftControlKeyDown = false;
@@ -528,35 +510,6 @@ void SDataflowGraphEditor::CreateComment()
 	}
 }
 
-void SDataflowGraphEditor::CreateVertexSelectionNode(const FString & InArray)
-{
-	if (TStrongObjectPtr<UEdGraph> EdGraph = EdGraphWeakPtr.Pin())
-	{
-		const TSharedPtr<SGraphEditor>& InGraphEditor = SharedThis(GetGraphEditor());
-
-		if (TSharedPtr<FAssetSchemaAction_Dataflow_CreateNode_DataflowEdNode> NodeAction = FAssetSchemaAction_Dataflow_CreateNode_DataflowEdNode::CreateAction(EdGraph.Get(), "FSelectionSetDataflowNode"))
-		{
-			if (UEdGraphNode* NewEdNode = NodeAction->PerformAction(EdGraph.Get(), nullptr, GetGraphEditor()->GetPasteLocation(), false))
-			{
-				if (UDataflow* DataflowAsset = UDataflow::GetDataflowAssetFromEdGraph(EdGraph.Get()))
-				{
-					FDataflowAssetEdit Edit = DataflowAsset->EditDataflow();
-					if (UE::Dataflow::FGraph* DataflowGraph = Edit.GetGraph())
-					{
-						if (TSharedPtr<FDataflowNode> Node = DataflowGraph->FindBaseNode(((UDataflowEdNode*)NewEdNode)->DataflowNodeGuid))
-						{
-							if (FSelectionSetDataflowNode* SelectionNode = Node->AsType<FSelectionSetDataflowNode>())
-							{
-								SelectionNode->Indices = InArray;
-							}
-						}
-					}
-				}
-			}
-		}
-	}
-}
-
 void SDataflowGraphEditor::AlignTop()
 {
 	GetGraphEditor()->OnAlignTop();
diff --git a/Engine/Plugins/Experimental/Dataflow/Source/DataflowEditor/Private/Dataflow/DataflowXml.cpp b/Engine/Plugins/Experimental/Dataflow/Source/DataflowEditor/Private/Dataflow/DataflowXml.cpp
deleted file mode 100644
diff --git a/Engine/Plugins/Experimental/Dataflow/Source/DataflowEditor/Public/Dataflow/DataflowGraphEditor.h b/Engine/Plugins/Experimental/Dataflow/Source/DataflowEditor/Public/Dataflow/DataflowGraphEditor.h
--- a/Engine/Plugins/Experimental/Dataflow/Source/DataflowEditor/Public/Dataflow/DataflowGraphEditor.h
+++ b/Engine/Plugins/Experimental/Dataflow/Source/DataflowEditor/Public/Dataflow/DataflowGraphEditor.h
@@ -89,9 +89,6 @@ public:
 	/** */
 	void CreateComment();
 
-	/** */
-	void CreateVertexSelectionNode(const FString& InArray);
-
 	/** */
 	void AlignTop();
 
diff --git a/Engine/Plugins/Experimental/Dataflow/Source/DataflowEditor/Public/Dataflow/DataflowXml.h b/Engine/Plugins/Experimental/Dataflow/Source/DataflowEditor/Public/Dataflow/DataflowXml.h
deleted file mode 100644
