Dataflow:
- Added Separators around Option Pin menus, and fixed text in Rt click menu

diff --git a/Engine/Plugins/Experimental/Dataflow/Source/DataflowEditor/Private/Dataflow/DataflowSchema.cpp b/Engine/Plugins/Experimental/Dataflow/Source/DataflowEditor/Private/Dataflow/DataflowSchema.cpp
--- a/Engine/Plugins/Experimental/Dataflow/Source/DataflowEditor/Private/Dataflow/DataflowSchema.cpp
+++ b/Engine/Plugins/Experimental/Dataflow/Source/DataflowEditor/Private/Dataflow/DataflowSchema.cpp
@@ -43,14 +43,16 @@ void UDataflowSchema::GetContextMenuActions(class UToolMenu* Menu, class UGraphN
 				Section.AddMenuEntry(FGenericCommands::Get().Cut);
 				Section.AddMenuEntry(FGenericCommands::Get().Copy);
 				Section.AddMenuEntry(FGenericCommands::Get().Duplicate);
-				Section.AddMenuEntry(FDataflowEditorCommands::Get().ToggleEnabledState, FText::FromString("Toggle Enabled State"));
+				Section.AddMenuEntry(FDataflowEditorCommands::Get().ToggleEnabledState, LOCTEXT("DataflowContextMenu_ToggleEnabledState", "Toggle Enabled State"));
 				Section.AddMenuEntry(FGraphEditorCommands::Get().BreakNodeLinks);
-				Section.AddMenuEntry(FDataflowEditorCommands::Get().AddOptionPin);
-				Section.AddMenuEntry(FDataflowEditorCommands::Get().RemoveOptionPin);
+				Section.AddSeparator(TEXT("ActionsSeparator1"));
+				Section.AddMenuEntry(FDataflowEditorCommands::Get().AddOptionPin, LOCTEXT("DataflowContextMenu_AddOptionPin", "Add Option Pin"));
+				Section.AddMenuEntry(FDataflowEditorCommands::Get().RemoveOptionPin, LOCTEXT("DataflowContextMenu_RemoveOptionPin", "Remove Option Pin"));
+				Section.AddSeparator(TEXT("ActionsSeparator2"));
 				Section.AddMenuEntry(FDataflowEditorCommands::Get().EvaluateNode);
-				Section.AddSeparator(TEXT("ActionsSeparator"));
-				Section.AddMenuEntry(FDataflowEditorCommands::Get().FreezeNodes, FText::FromString("Freeze"));
-				Section.AddMenuEntry(FDataflowEditorCommands::Get().UnfreezeNodes, FText::FromString("Unfreeze"));
+				Section.AddSeparator(TEXT("ActionsSeparator3"));
+				Section.AddMenuEntry(FDataflowEditorCommands::Get().FreezeNodes, LOCTEXT("DataflowContextMenu_Freeze", "Freeze"));
+				Section.AddMenuEntry(FDataflowEditorCommands::Get().UnfreezeNodes, LOCTEXT("DataflowContextMenu_Unfreeze", "Unfreeze"));
 			}
 		}
 		{
