Remove unused bHasTerminalsDirty in the dataflow editor code

diff --git a/Engine/Plugins/Experimental/Dataflow/Source/DataflowEditor/Private/Dataflow/DataflowEditor.cpp b/Engine/Plugins/Experimental/Dataflow/Source/DataflowEditor/Private/Dataflow/DataflowEditor.cpp
--- a/Engine/Plugins/Experimental/Dataflow/Source/DataflowEditor/Private/Dataflow/DataflowEditor.cpp
+++ b/Engine/Plugins/Experimental/Dataflow/Source/DataflowEditor/Private/Dataflow/DataflowEditor.cpp
@@ -99,14 +99,12 @@ void UDataflowEditor::RemoveTerminalContents(const TSharedPtr<UE::Dataflow::FGra
 			else
 			{
 				TerminalContents.RemoveAt(ContentIndex);
-				bHasTerminalsDirty = true;
 			}
 		}
 		else
 		{
 			// Removal of all the contents with invalid nodes
 			TerminalContents.RemoveAt(ContentIndex);
-			bHasTerminalsDirty = true;
 		}
 	}
 }
@@ -128,12 +126,10 @@ void UDataflowEditor::AddTerminalContents(const TSharedPtr<UE::Dataflow::FGraph>
 					(*TerminalContent)->SetDataflowContext(EditorContent->GetDataflowContext());
 					
 					(*TerminalContent)->SetLastModifiedTimestamp(EditorContent->GetLastModifiedTimestamp());
-					bHasTerminalsDirty = true;
 				}
 				if(TerminalNode->GetTerminalAsset() != (*TerminalContent)->GetTerminalAsset())
 				{
 					(*TerminalContent)->SetTerminalAsset(TerminalNode->GetTerminalAsset());
-					bHasTerminalsDirty = true;
 				}
 			}
 		}
@@ -142,8 +138,6 @@ void UDataflowEditor::AddTerminalContents(const TSharedPtr<UE::Dataflow::FGraph>
 
 void UDataflowEditor::UpdateTerminalContents(const UE::Dataflow::FTimestamp TimeStamp)
 {
-	bHasTerminalsDirty = false;
-	
 	// update of the terminal contents only if no terminal asset on the main editor content
 	if(EditorContent && !EditorContent->GetTerminalAsset() && EditorContent->GetDataflowAsset())
 	{
diff --git a/Engine/Plugins/Experimental/Dataflow/Source/DataflowEditor/Public/Dataflow/DataflowEditor.h b/Engine/Plugins/Experimental/Dataflow/Source/DataflowEditor/Public/Dataflow/DataflowEditor.h
--- a/Engine/Plugins/Experimental/Dataflow/Source/DataflowEditor/Public/Dataflow/DataflowEditor.h
+++ b/Engine/Plugins/Experimental/Dataflow/Source/DataflowEditor/Public/Dataflow/DataflowEditor.h
@@ -60,9 +60,6 @@ public:
 	TArray<TObjectPtr<UDataflowBaseContent>>& GetTerminalContents() { return TerminalContents; }
 	const TArray<TObjectPtr<UDataflowBaseContent>>& GetTerminalContents() const { return TerminalContents; }
 
-	/** Check if the terminals contents has changed or not */
-	bool HasTerminalsDirty() const { return bHasTerminalsDirty; }
-
 private :
 
 	friend class FDataflowEditorToolkit;
@@ -85,9 +82,6 @@ private :
 	/** List of dataflow contents available in the graph and coming from all the terminal nodes */
 	UPROPERTY()
 	TArray<TObjectPtr<UDataflowBaseContent>> TerminalContents;
-
-	/** Boolean to know when the terminals has been changed or not */
-	bool bHasTerminalsDirty = false;
 };
 
 DECLARE_LOG_CATEGORY_EXTERN(LogDataflowEditor, Log, All);
