Dataflow : fix node unique name generation to account for the subgraph the node is created in

diff --git a/Engine/Plugins/Experimental/Dataflow/Source/DataflowEditor/Private/Dataflow/DataflowAssetEditUtils.cpp b/Engine/Plugins/Experimental/Dataflow/Source/DataflowEditor/Private/Dataflow/DataflowAssetEditUtils.cpp
--- a/Engine/Plugins/Experimental/Dataflow/Source/DataflowEditor/Private/Dataflow/DataflowAssetEditUtils.cpp
+++ b/Engine/Plugins/Experimental/Dataflow/Source/DataflowEditor/Private/Dataflow/DataflowAssetEditUtils.cpp
@@ -88,12 +88,32 @@ namespace UE::Dataflow
 					BaseName = FName(Left);
 				}
 			}
-			
-			if (::IsUniqueObjectName(BaseName, Dataflow))
-			{
-				return BaseName;
-			}
-			return ::MakeUniqueObjectName(Dataflow, UDataflowEdNode::StaticClass(), BaseName);
+
+			// name must be unique for all nodes in the Dataflow::FGraph 
+			// Unreal require names to be unique within the parent but because we have one FGraph across all EdGraph ( include SubGraphs) objects 
+			// we need to make sure the name is unique across them, so that we don't get a assert when creating the EdNode
+			FName UniqueName = BaseName;
+			bool bNameWasChanged = false;
+			do {
+				// reset for this loop 
+				bNameWasChanged = false;
+				if (!::IsUniqueObjectName(UniqueName, Dataflow))
+				{
+					UniqueName = ::MakeUniqueObjectName(Dataflow, UDataflowEdNode::StaticClass(), UniqueName);
+					bNameWasChanged = true;
+				}
+
+				for (TObjectPtr<UDataflowSubGraph> SubGraph: Dataflow->GetSubGraphs())
+				{
+					if (!::IsUniqueObjectName(UniqueName, SubGraph))
+					{
+						UniqueName = ::MakeUniqueObjectName(SubGraph, UDataflowEdNode::StaticClass(), UniqueName);
+						bNameWasChanged = true;
+					}
+				}
+			} while (bNameWasChanged);
+
+			return UniqueName;
 		}
 
 		static TSharedPtr<FDataflowNode> AddDataflowNode(UDataflow* Dataflow, FName NodeName, FName NodeTypeName)
