Dataflow : fix node properties readonly logic
- Children of input struct are now properly disabled when input is connected
- Added a way for the node to avoid the properties to go read only when connected
- Subgraph output node ignore connected property being readonly so that user can rename them or change their type in teh property bag

diff --git a/Engine/Plugins/Experimental/Dataflow/Source/DataflowEditor/Private/Dataflow/DataflowEditorCommands.cpp b/Engine/Plugins/Experimental/Dataflow/Source/DataflowEditor/Private/Dataflow/DataflowEditorCommands.cpp
--- a/Engine/Plugins/Experimental/Dataflow/Source/DataflowEditor/Private/Dataflow/DataflowEditorCommands.cpp
+++ b/Engine/Plugins/Experimental/Dataflow/Source/DataflowEditor/Private/Dataflow/DataflowEditorCommands.cpp
@@ -438,16 +438,36 @@ void FDataflowEditorCommands::OnSelectedNodesChanged(TSharedPtr<IStructureDetail
 					{
 						if (TSharedPtr<FDataflowNode> DataflowNode = DataflowGraph->FindBaseNode(EdNode->GetDataflowNodeGuid()))
 						{
-							FIsPropertyReadOnly Delegate = FIsPropertyReadOnly::CreateLambda(
+							FIsPropertyReadOnly IsPropertyReadOnlyLambda = FIsPropertyReadOnly::CreateLambda(
 								[DataflowNode](const FPropertyAndParent& PropertyAndParent)
 								{
-									if (const FDataflowInput* Input = DataflowNode->FindInput(PropertyAndParent.Property.GetFName()))
+									if (DataflowNode->MakeConnectedPropertiesReadOnly())
 									{
-										return (Input->GetConnection() != nullptr);
+										if (const FDataflowInput* Input = DataflowNode->FindInput(PropertyAndParent.Property.GetFName()))
+										{
+											if (Input->IsConnected())
+											{
+												return true;
+											}
+										}
+										// let's check the parents so that we can disable cghildren properties of structures ( Vectors for example )
+										for (const FProperty* ParentProperty : PropertyAndParent.ParentProperties)
+										{
+											if (ParentProperty)
+											{
+												if (const FDataflowInput* Input = DataflowNode->FindInput(ParentProperty->GetFName()))
+												{
+													if (Input->IsConnected())
+													{
+														return true;
+													}
+												}
+											}
+										}
 									}
 									return false;
 								});
-							PropertiesEditor->GetDetailsView()->SetIsPropertyReadOnlyDelegate(Delegate);
+							PropertiesEditor->GetDetailsView()->SetIsPropertyReadOnlyDelegate(IsPropertyReadOnlyLambda);
 
 							StructData.Emplace(DataflowNode->NewStructOnScope());
 						}
diff --git a/Engine/Source/Runtime/Experimental/Dataflow/Core/Public/Dataflow/DataflowNode.h b/Engine/Source/Runtime/Experimental/Dataflow/Core/Public/Dataflow/DataflowNode.h
--- a/Engine/Source/Runtime/Experimental/Dataflow/Core/Public/Dataflow/DataflowNode.h
+++ b/Engine/Source/Runtime/Experimental/Dataflow/Core/Public/Dataflow/DataflowNode.h
@@ -671,6 +671,9 @@ PRAGMA_ENABLE_DEPRECATION_WARNINGS
 
 	DATAFLOWCORE_API virtual bool ShouldInvalidateOnPropertyChanged(const FPropertyChangedEvent& InPropertyChangedEvent) const { return true; };
 
+	// returns true to disable node properties in its detail view when the corresponding input is connected ( default is true )
+	DATAFLOWCORE_API virtual bool MakeConnectedPropertiesReadOnly() const { return true; }
+
 	bool HasValidConnections() const { return bHasValidConnections; }
 
 	virtual bool IsA(FName InType) const 
diff --git a/Engine/Source/Runtime/Experimental/Dataflow/Engine/Public/Dataflow/DataflowSubGraphNodes.h b/Engine/Source/Runtime/Experimental/Dataflow/Engine/Public/Dataflow/DataflowSubGraphNodes.h
--- a/Engine/Source/Runtime/Experimental/Dataflow/Engine/Public/Dataflow/DataflowSubGraphNodes.h
+++ b/Engine/Source/Runtime/Experimental/Dataflow/Engine/Public/Dataflow/DataflowSubGraphNodes.h
@@ -130,6 +130,9 @@ private:
 	virtual bool SupportsDropConnectionOnNode(FName TypeName, UE::Dataflow::FPin::EDirection Direction) const override;
 	virtual const FDataflowConnection* OnDropConnectionOnNode(const FDataflowConnection& DroppedConnection) override;
 
+	// we want the user to be able to change the type of the properties in the property bag or rename them 
+	virtual bool MakeConnectedPropertiesReadOnly() const override { return false; }
+
 	UPROPERTY()
 	FDataflowDynamicConnections DynamicConnections;
 
