Dataflow: add custom zoom level container for dataflow graph editor
- Added two extra zoom out levels ( up to -14 )

diff --git a/Engine/Plugins/Experimental/Dataflow/Source/DataflowEditor/Private/Dataflow/DataflowGraphEditor.cpp b/Engine/Plugins/Experimental/Dataflow/Source/DataflowEditor/Private/Dataflow/DataflowGraphEditor.cpp
--- a/Engine/Plugins/Experimental/Dataflow/Source/DataflowEditor/Private/Dataflow/DataflowGraphEditor.cpp
+++ b/Engine/Plugins/Experimental/Dataflow/Source/DataflowEditor/Private/Dataflow/DataflowGraphEditor.cpp
@@ -23,6 +23,98 @@
 
 #define LOCTEXT_NAMESPACE "DataflowGraphEditor"
 
+namespace DataflowGraphEditor::Private
+{
+	struct FDataflowZoomLevelEntry
+	{
+	public:
+		FDataflowZoomLevelEntry(float InZoomAmount, const FText& InDisplayText, EGraphRenderingLOD::Type InLOD)
+			: DisplayText(FText::Format(LOCTEXT("Zoom", "Zoom {0}"), InDisplayText))
+			, ZoomAmount(InZoomAmount)
+			, LOD(InLOD)
+		{}
+
+	public:
+		FText DisplayText;
+		float ZoomAmount;
+		EGraphRenderingLOD::Type LOD;
+	};
+
+	struct FDataflowZoomLevelsContainer : public FZoomLevelsContainer
+	{
+		FDataflowZoomLevelsContainer()
+		{
+			ZoomLevels.Append({
+					FDataflowZoomLevelEntry(0.025f, FText::FromString(TEXT("-14")), EGraphRenderingLOD::LowestDetail),
+					FDataflowZoomLevelEntry(0.070f, FText::FromString(TEXT("-13")), EGraphRenderingLOD::LowestDetail),
+					FDataflowZoomLevelEntry(0.100f, FText::FromString(TEXT("-12")), EGraphRenderingLOD::LowestDetail),
+					FDataflowZoomLevelEntry(0.125f, FText::FromString(TEXT("-11")), EGraphRenderingLOD::LowestDetail),
+					FDataflowZoomLevelEntry(0.150f, FText::FromString(TEXT("-10")), EGraphRenderingLOD::LowestDetail),
+					FDataflowZoomLevelEntry(0.175f, FText::FromString(TEXT("-9")), EGraphRenderingLOD::LowestDetail),
+					FDataflowZoomLevelEntry(0.200f, FText::FromString(TEXT("-8")), EGraphRenderingLOD::LowestDetail),
+					FDataflowZoomLevelEntry(0.225f, FText::FromString(TEXT("-7")), EGraphRenderingLOD::LowDetail),
+					FDataflowZoomLevelEntry(0.250f, FText::FromString(TEXT("-6")), EGraphRenderingLOD::LowDetail),
+					FDataflowZoomLevelEntry(0.375f, FText::FromString(TEXT("-5")), EGraphRenderingLOD::MediumDetail),
+					FDataflowZoomLevelEntry(0.500f, FText::FromString(TEXT("-4")), EGraphRenderingLOD::MediumDetail),
+					FDataflowZoomLevelEntry(0.675f, FText::FromString(TEXT("-3")), EGraphRenderingLOD::MediumDetail),
+					FDataflowZoomLevelEntry(0.750f, FText::FromString(TEXT("-2")), EGraphRenderingLOD::DefaultDetail),
+					FDataflowZoomLevelEntry(0.875f, FText::FromString(TEXT("-1")), EGraphRenderingLOD::DefaultDetail),
+					FDataflowZoomLevelEntry(1.000f, FText::FromString(TEXT("1:1")), EGraphRenderingLOD::DefaultDetail),
+					FDataflowZoomLevelEntry(1.250f, FText::FromString(TEXT("+1")), EGraphRenderingLOD::DefaultDetail),
+					FDataflowZoomLevelEntry(1.375f, FText::FromString(TEXT("+2")), EGraphRenderingLOD::DefaultDetail),
+					FDataflowZoomLevelEntry(1.500f, FText::FromString(TEXT("+3")), EGraphRenderingLOD::FullyZoomedIn),
+					FDataflowZoomLevelEntry(1.675f, FText::FromString(TEXT("+4")), EGraphRenderingLOD::FullyZoomedIn),
+					FDataflowZoomLevelEntry(1.750f, FText::FromString(TEXT("+5")), EGraphRenderingLOD::FullyZoomedIn),
+					FDataflowZoomLevelEntry(1.875f, FText::FromString(TEXT("+6")), EGraphRenderingLOD::FullyZoomedIn),
+					FDataflowZoomLevelEntry(2.000f, FText::FromString(TEXT("+7")), EGraphRenderingLOD::FullyZoomedIn),
+				});
+		}
+
+		float GetZoomAmount(int32 InZoomLevel) const override
+		{
+			checkSlow(ZoomLevels.IsValidIndex(InZoomLevel));
+			return ZoomLevels[InZoomLevel].ZoomAmount;
+		}
+
+		int32 GetNearestZoomLevel(float InZoomAmount) const override
+		{
+			for (int32 ZoomLevelIndex = 0; ZoomLevelIndex < GetNumZoomLevels(); ++ZoomLevelIndex)
+			{
+				if (InZoomAmount <= GetZoomAmount(ZoomLevelIndex))
+				{
+					return ZoomLevelIndex;
+				}
+			}
+
+			return GetDefaultZoomLevel();
+		}
+
+		FText GetZoomText(int32 InZoomLevel) const override
+		{
+			checkSlow(ZoomLevels.IsValidIndex(InZoomLevel));
+			return ZoomLevels[InZoomLevel].DisplayText;
+		}
+
+		int32 GetNumZoomLevels() const override
+		{
+			return ZoomLevels.Num();
+		}
+
+		int32 GetDefaultZoomLevel() const override
+		{
+			return 12;
+		}
+
+		EGraphRenderingLOD::Type GetLOD(int32 InZoomLevel) const override
+		{
+			checkSlow(ZoomLevels.IsValidIndex(InZoomLevel));
+			return ZoomLevels[InZoomLevel].LOD;
+		}
+
+		TArray<FDataflowZoomLevelEntry> ZoomLevels;
+	};
+}
+
 TSharedPtr<FDataflowGraphEditorNodeFactory> SDataflowGraphEditor::NodeFactory;
 TWeakPtr<SDataflowGraphEditor> SDataflowGraphEditor::SelectedGraphEditor;
 
@@ -173,6 +265,7 @@ void SDataflowGraphEditor::Construct(const FArguments& InArgs, UObject* InAssetO
 	SGraphEditor::Construct(Arguments);
 
 	SetNodeFactory( MakeShared<FDataflowGraphNodeFactory>(this) );
+	GetGraphPanel()->SetZoomLevelsContainer<DataflowGraphEditor::Private::FDataflowZoomLevelsContainer>();
 
 	InitGraphEditorMessageBar();
 
