Dataflow : add SelectAll command to the editor

diff --git a/Engine/Plugins/Experimental/Dataflow/Source/DataflowEditor/Private/Dataflow/DataflowGraphEditor.cpp b/Engine/Plugins/Experimental/Dataflow/Source/DataflowEditor/Private/Dataflow/DataflowGraphEditor.cpp
--- a/Engine/Plugins/Experimental/Dataflow/Source/DataflowEditor/Private/Dataflow/DataflowGraphEditor.cpp
+++ b/Engine/Plugins/Experimental/Dataflow/Source/DataflowEditor/Private/Dataflow/DataflowGraphEditor.cpp
@@ -136,6 +136,10 @@ void SDataflowGraphEditor::Construct(const FArguments& InArgs, UObject* InAssetO
 	{
 		GraphEditorCommands = MakeShareable(new FUICommandList);
 		{
+			GraphEditorCommands->MapAction(FGenericCommands::Get().SelectAll,
+				FExecuteAction::CreateSP(this, &SDataflowGraphEditor::SelectAllNodes)
+			);
+
 			GraphEditorCommands->MapAction(
 				FGenericCommands::Get().Delete,
 				FExecuteAction::CreateSP(this, &SDataflowGraphEditor::DeleteNode)
@@ -249,6 +253,7 @@ void SDataflowGraphEditor::Construct(const FArguments& InArgs, UObject* InAssetO
 					FDataflowEditorCommands::Get().AddNewSubGraph,
 					FExecuteAction::CreateSP(this, &SDataflowGraphEditor::AddNewSubGraph)
 			);
+
 		}
 	}
 
