Dataflow : Update layout to properly account for the newly introduced widgets ( SubGraphs / OutputLog )
- Rename Output log to Dataflow Output log to avoid confusion with the Editor general one
- change the output log window to match the UE one

diff --git a/Engine/Plugins/Experimental/Dataflow/Source/DataflowEditor/Private/Dataflow/DataflowEditorToolkit.cpp b/Engine/Plugins/Experimental/Dataflow/Source/DataflowEditor/Private/Dataflow/DataflowEditorToolkit.cpp
--- a/Engine/Plugins/Experimental/Dataflow/Source/DataflowEditor/Private/Dataflow/DataflowEditorToolkit.cpp
+++ b/Engine/Plugins/Experimental/Dataflow/Source/DataflowEditor/Private/Dataflow/DataflowEditorToolkit.cpp
@@ -122,7 +122,7 @@ FDataflowEditorToolkit::FDataflowEditorToolkit(UAssetEditor* InOwningAssetEditor
 				->Split
 				(
 					FTabManager::NewSplitter()->SetOrientation(Orient_Horizontal)
-					->SetSizeCoefficient(0.60f)	// Relative height of (Tools Panel, Construction Viewport, Preview Viewport) vs (Dataflow Graph Editor, Outliner)
+					->SetSizeCoefficient(0.50f)	// Relative height of (Tools Panel, Construction Viewport, Preview Viewport) vs (Dataflow Graph Editor, Outliner)
 					->Split
 					(
 						FTabManager::NewStack()
@@ -150,34 +150,40 @@ FDataflowEditorToolkit::FDataflowEditorToolkit(UAssetEditor* InOwningAssetEditor
 				->Split
 				(
 					FTabManager::NewSplitter()->SetOrientation(Orient_Horizontal)
-					->SetSizeCoefficient(0.40f)	// Relative height of (Dataflow Node Details) vs (Asset Details, Preview Scene Details)
-					->Split
-					(
-						FTabManager::NewStack()
-						->SetSizeCoefficient(0.1f)
-						->AddTab(CollectionSpreadSheetTabId_1, ETabState::OpenedTab)
-						->SetExtensionId("CollectionSpreadSheetArea")
-						->SetHideTabWell(false)
-					)
+					->SetSizeCoefficient(0.50f)	// Relative height of (Dataflow Node Details) vs (Asset Details, Preview Scene Details)
 					->Split
 					(
 						FTabManager::NewStack()
-						->SetSizeCoefficient(0.1f)	// Relative height of (Dataflow Graph Editor, Outliner) vs (Tools Panel, Construction Viewport, Preview Viewport)
-						->AddTab(GraphCanvasTabId, ETabState::OpenedTab)
-						->AddTab(SubGraphCanvasTabId, ETabState::OpenedTab)
+						->SetSizeCoefficient(0.2f)
+						->AddTab(MembersWidgetTabId, ETabState::OpenedTab)
 						->SetExtensionId("MembersWidgetArea")
 						->SetHideTabWell(false)
 						->SetForegroundTab(MembersWidgetTabId)
 					)
 					->Split
 					(
-						FTabManager::NewStack()
-						->SetSizeCoefficient(0.8f)	// Relative height of (Dataflow Graph Editor, Outliner) vs (Tools Panel, Construction Viewport, Preview Viewport)
-						->AddTab(GraphCanvasTabId, ETabState::OpenedTab)
-						->AddTab(SubGraphCanvasTabId, ETabState::OpenedTab)
-						->SetExtensionId("GraphEditorArea")
-						->SetHideTabWell(false)
-						->SetForegroundTab(GraphCanvasTabId)
+						FTabManager::NewSplitter()->SetOrientation(Orient_Vertical)
+						->SetSizeCoefficient(0.8f)	// Relative width of (Tools Panel, Construction Viewport, Preview Viewport, Dataflow Graph Editor, Outliner) vs (Asset Details, Preview Scene Details, Dataflow Node Details)
+						->Split
+						(
+							FTabManager::NewStack()
+							->SetSizeCoefficient(0.8f)	// Relative height of (Dataflow Graph Editor) vs (Dataflow Output log)
+							->AddTab(GraphCanvasTabId, ETabState::OpenedTab)
+							->AddTab(SubGraphCanvasTabId, ETabState::ClosedTab)
+							->SetExtensionId("GraphEditorArea")
+							->SetHideTabWell(false)
+							->SetForegroundTab(GraphCanvasTabId)
+						)
+						->Split
+						(
+							FTabManager::NewStack()
+							->SetSizeCoefficient(0.2f)	// Relative height of (Dataflow Graph Editor) vs (Dataflow Output log)
+							->AddTab(OutputLogTabId, ETabState::OpenedTab)
+							->AddTab(CollectionSpreadSheetTabId_1, ETabState::OpenedTab)
+							->SetExtensionId("OutputLogArea")
+							->SetHideTabWell(false)
+							->SetForegroundTab(OutputLogTabId)
+						)
 					)
 				)
 			)
@@ -202,7 +208,7 @@ FDataflowEditorToolkit::FDataflowEditorToolkit(UAssetEditor* InOwningAssetEditor
 					->SetSizeCoefficient(0.2f)	// Relative height of (Dataflow Node Details) vs (Asset Details, Preview Scene Details)
 					->AddTab(NodeDetailsTabId, ETabState::OpenedTab)
 					->SetExtensionId("NodeDetailsArea")
-					->SetHideTabWell(true)
+					->SetHideTabWell(false)
 				)
 			)
 		);
@@ -2143,9 +2149,9 @@ void FDataflowEditorToolkit::RegisterTabSpawners(const TSharedRef<FTabManager>&
 	DocumentManager->RegisterDocumentFactory(SubGraphTabSummoner);
 	
 	InTabManager->RegisterTabSpawner(OutputLogTabId, FOnSpawnTab::CreateSP(this, &FDataflowEditorToolkit::SpawnTab_OutputLog))
-		.SetDisplayName(LOCTEXT("OutputLogTab", "Output Log"))
+		.SetDisplayName(LOCTEXT("OutputLogTab", "Dataflow Output Log"))
 		.SetGroup(EditorMenuCategory.ToSharedRef())
-		.SetIcon(FSlateIcon(FAppStyle::GetAppStyleSetName(), "LevelEditor.Tabs.Outliner"));
+		.SetIcon(FSlateIcon(FAppStyle::GetAppStyleSetName(), "Log.TabIcon"));
 }
 
 void FDataflowEditorToolkit::UnregisterTabSpawners(const TSharedRef<class FTabManager>& InTabManager)
