Dataflow - Added the missing Freeze commands from the Dataflow Editor action list.

#rnx

diff --git a/Engine/Plugins/Experimental/Dataflow/Source/DataflowEditor/Private/Dataflow/DataflowGraphEditor.cpp b/Engine/Plugins/Experimental/Dataflow/Source/DataflowEditor/Private/Dataflow/DataflowGraphEditor.cpp
--- a/Engine/Plugins/Experimental/Dataflow/Source/DataflowEditor/Private/Dataflow/DataflowGraphEditor.cpp
+++ b/Engine/Plugins/Experimental/Dataflow/Source/DataflowEditor/Private/Dataflow/DataflowGraphEditor.cpp
@@ -148,6 +148,14 @@ void SDataflowGraphEditor::Construct(const FArguments& InArgs, UObject* InAssetO
 				FDataflowEditorCommands::Get().EvaluateNode,
 				FExecuteAction::CreateSP(this, &SDataflowGraphEditor::EvaluateNode)
 			);
+			GraphEditorCommands->MapAction(
+				FDataflowEditorCommands::Get().FreezeNodes,
+				FExecuteAction::CreateSP(this, &SDataflowGraphEditor::FreezeNodes)
+			);
+			GraphEditorCommands->MapAction(
+				FDataflowEditorCommands::Get().UnfreezeNodes,
+				FExecuteAction::CreateSP(this, &SDataflowGraphEditor::UnfreezeNodes)
+			);
 			GraphEditorCommands->MapAction(
 				FGraphEditorCommands::Get().CreateComment,
 				FExecuteAction::CreateRaw(this, &SDataflowGraphEditor::CreateComment)
@@ -371,6 +379,22 @@ void SDataflowGraphEditor::EvaluateNode()
 	}
 }
 
+void SDataflowGraphEditor::FreezeNodes()
+{
+	if (UE::Dataflow::FContext* const DataflowContext = GetDataflowContext().Get())
+	{
+		FDataflowEditorCommands::FreezeNodes(*DataflowContext, GetSelectedNodes());
+	}
+}
+
+void SDataflowGraphEditor::UnfreezeNodes()
+{
+	if (UE::Dataflow::FContext* const DataflowContext = GetDataflowContext().Get())
+	{
+		FDataflowEditorCommands::UnfreezeNodes(*DataflowContext, GetSelectedNodes());
+	}
+}
+
 void SDataflowGraphEditor::DeleteNode()
 {
 	if (TStrongObjectPtr<UEdGraph> EdGraph = EdGraphWeakPtr.Pin())
diff --git a/Engine/Plugins/Experimental/Dataflow/Source/DataflowEditor/Public/Dataflow/DataflowGraphEditor.h b/Engine/Plugins/Experimental/Dataflow/Source/DataflowEditor/Public/Dataflow/DataflowGraphEditor.h
--- a/Engine/Plugins/Experimental/Dataflow/Source/DataflowEditor/Public/Dataflow/DataflowGraphEditor.h
+++ b/Engine/Plugins/Experimental/Dataflow/Source/DataflowEditor/Public/Dataflow/DataflowGraphEditor.h
@@ -80,6 +80,12 @@ public:
 	/** */
 	void EvaluateNode();
 
+	/** */
+	void FreezeNodes();
+
+	/** */
+	void UnfreezeNodes();
+
 	/** */
 	void DeleteNode();
 
