Dataflow Editor: fix normal rendering of DynamicMesh outputs

diff --git a/Engine/Plugins/Experimental/Dataflow/Source/DataflowAssetTools/Private/Dataflow/CollectionRenderingPatternUtility.cpp b/Engine/Plugins/Experimental/Dataflow/Source/DataflowAssetTools/Private/Dataflow/CollectionRenderingPatternUtility.cpp
--- a/Engine/Plugins/Experimental/Dataflow/Source/DataflowAssetTools/Private/Dataflow/CollectionRenderingPatternUtility.cpp
+++ b/Engine/Plugins/Experimental/Dataflow/Source/DataflowAssetTools/Private/Dataflow/CollectionRenderingPatternUtility.cpp
@@ -76,9 +76,19 @@ namespace UE::Dataflow
 
 				FDynamicMeshNormalOverlay* const NormalOverlay = DynamicMesh.Attributes()->PrimaryNormals();
 				NormalOverlay->CreateFromPredicate([](int ParentVID, int TriIDA, int TriIDB) { return true; }, 0.f);
-				for (int ElemIdx : NormalOverlay->ElementIndicesItr())
+				auto SetOverlayNormalsFromVertexNormals = [&](int TriangleID)
 				{
-					NormalOverlay->SetElement(ElemIdx, Normals[ElemIdx]);
+					const FIndex3i Tri = DynamicMesh.GetTriangle(TriangleID);
+					const FIndex3i NormalElementTri = NormalOverlay->GetTriangle(TriangleID);
+					for (int TriVertIndex = 0; TriVertIndex < 3; ++TriVertIndex)
+					{
+						FVector3f Normal(Normals[Remapping[Tri[TriVertIndex]]]);
+						NormalOverlay->SetElement(NormalElementTri[TriVertIndex], Normal);
+					}
+				};
+				for (const int TriangleID : DynamicMesh.TriangleIndicesItr())
+				{
+					SetOverlayNormalsFromVertexNormals(TriangleID);
 				}
 
 				DynamicMesh.Attributes()->EnablePrimaryColors();
@@ -104,7 +114,7 @@ namespace UE::Dataflow
 				FDynamicMeshMaterialAttribute* const MaterialIDAttrib = DynamicMesh.Attributes()->GetMaterialID();
 				for (const int TriangleID : DynamicMesh.TriangleIndicesItr())
 				{
-					MaterialIDAttrib->SetValue(TriangleID, FacadeMaterialID[TriangleID] - StartMaterials);
+					MaterialIDAttrib->SetValue(TriangleID, FacadeMaterialID[TriangleID + StartTriangles] - StartMaterials);
 				}
 
 				const int32 NumUVLayers = Remapping.Num() > 0 ? UVs[Remapping[0]].Num() : 0;
diff --git a/Engine/Plugins/Experimental/Dataflow/Source/DataflowEditor/Private/Dataflow/DataflowEngineRendering.cpp b/Engine/Plugins/Experimental/Dataflow/Source/DataflowEditor/Private/Dataflow/DataflowEngineRendering.cpp
--- a/Engine/Plugins/Experimental/Dataflow/Source/DataflowEditor/Private/Dataflow/DataflowEngineRendering.cpp
+++ b/Engine/Plugins/Experimental/Dataflow/Source/DataflowEditor/Private/Dataflow/DataflowEngineRendering.cpp
@@ -7,6 +7,7 @@
 #include "Dataflow/DataflowRenderingFactory.h"
 #include "Dataflow/DataflowMesh.h"
 #include "DynamicMesh/DynamicMesh3.h"
+#include "DynamicMesh/MeshNormals.h"
 #include "Field/FieldSystemTypes.h"
 #include "GeometryCollection/GeometryCollectionUtility.h"
 #include "GeometryCollection/Facades/CollectionRenderingFacade.h"
@@ -56,13 +57,54 @@ namespace UE::Dataflow
 					Vertices[VertexID] = (FVector3f)DynamicMesh.GetVertex(VertexID);
 				}
 
+				TArray<FVector3f> VertexNormals;
+				VertexNormals.AddUninitialized(Vertices.Num());
+				if (DynamicMesh.HasVertexNormals())
+				{
+					for (int32 VertexID : DynamicMesh.VertexIndicesItr())
+					{
+						VertexNormals[VertexID] = DynamicMesh.GetVertexNormal(VertexID);
+					}
+				}
+				else if (DynamicMesh.HasAttributes() && DynamicMesh.Attributes()->PrimaryNormals())
+				{
+					const UE::Geometry::FDynamicMeshNormalOverlay* const NormalOverlay = DynamicMesh.Attributes()->PrimaryNormals();
+					for (int32 VertexID : DynamicMesh.VertexIndicesItr())
+					{
+						TArray<int> OverlayElements;
+						NormalOverlay->GetVertexElements(VertexID, OverlayElements);
+
+						FVector3f AvgNormal(0.0f, 0.0f, 0.0f);
+						if (OverlayElements.Num() > 0)
+						{
+							for (int32 ElementID : OverlayElements)
+							{
+								AvgNormal += NormalOverlay->GetElement(ElementID);
+							}
+							AvgNormal /= (float)OverlayElements.Num();
+						}
+
+						VertexNormals[VertexID] = AvgNormal;
+					}
+				}
+				else
+				{
+					// No vertex normals and no overlay: compute per-vertex normals
+					UE::Geometry::FMeshNormals MeshNormals(&DynamicMesh);
+					MeshNormals.ComputeVertexNormals();
+					const TArray<FVector3d>& ComputedNormals = MeshNormals.GetNormals();
+					check(ComputedNormals.Num() == Vertices.Num());
+
+					for (int32 VertexIndex = 0; VertexIndex < Vertices.Num(); ++VertexIndex)
+					{
+						VertexNormals[VertexIndex] = (FVector3f)ComputedNormals[VertexIndex];
+					}
+				}
+
 				// Add VertexNormal and VertexColor
-				TArray<FVector3f> VertexNormals; VertexNormals.AddUninitialized(Vertices.Num());
 				TArray<FLinearColor> VertexColors; VertexColors.AddUninitialized(Vertices.Num());
 				for (int32 VertexIdx = 0; VertexIdx < VertexNormals.Num(); ++VertexIdx)
 				{
-					// TODO: Get the normal from FDynamicMesh3
-					VertexNormals[VertexIdx] = FVector3f(0.f);
 					VertexColors[VertexIdx] = FLinearColor(FDataflowEditorModule::SurfaceColor);
 				}
 
