Dataflow Editor: add a UV view mode for visualizing texture coordinates of UDataflowMeshes (can also be adopted by other data types in the future)

diff --git a/Engine/Plugins/Experimental/Dataflow/Source/DataflowEditor/Private/Dataflow/DataflowEngineRendering.cpp b/Engine/Plugins/Experimental/Dataflow/Source/DataflowEditor/Private/Dataflow/DataflowEngineRendering.cpp
--- a/Engine/Plugins/Experimental/Dataflow/Source/DataflowEditor/Private/Dataflow/DataflowEngineRendering.cpp
+++ b/Engine/Plugins/Experimental/Dataflow/Source/DataflowEditor/Private/Dataflow/DataflowEngineRendering.cpp
@@ -168,6 +168,50 @@ namespace UE::Dataflow
 			}
 		}
 
+
+		void RenderDynamicMeshUV(GeometryCollection::Facades::FRenderingFacade& RenderCollection,
+			const FString& GeometryGroupName,
+			const UE::Geometry::FDynamicMesh3& DynamicMesh)
+		{
+			using namespace UE::Geometry;
+
+			const int32 NumVertices = DynamicMesh.VertexCount();
+			const int32 NumTriangles = DynamicMesh.TriangleCount();
+			const int32 NumUVLayers = DynamicMesh.Attributes()->NumUVLayers();
+
+			if (NumVertices > 0 && NumTriangles > 0 && NumUVLayers > 0)
+			{
+				const FDynamicMeshUVOverlay* const UVOverlay = DynamicMesh.Attributes()->PrimaryUV();
+
+				TArray<FIntVector> Tris; 
+				Tris.Reserve(DynamicMesh.TriangleCount());
+				for (const int32 TriangleID : DynamicMesh.TriangleIndicesItr())
+				{
+					const FIndex3i UVTriangle = UVOverlay->GetTriangle(TriangleID);
+					Tris.Add(FIntVector(UVTriangle.A, UVTriangle.B, UVTriangle.C));
+				}
+
+				TArray<FVector3f> UVVertices; 
+				UVVertices.AddZeroed(UVOverlay->ElementCount());
+				for (int32 ElementID : UVOverlay->ElementIndicesItr())
+				{
+					UVVertices[ElementID] = FVector3f(UVOverlay->GetElement(ElementID)[0], UVOverlay->GetElement(ElementID)[1], 0.0f);
+				}
+
+				TArray<FVector3f> VertexNormals;
+				VertexNormals.Init(FVector3f(0.0, 0.0, 1.0), UVVertices.Num());
+
+				TArray<FLinearColor> VertexColors;
+				VertexColors.Init(FLinearColor(0,0,0,0), UVVertices.Num());
+
+				const int32 GeometryIndex = RenderCollection.StartGeometryGroup(GeometryGroupName);
+
+				RenderCollection.AddSurface(MoveTemp(UVVertices), MoveTemp(Tris), MoveTemp(VertexNormals), MoveTemp(VertexColors));
+
+				RenderCollection.EndGeometryGroup(GeometryIndex);
+			}
+		}
+
 	}
 
 	void RenderBasicGeometryCollection(GeometryCollection::Facades::FRenderingFacade& RenderCollection, const UE::Dataflow::FGraphRenderingState& State, TArray<FLinearColor>* VertexColorOverride = nullptr)
@@ -598,7 +642,7 @@ namespace UE::Dataflow
 
 		virtual bool CanRender(const UE::Dataflow::IDataflowConstructionViewMode& ViewMode) const override
 		{
-			return (ViewMode.GetName() == FDataflowConstruction3DViewMode::Name);
+			return (ViewMode.GetName() == FDataflowConstruction3DViewMode::Name || ViewMode.GetName() == FDataflowConstructionUVViewMode::Name);
 		}
 
 		virtual void Render(GeometryCollection::Facades::FRenderingFacade& RenderCollection, const FGraphRenderingState& State) override
@@ -612,7 +656,14 @@ namespace UE::Dataflow
 				{
 					if (const UE::Geometry::FDynamicMesh3* const DynamicMesh = Mesh->GetDynamicMesh())
 					{
-						Private::RenderDynamicMesh(RenderCollection, State.GetGuid().ToString(), *DynamicMesh, Mesh->GetMaterials());
+						if (State.GetViewMode().GetName() == Dataflow::FDataflowConstruction3DViewMode::Name)
+						{
+							Private::RenderDynamicMesh(RenderCollection, State.GetGuid().ToString(), *DynamicMesh, Mesh->GetMaterials());
+						}
+						else if (State.GetViewMode().GetName() == Dataflow::FDataflowConstructionUVViewMode::Name)
+						{
+							Private::RenderDynamicMeshUV(RenderCollection, State.GetGuid().ToString(), *DynamicMesh);
+						}
 					}
 				}
 			}
diff --git a/Engine/Plugins/Experimental/Dataflow/Source/DataflowEditor/Private/Dataflow/DataflowRenderingViewMode.cpp b/Engine/Plugins/Experimental/Dataflow/Source/DataflowEditor/Private/Dataflow/DataflowRenderingViewMode.cpp
--- a/Engine/Plugins/Experimental/Dataflow/Source/DataflowEditor/Private/Dataflow/DataflowRenderingViewMode.cpp
+++ b/Engine/Plugins/Experimental/Dataflow/Source/DataflowEditor/Private/Dataflow/DataflowRenderingViewMode.cpp
@@ -57,6 +57,22 @@ namespace UE::Dataflow
 		return FText(LOCTEXT("3DViewTooltipText", "Default 3D View"));
 	}
 
+	FName FDataflowConstructionUVViewMode::Name = FName("UVView");
+
+	FName FDataflowConstructionUVViewMode::GetName() const
+	{
+		return Name;
+	}
+
+	FText FDataflowConstructionUVViewMode::GetButtonText() const
+	{
+		return FText(LOCTEXT("UVViewButtonText", "UVView"));
+	}
+
+	FText FDataflowConstructionUVViewMode::GetTooltipText() const
+	{
+		return FText(LOCTEXT("UVViewTooltipText", "UV/Texture Coordinate View"));
+	}
 
 	//
 	// Factory
@@ -67,6 +83,7 @@ namespace UE::Dataflow
 		// Default modes
 		ViewModeMap.Add(FDataflowConstruction2DViewMode::Name, MakeUnique<FDataflowConstruction2DViewMode>());
 		ViewModeMap.Add(FDataflowConstruction3DViewMode::Name, MakeUnique<FDataflowConstruction3DViewMode>());
+		ViewModeMap.Add(FDataflowConstructionUVViewMode::Name, MakeUnique<FDataflowConstructionUVViewMode>());
 	}
 
 	FRenderingViewModeFactory& FRenderingViewModeFactory::GetInstance()
diff --git a/Engine/Plugins/Experimental/Dataflow/Source/DataflowEditor/Public/Dataflow/DataflowRenderingViewMode.h b/Engine/Plugins/Experimental/Dataflow/Source/DataflowEditor/Public/Dataflow/DataflowRenderingViewMode.h
--- a/Engine/Plugins/Experimental/Dataflow/Source/DataflowEditor/Public/Dataflow/DataflowRenderingViewMode.h
+++ b/Engine/Plugins/Experimental/Dataflow/Source/DataflowEditor/Public/Dataflow/DataflowRenderingViewMode.h
@@ -77,6 +77,19 @@ namespace UE::Dataflow
 		DATAFLOWEDITOR_API virtual FText GetTooltipText() const override;
 	};
 
+	// UV view mode (same as 2D but with a different name and button text)
+	class FDataflowConstructionUVViewMode : public FDataflowConstruction2DViewModeBase
+	{
+	public:
+		DATAFLOWEDITOR_API static FName Name;
+		virtual ~FDataflowConstructionUVViewMode() = default;
+	private:
+		DATAFLOWEDITOR_API virtual FName GetName() const override;
+		DATAFLOWEDITOR_API virtual FText GetButtonText() const override;
+		DATAFLOWEDITOR_API virtual FText GetTooltipText() const override;
+	};
+
+
 	//
 	// ViewMode registry/factory
 	//
