Dataflow Editor: remember the last-selected view mode for nodes with more than one view mode available

diff --git a/Engine/Plugins/Experimental/Dataflow/Source/DataflowEditor/Private/Dataflow/DataflowEditorMode.cpp b/Engine/Plugins/Experimental/Dataflow/Source/DataflowEditor/Private/Dataflow/DataflowEditorMode.cpp
--- a/Engine/Plugins/Experimental/Dataflow/Source/DataflowEditor/Private/Dataflow/DataflowEditorMode.cpp
+++ b/Engine/Plugins/Experimental/Dataflow/Source/DataflowEditor/Private/Dataflow/DataflowEditorMode.cpp
@@ -829,6 +829,32 @@ void UDataflowEditorMode::SetConstructionViewMode(const FName& NewViewModeName)
 			DataflowToolBuilder->SceneStateChanged(ActiveTool, SceneState);
 		}
 	}
+
+
+	if (const TSharedPtr<const SDataflowGraphEditor> PinnedDataflowGraphEditor = DataflowGraphEditor.Pin())
+	{
+		if (const UEdGraphNode* const SelectedNode = PinnedDataflowGraphEditor->GetSingleSelectedNode())
+		{
+			if (const UDataflowEdNode* const SelectedDataflowEdNode = Cast<UDataflowEdNode>(SelectedNode))
+			{
+				if (const TObjectPtr<UDataflowBaseContent>& EditorContent = ConstructionScene->GetEditorContent())
+				{
+					TArray<FName> ViewModesForNode;
+					UE::Dataflow::GetViewModesForNode(*SelectedDataflowEdNode, *EditorContent, ViewModesForNode);
+
+					check(ViewModesForNode.Contains(NewViewModeName));
+
+					ViewModesForNode.Sort([](const FName& A, const FName& B)
+					{
+						return A.GetNumber() < B.GetNumber();
+					});
+
+					NodeViewModeHistory.FindOrAdd(ViewModesForNode) = NewViewModeName;
+				}
+			}
+		}
+	}
+
 }
 
 const UE::Dataflow::IDataflowConstructionViewMode* UDataflowEditorMode::GetConstructionViewMode() const
@@ -1120,5 +1146,15 @@ UEdGraphNode* UDataflowEditorMode::CreateAndConnectNewNode(const FName& NewNodeT
 	return nullptr;
 }
 
+const FName* UDataflowEditorMode::FindRecentlyUsedViewMode(const TArray<FName>& AvailableViewModes) const
+{
+	TArray<FName> Key = AvailableViewModes;
+	Key.Sort([](const FName& A, const FName& B)
+	{
+		return A.GetNumber() < B.GetNumber();
+	});
+	return NodeViewModeHistory.Find(Key);
+}
+
 #undef LOCTEXT_NAMESPACE
 
diff --git a/Engine/Plugins/Experimental/Dataflow/Source/DataflowEditor/Private/Dataflow/DataflowEditorToolkit.cpp b/Engine/Plugins/Experimental/Dataflow/Source/DataflowEditor/Private/Dataflow/DataflowEditorToolkit.cpp
--- a/Engine/Plugins/Experimental/Dataflow/Source/DataflowEditor/Private/Dataflow/DataflowEditorToolkit.cpp
+++ b/Engine/Plugins/Experimental/Dataflow/Source/DataflowEditor/Private/Dataflow/DataflowEditorToolkit.cpp
@@ -1224,17 +1224,26 @@ void FDataflowEditorToolkit::OnNodeSelectionChanged(const TSet<UObject*>& InNewS
 
 	if (UDataflowEditorMode* const DataflowMode = Cast<UDataflowEditorMode>(EditorModeManager->GetActiveScriptableMode(UDataflowEditorMode::EM_DataflowEditorModeId)))
 	{
-		bool bFoundViewMode = true;
-
 		if (PrimarySelection && GetEditorContent())
 		{
-			if (!UE::Dataflow::CanRenderNodeOutput(*PrimarySelection, *GetEditorContent(), *DataflowMode->GetConstructionViewMode()))
+			// Check the most recently used view mode for this "kind" of node -- where "kind" means nodes that have the same set of valid view modes
+			auto FindRecentViewMode = [&](FName& OutViewModeName) -> bool
 			{
-				// Selected node can't render with the current view mode. Check through available view modes and see if it can render with any of them
+				TArray<FName> ViewModesForNode;
+				UE::Dataflow::GetViewModesForNode(*PrimarySelection, *GetEditorContent(), ViewModesForNode);
 
-				bFoundViewMode = false;
+				if (const FName* const FoundViewMode = DataflowMode->FindRecentlyUsedViewMode(ViewModesForNode))
+				{
+					OutViewModeName = *FoundViewMode;
+					return true;
+				}
+				return false;
+			};
 
-				TArray<UE::Dataflow::FRenderingParameter> RenderingParameters = PrimarySelection->GetRenderParameters();
+			// Get the first view mode that this node can use
+			auto FindAnyValidViewMode = [&](FName& OutViewModeName) -> bool
+			{
+				const TArray<UE::Dataflow::FRenderingParameter> RenderingParameters = PrimarySelection->GetRenderParameters();
 				for (const UE::Dataflow::FRenderingParameter& Param : RenderingParameters)
 				{
 					const FName NodeOutputTypeName = Param.Type;
@@ -1242,29 +1251,36 @@ void FDataflowEditorToolkit::OnNodeSelectionChanged(const TSet<UObject*>& InNewS
 					for (const TPair<FName, TUniquePtr<UE::Dataflow::IDataflowConstructionViewMode>>& ViewMode : UE::Dataflow::FRenderingViewModeFactory::GetInstance().GetViewModes())
 					{
 						check(ViewMode.Value.IsValid());
-
 						const bool bCanRender = UE::Dataflow::CanRenderNodeOutput(*PrimarySelection, *GetEditorContent(), *ViewMode.Value);
 
 						if (bCanRender)
 						{
-							DataflowMode->SetConstructionViewMode(ViewMode.Key);
-							bFoundViewMode = true;
-							break;
+							OutViewModeName = ViewMode.Key;
+							return true;
 						}
 					}
-
-					if (bFoundViewMode)
-					{
-						break;
-					}
 				}
-			}
-		}
+				return false;
+			};
 
-		if (!bFoundViewMode)
-		{
-			// TODO: Clear and disable View Mode Button. For now set default mode to the built-in 3D view mode.
-			DataflowMode->SetConstructionViewMode(UE::Dataflow::FDataflowConstruction3DViewMode::Name);
+			FName NewViewMode;
+			if (FindRecentViewMode(NewViewMode))               // Is there a recently used view mode we can switch to?
+			{
+				DataflowMode->SetConstructionViewMode(NewViewMode);
+			}
+			else if (UE::Dataflow::CanRenderNodeOutput(*PrimarySelection, *GetEditorContent(), *DataflowMode->GetConstructionViewMode()))        // Can we use the current view mode?
+			{
+				// no need to switch
+			}		
+			else if (FindAnyValidViewMode(NewViewMode))             // Is there *any* view mode that works?
+			{
+				DataflowMode->SetConstructionViewMode(NewViewMode);
+			}
+			else                                                    // No valid view mode found
+			{
+				// TODO: We should clear and disable View Mode Button. For now set default mode to the built-in 3D view mode.
+				DataflowMode->SetConstructionViewMode(UE::Dataflow::FDataflowConstruction3DViewMode::Name);
+			}
 		}
 
 		DataflowMode->SetPendingNodeSelectionChanged(false);
diff --git a/Engine/Plugins/Experimental/Dataflow/Source/DataflowEditor/Private/Dataflow/DataflowEditorUtil.cpp b/Engine/Plugins/Experimental/Dataflow/Source/DataflowEditor/Private/Dataflow/DataflowEditorUtil.cpp
--- a/Engine/Plugins/Experimental/Dataflow/Source/DataflowEditor/Private/Dataflow/DataflowEditorUtil.cpp
+++ b/Engine/Plugins/Experimental/Dataflow/Source/DataflowEditor/Private/Dataflow/DataflowEditorUtil.cpp
@@ -160,4 +160,22 @@ namespace UE::Dataflow
 		}
 	}
 
+
+	void GetViewModesForNode(const UDataflowEdNode& EdNode, const UDataflowBaseContent& EditorContent, TArray<FName>& OutValidViewModeNames)
+	{
+		TArray<FName> AllViewModeNames;
+		FRenderingViewModeFactory::GetInstance().GetViewModes().GetKeys(AllViewModeNames);
+
+		for (const FName& ViewModeName : AllViewModeNames)
+		{
+			if (const IDataflowConstructionViewMode* const ViewMode = FRenderingViewModeFactory::GetInstance().GetViewMode(ViewModeName))
+			{
+				if (CanRenderNodeOutput(EdNode, EditorContent, *ViewMode))
+				{
+					OutValidViewModeNames.Add(ViewModeName);
+				}
+			}
+		}
+	}
+
 }
\ No newline at end of file
diff --git a/Engine/Plugins/Experimental/Dataflow/Source/DataflowEditor/Private/Dataflow/DataflowEditorUtil.h b/Engine/Plugins/Experimental/Dataflow/Source/DataflowEditor/Private/Dataflow/DataflowEditorUtil.h
--- a/Engine/Plugins/Experimental/Dataflow/Source/DataflowEditor/Private/Dataflow/DataflowEditorUtil.h
+++ b/Engine/Plugins/Experimental/Dataflow/Source/DataflowEditor/Private/Dataflow/DataflowEditorUtil.h
@@ -50,4 +50,5 @@ namespace UE::Dataflow
 
 	bool CanRenderNodeOutput(const UDataflowEdNode& EdNode, const UDataflowBaseContent& EditorContent, const UE::Dataflow::IDataflowConstructionViewMode& ViewMode);
 	void RenderNodeOutput(GeometryCollection::Facades::FRenderingFacade& Facade, const UDataflowEdNode& Node, const UDataflowBaseContent& EditorContent);
+	void GetViewModesForNode(const UDataflowEdNode& EdNode, const UDataflowBaseContent& EditorContent, TArray<FName>& OutValidViewModeNames);
 }
\ No newline at end of file
diff --git a/Engine/Plugins/Experimental/Dataflow/Source/DataflowEditor/Public/Dataflow/DataflowEditorMode.h b/Engine/Plugins/Experimental/Dataflow/Source/DataflowEditor/Public/Dataflow/DataflowEditorMode.h
--- a/Engine/Plugins/Experimental/Dataflow/Source/DataflowEditor/Public/Dataflow/DataflowEditorMode.h
+++ b/Engine/Plugins/Experimental/Dataflow/Source/DataflowEditor/Public/Dataflow/DataflowEditorMode.h
@@ -89,6 +89,12 @@ public:
 
 	void StartToolForSelectedNode(const UObject* SelectedNode);
 
+	/** Each node has a set of view modes it can use. For each distinct set of available view modes, we store the last mode that the user has chosen.
+	 *  Use this function to find the most recently used view mode for nodes having the same set of available modes. 
+	 *	Returns nullptr if there is nothing stored for the input set of modes 
+	 */
+	const FName* FindRecentlyUsedViewMode(const TArray<FName>& AvailableViewModes) const;
+
 private:
 	friend class FDataflowEditorToolkit;
 	friend class FDataflowEditorModeToolkit;
@@ -234,5 +240,9 @@ private:
 	// Timestamp for telemetry
 	FDateTime LastModeStartTimestamp;
 
+	// Each node has a set of view modes it can use. For each distinct set of available view modes, we store the last mode that the user has chosen.
+	// This will enable switching between different kinds of nodes and not have the view mode reset to a default every time.
+	TMap<TArray<FName>, FName> NodeViewModeHistory;
+	
 };
 
