DataflowEditorMode: fix FName comparison operator

diff --git a/Engine/Plugins/Experimental/Dataflow/Source/DataflowEditor/Private/Dataflow/DataflowEditorMode.cpp b/Engine/Plugins/Experimental/Dataflow/Source/DataflowEditor/Private/Dataflow/DataflowEditorMode.cpp
--- a/Engine/Plugins/Experimental/Dataflow/Source/DataflowEditor/Private/Dataflow/DataflowEditorMode.cpp
+++ b/Engine/Plugins/Experimental/Dataflow/Source/DataflowEditor/Private/Dataflow/DataflowEditorMode.cpp
@@ -846,7 +846,7 @@ void UDataflowEditorMode::SetConstructionViewMode(const FName& NewViewModeName)
 
 					ViewModesForNode.Sort([](const FName& A, const FName& B)
 					{
-						return A.GetNumber() < B.GetNumber();
+						return A.FastLess(B);
 					});
 
 					NodeViewModeHistory.FindOrAdd(ViewModesForNode) = NewViewModeName;
@@ -1151,7 +1151,7 @@ const FName* UDataflowEditorMode::FindRecentlyUsedViewMode(const TArray<FName>&
 	TArray<FName> Key = AvailableViewModes;
 	Key.Sort([](const FName& A, const FName& B)
 	{
-		return A.GetNumber() < B.GetNumber();
+		return A.FastLess(B);
 	});
 	return NodeViewModeHistory.Find(Key);
 }
