Dataflow - Added Freeze/Unfreeze and Refreeze to the Dataflow node UI.

#rnx

diff --git a/Engine/Plugins/Experimental/Dataflow/Source/DataflowEditor/Private/Dataflow/DataflowEditorModule.cpp b/Engine/Plugins/Experimental/Dataflow/Source/DataflowEditor/Private/Dataflow/DataflowEditorModule.cpp
--- a/Engine/Plugins/Experimental/Dataflow/Source/DataflowEditor/Private/Dataflow/DataflowEditorModule.cpp
+++ b/Engine/Plugins/Experimental/Dataflow/Source/DataflowEditor/Private/Dataflow/DataflowEditorModule.cpp
@@ -8,6 +8,7 @@
 #include "Dataflow/DataflowEditorToolkit.h"
 #include "Dataflow/DataflowEditorCommands.h"
 #include "Dataflow/DataflowEngineRendering.h"
+#include "Dataflow/DataflowFreezeActionsCustomization.h"
 #include "Dataflow/DataflowFunctionProperty.h"
 #include "Dataflow/DataflowFunctionPropertyCustomization.h"
 #include "Dataflow/DataflowSNodeFactories.h"
@@ -80,6 +81,7 @@ void FDataflowEditorModule::StartupModule()
 		PropertyModule->RegisterCustomPropertyTypeLayout(ScalarVertexPropertyGroupName, FOnGetPropertyTypeCustomizationInstance::CreateStatic(&UE::Dataflow::FScalarVertexPropertyGroupCustomization::MakeInstance));
 		PropertyModule->RegisterCustomPropertyTypeLayout(DataflowFunctionPropertyName, FOnGetPropertyTypeCustomizationInstance::CreateStatic(&UE::Dataflow::FFunctionPropertyCustomization::MakeInstance));
 		PropertyModule->RegisterCustomPropertyTypeLayout("DataflowVariableOverrides", FOnGetPropertyTypeCustomizationInstance::CreateStatic(&FDataflowVariableOverridesDetails::MakeInstance));
+		PropertyModule->RegisterCustomPropertyTypeLayout(FDataflowFreezeActions::StaticStruct()->GetFName(), FOnGetPropertyTypeCustomizationInstance::CreateStatic(&UE::Dataflow::FFreezeActionsCustomization::MakeInstance));
 	}
 
 	UE::Dataflow::RenderingCallbacks();
@@ -105,6 +107,7 @@ void FDataflowEditorModule::ShutdownModule()
 	{
 		PropertyModule->UnregisterCustomPropertyTypeLayout(ScalarVertexPropertyGroupName);
 		PropertyModule->UnregisterCustomPropertyTypeLayout(DataflowFunctionPropertyName);
+		PropertyModule->UnregisterCustomPropertyTypeLayout(FDataflowFreezeActions::StaticStruct()->GetFName());
 	}
 
 	FDataflowEditorCommands::Unregister();
diff --git a/Engine/Plugins/Experimental/Dataflow/Source/DataflowEditor/Private/Dataflow/DataflowFreezeActionsCustomization.cpp b/Engine/Plugins/Experimental/Dataflow/Source/DataflowEditor/Private/Dataflow/DataflowFreezeActionsCustomization.cpp
new file mode 100644
--- /dev/null
+++ b/Engine/Plugins/Experimental/Dataflow/Source/DataflowEditor/Private/Dataflow/DataflowFreezeActionsCustomization.cpp
@@ -0,0 +1,191 @@
+// Copyright Epic Games, Inc. All Rights Reserved.
+
+#include "Dataflow/DataflowFreezeActionsCustomization.h"
+#include "Dataflow/DataflowEditorStyle.h"
+#include "Dataflow/DataflowGraphEditor.h"
+#include "Dataflow/DataflowNode.h"
+#include "Dataflow/DataflowTerminalNode.h"
+#include "Widgets/Images/SImage.h"
+#include "Widgets/Input/SButton.h"
+#include "Widgets/Layout/SWrapBox.h"
+#include "Widgets/Text/STextBlock.h"
+#include "DetailWidgetRow.h"
+#include "PropertyHandle.h"
+
+#define LOCTEXT_NAMESPACE "DataflowNodeCustomization"
+
+namespace UE::Dataflow
+{
+	namespace Private
+	{
+		FDataflowNode* GetDataflowNode(TSharedRef<IPropertyHandle> StructPropertyHandle)
+		{
+			const TSharedPtr<IPropertyHandle> ParentHandle = StructPropertyHandle->GetParentHandle();
+			if (const TSharedPtr<IPropertyHandleStruct> ParentHandleStruct = ParentHandle->AsStruct())
+			{
+				if (const TSharedPtr<FStructOnScope> StructOnScope = ParentHandleStruct->GetStructData())
+				{
+					if (const UStruct* const Struct = StructOnScope->GetStruct())
+					{
+						if (Struct->IsChildOf<FDataflowNode>())
+						{
+							return reinterpret_cast<FDataflowNode*>(StructOnScope->GetStructMemory());
+						}
+					}
+				}
+			}
+			return nullptr;
+		}
+	}
+
+	TSharedRef<IPropertyTypeCustomization> FFreezeActionsCustomization::MakeInstance()
+	{
+		return MakeShareable(new FFreezeActionsCustomization);
+	}
+
+	void FFreezeActionsCustomization::CustomizeHeader(TSharedRef<IPropertyHandle> StructPropertyHandle, FDetailWidgetRow& HeaderRow, IPropertyTypeCustomizationUtils& CustomizationUtils)
+	{
+		// Keep a weak pointer to the graph editor that is creating this customization
+		TWeakPtr<const SDataflowGraphEditor> DataflowGraphEditor = SDataflowGraphEditor::GetSelectedGraphEditor();
+
+		// Create WrapBox
+		TSharedPtr<SWrapBox> WrapBox;
+		HeaderRow
+		[
+			SAssignNew(WrapBox, SWrapBox)
+			.PreferredSize(2000)  // Copied from FObjectDetails::AddCallInEditorMethods()
+			.UseAllottedSize(true)
+		];
+
+		// Create button image + text
+		WrapBox->AddSlot()
+			.Padding(0.f, 0.f, 5.f, 3.f)
+			[
+				SNew(SButton)
+					.Text(LOCTEXT("Freeze", "Freeze"))
+					.ToolTipText_Lambda([StructPropertyHandle]() -> FText
+						{
+							const FDataflowNode* const DataflowNode = Private::GetDataflowNode(StructPropertyHandle);
+							return DataflowNode && !DataflowNode->IsFrozen() ?
+								LOCTEXT("FreezeToolTip", "Freeze all outputs, and disable evaluation for this node.") :
+								LOCTEXT("UnfreezeToolTip", "UInfreeze all outputs, and re-enable evaluation for this node.");
+						})
+					.OnClicked_Lambda(
+						[StructPropertyHandle, DataflowGraphEditor]() -> FReply
+						{
+							if (FDataflowNode* const DataflowNode = Private::GetDataflowNode(StructPropertyHandle))
+							{
+								// Retrieve context if any
+								const TSharedPtr<const SDataflowGraphEditor> DataflowGraphEditorPtr = DataflowGraphEditor.Pin();
+								const TSharedPtr<UE::Dataflow::FContext> Context =
+									DataflowGraphEditorPtr ? DataflowGraphEditorPtr->GetDataflowContext() :
+									TSharedPtr<UE::Dataflow::FContext>();
+								if (DataflowNode->IsFrozen())
+								{
+									DataflowNode->Unfreeze(*Context);
+								}
+								else
+								{
+									DataflowNode->Freeze(*Context);
+								}
+							}
+							return FReply::Handled();
+						})
+					.ContentPadding(FMargin(0.f, 2.f))  // Too much horizontal padding otherwise (default is 4, 2)
+					.Visibility_Lambda([StructPropertyHandle]() -> EVisibility
+						{
+							const FDataflowNode* const DataflowNode = Private::GetDataflowNode(StructPropertyHandle);
+							return DataflowNode && !DataflowNode->AsType<const FDataflowTerminalNode>() ? EVisibility::Visible : EVisibility::Collapsed;
+						})
+					[
+						SNew(SHorizontalBox)
+							+ SHorizontalBox::Slot()
+							.AutoWidth()
+							.Padding(0, 2)
+							[
+								SNew(SImage)
+									.DesiredSizeOverride(FVector2D(16, 16))
+									.Image(FDataflowEditorStyle::Get().GetBrush("Dataflow.FreezeNode"))
+							]
+							+ SHorizontalBox::Slot()
+							.AutoWidth()
+							.Padding(5.f, 0.f)
+							+ SHorizontalBox::Slot()
+							.AutoWidth()
+							.VAlign(VAlign_Center)
+							[
+								SNew(STextBlock)
+									.Text_Lambda([StructPropertyHandle]() -> FText
+										{
+											const FDataflowNode* const DataflowNode = Private::GetDataflowNode(StructPropertyHandle);
+											return DataflowNode && !DataflowNode->IsFrozen() ?
+												LOCTEXT("Freeze", "Freeze") :
+												LOCTEXT("Unfreeze", "Unfreeze");
+										})
+							]
+					]
+			];
+		// Add Refreeze button
+		WrapBox->AddSlot()
+			.Padding(0.f, 0.f, 5.f, 3.f)
+			[
+				SNew(SButton)
+					.Text(LOCTEXT("Refreeze", "Refreeze"))
+					.ToolTipText(LOCTEXT("RefreezeToolTip", "Unfreeze all outputs, redo this node's evaluation, and freeze all updated outputs again."))
+					.OnClicked_Lambda(
+						[StructPropertyHandle, DataflowGraphEditor]() -> FReply
+						{
+							if (FDataflowNode* const DataflowNode = Private::GetDataflowNode(StructPropertyHandle))
+							{
+								if (DataflowNode->IsFrozen())
+								{
+									// Retrieve context if any
+									const TSharedPtr<const SDataflowGraphEditor> DataflowGraphEditorPtr = DataflowGraphEditor.Pin();
+									const TSharedPtr<UE::Dataflow::FContext> Context =
+										DataflowGraphEditorPtr ? DataflowGraphEditorPtr->GetDataflowContext() :
+										TSharedPtr<UE::Dataflow::FContext>();
+									// Refreeze
+									DataflowNode->Unfreeze(*Context);
+									Context->Evaluate(DataflowNode, nullptr);
+									DataflowNode->Freeze(*Context);
+								}
+							}
+							return FReply::Handled();
+						})
+					.ContentPadding(FMargin(0.f, 2.f))  // Too much horizontal padding otherwise (default is 4, 2)
+					.IsEnabled_Lambda([StructPropertyHandle]() -> bool
+						{
+							const FDataflowNode* const DataflowNode = Private::GetDataflowNode(StructPropertyHandle);
+							return DataflowNode && DataflowNode->IsFrozen();
+						})
+					.Visibility_Lambda([StructPropertyHandle]() -> EVisibility
+						{
+							const FDataflowNode* const DataflowNode = Private::GetDataflowNode(StructPropertyHandle);
+							return DataflowNode && !DataflowNode->AsType<const FDataflowTerminalNode>() ? EVisibility::Visible : EVisibility::Collapsed;
+						})
+					[
+						SNew(SHorizontalBox)
+							+ SHorizontalBox::Slot()
+							.AutoWidth()
+							.Padding(0, 2)
+							[
+								SNew(SImage)
+									.DesiredSizeOverride(FVector2D(16, 16))
+									.Image(FAppStyle::GetBrush("Icons.Refresh"))
+							]
+							+ SHorizontalBox::Slot()
+							.AutoWidth()
+							.Padding(5.f, 0.f)
+							+ SHorizontalBox::Slot()
+							.AutoWidth()
+							.VAlign(VAlign_Center)
+							[
+								SNew(STextBlock)
+									.Text(LOCTEXT("Refreeze", "Refreeze"))
+							]
+					]
+			];
+	}
+}  // End namespace UE::Dataflow
+
+#undef LOCTEXT_NAMESPACE
diff --git a/Engine/Plugins/Experimental/Dataflow/Source/DataflowEditor/Private/Dataflow/DataflowFreezeActionsCustomization.h b/Engine/Plugins/Experimental/Dataflow/Source/DataflowEditor/Private/Dataflow/DataflowFreezeActionsCustomization.h
new file mode 100644
--- /dev/null
+++ b/Engine/Plugins/Experimental/Dataflow/Source/DataflowEditor/Private/Dataflow/DataflowFreezeActionsCustomization.h
@@ -0,0 +1,23 @@
+// Copyright Epic Games, Inc. All Rights Reserved.
+
+#pragma once
+
+#include "IPropertyTypeCustomization.h"
+
+namespace UE::Dataflow
+{
+	/**
+	 * Customization for the Dataflow node UI.
+	 */
+	class FFreezeActionsCustomization final : public IPropertyTypeCustomization
+	{
+	public:
+		static TSharedRef<IPropertyTypeCustomization> MakeInstance();
+
+	private:
+		//~ Begin IPropertyTypeCustomization interface
+		virtual void CustomizeHeader(TSharedRef<IPropertyHandle> StructPropertyHandle, FDetailWidgetRow& HeaderRow, IPropertyTypeCustomizationUtils& customizationutils) override;
+		virtual void CustomizeChildren(TSharedRef<IPropertyHandle> /*StructPropertyHandle*/, IDetailChildrenBuilder& /*ChildBuilder*/, IPropertyTypeCustomizationUtils& /*CustomizationUtils*/) override {}
+		//~ End IPropertyTypeCustomization interface
+	};
+}
diff --git a/Engine/Source/Runtime/Experimental/Dataflow/Core/Public/Dataflow/DataflowNode.h b/Engine/Source/Runtime/Experimental/Dataflow/Core/Public/Dataflow/DataflowNode.h
--- a/Engine/Source/Runtime/Experimental/Dataflow/Core/Public/Dataflow/DataflowNode.h
+++ b/Engine/Source/Runtime/Experimental/Dataflow/Core/Public/Dataflow/DataflowNode.h
@@ -25,6 +25,13 @@ namespace UE::Dataflow {
 }
 class IDataflowDebugDrawInterface;
 
+/** Empty class for the freeze UI customization. */
+USTRUCT()
+struct FDataflowFreezeActions final
+{
+	GENERATED_USTRUCT_BODY()
+};
+
 /**
 * FNode
 *		Base class for node based evaluation within the Dataflow graph. 
@@ -56,7 +63,7 @@ PRAGMA_ENABLE_DEPRECATION_WARNINGS
 	TMap< int, FDataflowOutput* > Outputs;
 
 	UE_DEPRECATED(5.6, "Will be made private in 5.8, use IsActive function instead.")
-	UPROPERTY(EditAnywhere, Category = "Dataflow")
+	UPROPERTY(EditAnywhere, Category = "Dataflow", Meta = (DisplayAfter = "FreezeActions"))
 	bool bActive = true;
 
 	UPROPERTY(EditAnywhere, Category = "Dataflow", meta = (InlineEditConditionToggle))
@@ -790,6 +797,12 @@ private:
 	TMap<UE::Dataflow::FConnectionKey, FDataflowOutput*> ExpandedOutputs;
 	TMap<uint32, const FArrayProperty*> OutputArrayProperties; // Used to calculate ContainerElementOffsets
 
+#if WITH_EDITORONLY_DATA
+	// Freeze node customization
+	UPROPERTY(EditAnywhere, Category = "Dataflow")
+	FDataflowFreezeActions FreezeActions;
+#endif
+
 	UPROPERTY()
 	FInstancedPropertyBag FrozenProperties;  // Property bag for frozen outputs
 
