Dataflow : Add support for array of dynamic mesh anytype so we can connect a single dynamic mesh in an array input

diff --git a/Engine/Plugins/Experimental/Dataflow/Source/DataflowEnginePlugin/Private/Dataflow/DataflowEngineAnyTypes.cpp b/Engine/Plugins/Experimental/Dataflow/Source/DataflowEnginePlugin/Private/Dataflow/DataflowEngineAnyTypes.cpp
new file mode 100644
--- /dev/null
+++ b/Engine/Plugins/Experimental/Dataflow/Source/DataflowEnginePlugin/Private/Dataflow/DataflowEngineAnyTypes.cpp
@@ -0,0 +1,14 @@
+// Copyright Epic Games, Inc. All Rights Reserved.
+
+#include "Dataflow/DataflowEngineAnyTypes.h"
+
+#include "Dataflow/DataflowAnyTypeRegistry.h"
+
+namespace UE::Dataflow
+{
+	void RegisterEngineAnyTypes()
+	{
+		UE_DATAFLOW_REGISTER_ANYTYPE(FDataflowDynamicMeshArray);
+	}
+}
+
diff --git a/Engine/Plugins/Experimental/Dataflow/Source/DataflowEnginePlugin/Private/Dataflow/DataflowEnginePlugin.cpp b/Engine/Plugins/Experimental/Dataflow/Source/DataflowEnginePlugin/Private/Dataflow/DataflowEnginePlugin.cpp
--- a/Engine/Plugins/Experimental/Dataflow/Source/DataflowEnginePlugin/Private/Dataflow/DataflowEnginePlugin.cpp
+++ b/Engine/Plugins/Experimental/Dataflow/Source/DataflowEnginePlugin/Private/Dataflow/DataflowEnginePlugin.cpp
@@ -5,6 +5,7 @@
 #include "CoreMinimal.h"
 #include "Modules/ModuleManager.h"
 #include "Dataflow/DataflowEngineContextCaching.h"
+#include "Dataflow/DataflowEngineAnyTypes.h"
 #include "GeometryCollection/GeometryCollectionUtility.h"
 #include "GeometryCollection/Facades/CollectionRenderingFacade.h"
 #include "GeometryCollection/ManagedArrayCollection.h"
@@ -29,6 +30,7 @@ IMPLEMENT_MODULE( FDataflowEnginePlugin, DataflowEnginePlugin)
 void FDataflowEnginePlugin::StartupModule()
 {
 	UE::Dataflow::ContextCachingCallbacks();
+	UE::Dataflow::RegisterEngineAnyTypes();
 	FModuleManager::Get().LoadModule("DataflowEngine");
 	FModuleManager::Get().LoadModule("DataflowSimulation");
 }
diff --git a/Engine/Plugins/Experimental/Dataflow/Source/DataflowEnginePlugin/Public/Dataflow/DataflowEngineAnyTypes.h b/Engine/Plugins/Experimental/Dataflow/Source/DataflowEnginePlugin/Public/Dataflow/DataflowEngineAnyTypes.h
new file mode 100644
--- /dev/null
+++ b/Engine/Plugins/Experimental/Dataflow/Source/DataflowEnginePlugin/Public/Dataflow/DataflowEngineAnyTypes.h
@@ -0,0 +1,72 @@
+// Copyright Epic Games, Inc. All Rights Reserved.
+
+#pragma once
+
+#include "Dataflow/DataflowAnyType.h"
+#include "Dataflow/DataflowTypePolicy.h"
+#include "UObject/ObjectPtr.h"
+
+#include "DataflowEngineAnyTypes.generated.h"
+
+class UDynamicMesh;
+
+// todo : move this to general type policy file ?
+template<typename T>
+struct FDataflowTypedArrayTypePolicy : 
+public TDataflowMultiTypePolicy<T, TArray<T>>
+{
+};
+
+/**
+* Dynamic mesh array TArray with compatibility to plug one single element
+*/
+USTRUCT()
+struct FDataflowDynamicMeshArray : public FDataflowAnyType
+{
+	using FPolicyType = FDataflowTypedArrayTypePolicy<TObjectPtr<UDynamicMesh>>;
+	using FStorageType = TArray<TObjectPtr<UDynamicMesh>>;
+
+	GENERATED_BODY()
+
+	UPROPERTY(EditAnywhere, Category=Value)
+	TArray<TObjectPtr<UDynamicMesh>> Value;
+};
+
+UE_DATAFLOW_POLICY_DECLARE_TYPENAME(TObjectPtr<UDynamicMesh>)
+
+template<>
+struct FDataflowConverter<TArray<TObjectPtr<UDynamicMesh>>>
+{
+	template <typename TFromType>
+	static void From(const TFromType& From, TArray<TObjectPtr<UDynamicMesh>>& To)
+	{
+		if constexpr (std::is_same_v<TFromType, TObjectPtr<UDynamicMesh>>)
+		{
+			To.Init(From, 1);
+		}
+		else if constexpr (std::is_same_v<TFromType, TArray<TObjectPtr<UDynamicMesh>>>)
+		{
+			To = From;
+		}
+	}
+
+	template <typename TToType>
+	static void To(const TArray<TObjectPtr<UDynamicMesh>>& From, TToType& To)
+	{
+		if constexpr (std::is_same_v<TToType, TObjectPtr<UDynamicMesh>>)
+		{
+			To = (From.Num()) ? From[0] : TObjectPtr<UDynamicMesh>();
+		}
+		else if constexpr (std::is_same_v<TToType, TArray<TObjectPtr<UDynamicMesh>>>)
+		{
+			To = From;
+		}
+
+		
+	}
+};
+
+namespace UE::Dataflow
+{
+	void RegisterEngineAnyTypes();
+}
