Dataflow:
- Fixed the priorities to display node states

diff --git a/Engine/Plugins/Experimental/Dataflow/Source/DataflowEditor/Private/Dataflow/DataflowSNode.cpp b/Engine/Plugins/Experimental/Dataflow/Source/DataflowEditor/Private/Dataflow/DataflowSNode.cpp
--- a/Engine/Plugins/Experimental/Dataflow/Source/DataflowEditor/Private/Dataflow/DataflowSNode.cpp
+++ b/Engine/Plugins/Experimental/Dataflow/Source/DataflowEditor/Private/Dataflow/DataflowSNode.cpp
@@ -287,49 +287,61 @@ TArray<FOverlayWidgetInfo> SDataflowEdNode::GetOverlayWidgets(bool bSelected, co
 	return Widgets;
 }
 
-const FSlateColor CExperimentalColor(FLinearColor(26.f/255.f, 188.f/255.f, 156.f/255.f, 255.f)); // Turquoise
-const FSlateColor CDeprecatedColor(FLinearColor(46.f / 255.f, 204.f / 255.f, 113.f / 255.f, 255.f)); // Emerald
-const FSlateColor CFailedColor(FLinearColor(243.f / 255.f, 156.f / 255.f, 18.f / 255.f, 255.f)); // Orange
+const FSlateColor CExperimentalColor(FColor(26, 188, 156, 255)); // Turquoise
+const FSlateColor CDeprecatedColor(FColor(46, 204, 113, 255)); // Emerald
+const FSlateColor CFailedColor(FColor(243, 156, 18, 255)); // Orange
 
 void SDataflowEdNode::UpdateErrorInfo()
 {
+	// Priority should be:
+	// 1. Error
+	// 2. Warning
+	// 3. Failed
+	// 4. Deprecated
+	// 5. Experimental
+
 	if (DataflowGraphNode)
 	{
 		if (const TSharedPtr<FDataflowNode> DataflowNode = DataflowGraphNode->GetDataflowNode())
 		{
-			if (UE::Dataflow::FNodeFactory::IsNodeExperimental(DataflowNode->GetType()))
-			{
-				ErrorMsg = FString(TEXT("Experimental"));
-				ErrorColor = CExperimentalColor;
-			}
-			if (UE::Dataflow::FNodeFactory::IsNodeDeprecated(DataflowNode->GetType()))
-			{
-				ErrorMsg = FString(TEXT("Deprecated"));
-				ErrorColor = CDeprecatedColor;
-			}
-
 			if (DataflowInterface)
 			{
 				if (const TSharedPtr<UE::Dataflow::FContext> Context = DataflowInterface->GetDataflowContext())
 				{
-					if (Context->NodeHasWarning(DataflowNode.Get()))
-					{
-						ErrorMsg = FString(TEXT("WARNING"));
-						ErrorColor = FAppStyle::GetColor("ErrorReporting.WarningBackgroundColor");
-					}
-					else if (Context->NodeHasError(DataflowNode.Get()))
+					if (Context->NodeHasError(DataflowNode.Get()))
 					{
 						ErrorMsg = FString(TEXT("ERROR"));
 						const FLinearColor ErrorBackgroundColor = FLinearColor(1.f, 0.02f, 0.003f);
 						ErrorColor = ErrorBackgroundColor;
+						return;
+					}
+					else if (Context->NodeHasWarning(DataflowNode.Get()))
+					{
+						ErrorMsg = FString(TEXT("WARNING"));
+						ErrorColor = FAppStyle::GetColor("ErrorReporting.WarningBackgroundColor");
+						return;
 					}
 					else if (Context->NodeFailed(DataflowNode.Get()))
 					{
 						ErrorMsg = FString(TEXT("FAILED"));
 						ErrorColor = CFailedColor;
+						return;
 					}
 				}
 			}
+
+			if (UE::Dataflow::FNodeFactory::IsNodeDeprecated(DataflowNode->GetType()))
+			{
+				ErrorMsg = FString(TEXT("Deprecated"));
+				ErrorColor = CDeprecatedColor;
+				return;
+			}
+			else if (UE::Dataflow::FNodeFactory::IsNodeExperimental(DataflowNode->GetType()))
+			{
+				ErrorMsg = FString(TEXT("Experimental"));
+				ErrorColor = CExperimentalColor;
+				return;
+			}
 		}
 	}
 }
