Chaosflesh: fixed geometry cache frame rate (defaulted 24, now set to simulation fps)

diff --git a/Engine/Plugins/Animation/MLDeformer/ChaosFleshGenerator/Source/ChaosFleshGenerator/Private/ChaosFleshGenerator.cpp b/Engine/Plugins/Animation/MLDeformer/ChaosFleshGenerator/Source/ChaosFleshGenerator/Private/ChaosFleshGenerator.cpp
--- a/Engine/Plugins/Animation/MLDeformer/ChaosFleshGenerator/Source/ChaosFleshGenerator/Private/ChaosFleshGenerator.cpp
+++ b/Engine/Plugins/Animation/MLDeformer/ChaosFleshGenerator/Source/ChaosFleshGenerator/Private/ChaosFleshGenerator.cpp
@@ -214,7 +214,7 @@ namespace UE::Chaos::FleshGenerator
 	
 		{
 			UE::Chaos::FleshGenerator::Private::FTimeScope TimeScope(TEXT("Saving"));
-			UE::DataflowSimulationGeometryCache::SaveGeometryCache(*TaskResource->Cache, *Properties->SkeletalMeshAsset, TaskResource->ImportedVertexNumbers, TaskResource->SimulatedPositions);
+			UE::DataflowSimulationGeometryCache::SaveGeometryCache(*TaskResource->Cache, Properties->SolverTiming.FrameRate, *Properties->SkeletalMeshAsset, TaskResource->ImportedVertexNumbers, TaskResource->SimulatedPositions);
 			UE::DataflowSimulationGeometryCache::SavePackage(*TaskResource->Cache);
 		}
 		if (bCancelled)
diff --git a/Engine/Plugins/Experimental/Dataflow/Source/DataflowEditor/Private/Dataflow/DataflowSimulationGeometryCache.cpp b/Engine/Plugins/Experimental/Dataflow/Source/DataflowEditor/Private/Dataflow/DataflowSimulationGeometryCache.cpp
--- a/Engine/Plugins/Experimental/Dataflow/Source/DataflowEditor/Private/Dataflow/DataflowSimulationGeometryCache.cpp
+++ b/Engine/Plugins/Experimental/Dataflow/Source/DataflowEditor/Private/Dataflow/DataflowSimulationGeometryCache.cpp
@@ -44,7 +44,7 @@ namespace UE::DataflowSimulationGeometryCache
 		return TArrayView<TArray<FVector3f>>(Positions.GetData(), NumValidFrames);
 	}
 
-	void SaveGeometryCache(UGeometryCache& GeometryCache, const USkinnedAsset& Asset, TConstArrayView<uint32> ImportedVertexNumbers, TArrayView<TArray<FVector3f>> PositionsToMoveFrom)
+	void SaveGeometryCache(UGeometryCache& GeometryCache, float FrameRate, const USkinnedAsset& Asset, TConstArrayView<uint32> ImportedVertexNumbers, TArrayView<TArray<FVector3f>> PositionsToMoveFrom)
 	{
 #if WITH_EDITOR 
 		const FSkeletalMeshRenderData* RenderData = Asset.GetResourceForRendering();
@@ -59,7 +59,7 @@ namespace UE::DataflowSimulationGeometryCache
 
 		using UE::GeometryCacheHelpers::FGeometryCacheConstantTopologyWriter;
 		using UE::GeometryCacheHelpers::AddTrackWriterFromSkinnedAsset;
-		FGeometryCacheConstantTopologyWriter::FConfig Config; Config.FPS = 24;
+		FGeometryCacheConstantTopologyWriter::FConfig Config; Config.FPS = FrameRate;
 		using FTrackWriter = FGeometryCacheConstantTopologyWriter::FTrackWriter;
 		FGeometryCacheConstantTopologyWriter Writer(GeometryCache, Config);
 		const int32 Index = AddTrackWriterFromSkinnedAsset(Writer, Asset);
diff --git a/Engine/Plugins/Experimental/Dataflow/Source/DataflowEditor/Private/Dataflow/DataflowSimulationScene.cpp b/Engine/Plugins/Experimental/Dataflow/Source/DataflowEditor/Private/Dataflow/DataflowSimulationScene.cpp
--- a/Engine/Plugins/Experimental/Dataflow/Source/DataflowEditor/Private/Dataflow/DataflowSimulationScene.cpp
+++ b/Engine/Plugins/Experimental/Dataflow/Source/DataflowEditor/Private/Dataflow/DataflowSimulationScene.cpp
@@ -457,7 +457,7 @@ void UDataflowSimulationSceneDescription::GenerateGeometryCache()
 			RenderPositions[Frame] = GeometryCachable->GetGeometryCachePositions(EmbeddedSkeletalMesh);
 			Time += SimulationScene->GetDeltaTime();
 		}
-		UE::DataflowSimulationGeometryCache::SaveGeometryCache(*GeometryCacheAsset, *EmbeddedSkeletalMesh, ImportedVertexNumbers, RenderPositions);
+		UE::DataflowSimulationGeometryCache::SaveGeometryCache(*GeometryCacheAsset, float(CacheParams.FrameRate), *EmbeddedSkeletalMesh, ImportedVertexNumbers, RenderPositions);
 		UE::DataflowSimulationGeometryCache::SavePackage(*GeometryCacheAsset);
 	}
 }
diff --git a/Engine/Plugins/Experimental/Dataflow/Source/DataflowEditor/Public/Dataflow/DataflowSimulationGeometryCache.h b/Engine/Plugins/Experimental/Dataflow/Source/DataflowEditor/Public/Dataflow/DataflowSimulationGeometryCache.h
--- a/Engine/Plugins/Experimental/Dataflow/Source/DataflowEditor/Public/Dataflow/DataflowSimulationGeometryCache.h
+++ b/Engine/Plugins/Experimental/Dataflow/Source/DataflowEditor/Public/Dataflow/DataflowSimulationGeometryCache.h
@@ -14,7 +14,7 @@ DECLARE_LOG_CATEGORY_EXTERN(LogDataflowSimulationGeometryCache, Log, All);
 namespace UE::DataflowSimulationGeometryCache
 {
 	/** Saves positions into the geometry cache */
-	DATAFLOWEDITOR_API void SaveGeometryCache(UGeometryCache& GeometryCache, const USkinnedAsset& Asset, TConstArrayView<uint32> ImportedVertexNumbers, TArrayView<TArray<FVector3f>> PositionsToMoveFrom);
+	DATAFLOWEDITOR_API void SaveGeometryCache(UGeometryCache& GeometryCache, float FrameRate, const USkinnedAsset& Asset, TConstArrayView<uint32> ImportedVertexNumbers, TArrayView<TArray<FVector3f>> PositionsToMoveFrom);
 
 	/** Saves the geometry cache */
 	DATAFLOWEDITOR_API void SavePackage(UObject& Object);
