Dataflow - Fixed crash in FFreezeActionsCustomization.
#rnx
preflight 67c07a0a292e55486fd23421

diff --git a/Engine/Plugins/Experimental/Dataflow/Source/DataflowEditor/Private/Dataflow/DataflowFreezeActionsCustomization.cpp b/Engine/Plugins/Experimental/Dataflow/Source/DataflowEditor/Private/Dataflow/DataflowFreezeActionsCustomization.cpp
--- a/Engine/Plugins/Experimental/Dataflow/Source/DataflowEditor/Private/Dataflow/DataflowFreezeActionsCustomization.cpp
+++ b/Engine/Plugins/Experimental/Dataflow/Source/DataflowEditor/Private/Dataflow/DataflowFreezeActionsCustomization.cpp
@@ -20,16 +20,18 @@ namespace UE::Dataflow
 	{
 		FDataflowNode* GetDataflowNode(TSharedRef<IPropertyHandle> StructPropertyHandle)
 		{
-			const TSharedPtr<IPropertyHandle> ParentHandle = StructPropertyHandle->GetParentHandle();
-			if (const TSharedPtr<IPropertyHandleStruct> ParentHandleStruct = ParentHandle->AsStruct())
+			if (const TSharedPtr<IPropertyHandle> ParentHandle = StructPropertyHandle->GetParentHandle())
 			{
-				if (const TSharedPtr<FStructOnScope> StructOnScope = ParentHandleStruct->GetStructData())
+				if (const TSharedPtr<IPropertyHandleStruct> ParentHandleStruct = ParentHandle->AsStruct())
 				{
-					if (const UStruct* const Struct = StructOnScope->GetStruct())
+					if (const TSharedPtr<FStructOnScope> StructOnScope = ParentHandleStruct->GetStructData())
 					{
-						if (Struct->IsChildOf<FDataflowNode>())
+						if (const UStruct* const Struct = StructOnScope->GetStruct())
 						{
-							return reinterpret_cast<FDataflowNode*>(StructOnScope->GetStructMemory());
+							if (Struct->IsChildOf<FDataflowNode>())
+							{
+								return reinterpret_cast<FDataflowNode*>(StructOnScope->GetStructMemory());
+							}
 						}
 					}
 				}
diff --git a/Engine/Source/Runtime/Experimental/Dataflow/Core/Private/Dataflow/DataflowNode.cpp b/Engine/Source/Runtime/Experimental/Dataflow/Core/Private/Dataflow/DataflowNode.cpp
--- a/Engine/Source/Runtime/Experimental/Dataflow/Core/Private/Dataflow/DataflowNode.cpp
+++ b/Engine/Source/Runtime/Experimental/Dataflow/Core/Private/Dataflow/DataflowNode.cpp
@@ -643,7 +643,7 @@ void FDataflowNode::Freeze(UE::Dataflow::FContext& Context)
 	bIsFrozen = true;
 	for (FDataflowOutput* const Output : GetOutputs())
 	{
-		if (Output)
+		if (Output && Output->HasConcreteType())
 		{
 			Output->Freeze(Context, FrozenProperties);
 		}
