Dataflow Editor: when caching recenly used View Mode, handle the case where a node does not specify any view mode

diff --git a/Engine/Plugins/Experimental/Dataflow/Source/DataflowEditor/Private/Dataflow/DataflowEditorMode.cpp b/Engine/Plugins/Experimental/Dataflow/Source/DataflowEditor/Private/Dataflow/DataflowEditorMode.cpp
--- a/Engine/Plugins/Experimental/Dataflow/Source/DataflowEditor/Private/Dataflow/DataflowEditorMode.cpp
+++ b/Engine/Plugins/Experimental/Dataflow/Source/DataflowEditor/Private/Dataflow/DataflowEditorMode.cpp
@@ -842,14 +842,15 @@ void UDataflowEditorMode::SetConstructionViewMode(const FName& NewViewModeName)
 					TArray<FName> ViewModesForNode;
 					UE::Dataflow::GetViewModesForNode(*SelectedDataflowEdNode, *EditorContent, ViewModesForNode);
 
-					check(ViewModesForNode.Contains(NewViewModeName));
-
-					ViewModesForNode.Sort([](const FName& A, const FName& B)
+					if (ViewModesForNode.Contains(NewViewModeName))
 					{
-						return A.FastLess(B);
-					});
+						ViewModesForNode.Sort([](const FName& A, const FName& B)
+						{
+							return A.FastLess(B);
+						});
 
-					NodeViewModeHistory.FindOrAdd(ViewModesForNode) = NewViewModeName;
+						NodeViewModeHistory.FindOrAdd(ViewModesForNode) = NewViewModeName;
+					}
 				}
 			}
 		}
