Add support for rendering UV in collection
- Add output rendering to UV based nodes and bake nodes

diff --git a/Engine/Plugins/Experimental/Dataflow/Source/DataflowEditor/Private/Dataflow/DataflowEngineRendering.cpp b/Engine/Plugins/Experimental/Dataflow/Source/DataflowEditor/Private/Dataflow/DataflowEngineRendering.cpp
--- a/Engine/Plugins/Experimental/Dataflow/Source/DataflowEditor/Private/Dataflow/DataflowEngineRendering.cpp
+++ b/Engine/Plugins/Experimental/Dataflow/Source/DataflowEditor/Private/Dataflow/DataflowEngineRendering.cpp
@@ -15,6 +15,8 @@
 #include "GeometryCollection/ManagedArrayCollection.h"
 #include "GeometryCollection/GeometryCollectionAlgo.h"
 #include "Materials/MaterialInterface.h"
+#include "GeometryCollection/Facades/CollectionUVFacade.h"
+#include "GeometryCollection/Facades/CollectionMeshFacade.h"
 #include "UDynamicMesh.h"
 
 namespace UE::Dataflow
@@ -447,7 +449,68 @@ namespace UE::Dataflow
 		}
 	}
 
-	
+	void RenderGeometryCollectionUV(GeometryCollection::Facades::FRenderingFacade& RenderCollection, const UE::Dataflow::FGraphRenderingState& State, TArray<FLinearColor>* VertexColorOverride = nullptr)
+	{
+		const TArray<FName>& RenderOutputs = State.GetRenderOutputs();
+		if (RenderOutputs.Num() == 0)
+		{
+			return; // no outputs
+		}
+		const FName PrimaryOutput = RenderOutputs[0]; // "Collection"
+		const FName UVChannelOutput = (RenderOutputs.Num() > 1)? RenderOutputs[1]: NAME_None; // "UVChannel"
+
+		const FManagedArrayCollection Default;
+		const FManagedArrayCollection& Collection = State.GetValue<FManagedArrayCollection>(PrimaryOutput, Default);
+
+		const int32 DefaultUVChannel = 0;
+		int32 UVChannel = (UVChannelOutput.IsNone())? DefaultUVChannel: State.GetValue<int32>(UVChannelOutput, DefaultUVChannel);
+
+		const GeometryCollection::Facades::FCollectionUVFacade UVFacade(Collection);
+		if (!UVFacade.IsValid())
+		{
+			return; // no UV data 
+		}
+		if (UVFacade.FindUVLayer(UVChannel) == nullptr)
+		{
+			UVChannel = 0;
+		}
+
+		const GeometryCollection::Facades::FCollectionMeshFacade MeshFacade(Collection);
+		if (!MeshFacade.IndicesAttribute.IsValid())
+		{
+			return; // no face nothing to render
+		}
+
+		TArray<FIntVector> Triangles;
+		Triangles.Reserve(MeshFacade.IndicesAttribute.Num());
+		for (int32 TriangleIdx = 0; TriangleIdx < MeshFacade.IndicesAttribute.Num(); ++TriangleIdx)
+		{
+			const bool IsVisible = (MeshFacade.VisibleAttribute.IsValid()) ? MeshFacade.VisibleAttribute.Get()[TriangleIdx] : true;
+			if (IsVisible)
+			{
+				Triangles.Add(MeshFacade.IndicesAttribute[TriangleIdx]);
+			}
+		}
+
+		const TManagedArray<FVector2f>& UVs = UVFacade.GetUVLayer(UVChannel);
+		TArray<FVector3f> UVVertices;
+		UVVertices.Reserve(UVs.Num());
+		for (int32 UvIdx = 0; UvIdx < UVs.Num(); ++UvIdx)
+		{
+			const FVector2f& Uv = UVs[UvIdx];
+			UVVertices.Add(FVector3f(Uv.X, Uv.Y, 0.f));
+		}
+
+		TArray<FVector3f> VertexNormals;
+		VertexNormals.Init(FVector3f(0.0, 0.0, 1.0), UVVertices.Num());
+
+		TArray<FLinearColor> VertexColors;
+		VertexColors.Init(FLinearColor(0, 0, 0, 0), UVVertices.Num());
+
+		const int32 GeometryIndex = RenderCollection.StartGeometryGroup(FGeometryCollection::GeometryGroup.ToString());
+		RenderCollection.AddSurface(MoveTemp(UVVertices), MoveTemp(Triangles), MoveTemp(VertexNormals), MoveTemp(VertexColors));
+		RenderCollection.EndGeometryGroup(GeometryIndex);
+	}
 
 	class FGeometryCollectionSurfaceRenderCallbacks : public FRenderingFactory::ICallbackInterface
 	{
@@ -458,7 +521,7 @@ namespace UE::Dataflow
 
 		virtual bool CanRender(const UE::Dataflow::IDataflowConstructionViewMode& ViewMode) const override
 		{
-			return (ViewMode.GetName() == FDataflowConstruction3DViewMode::Name);
+			return (ViewMode.GetName() == FDataflowConstruction3DViewMode::Name || ViewMode.GetName() == FDataflowConstructionUVViewMode::Name);
 		}
 
 		virtual void Render(GeometryCollection::Facades::FRenderingFacade& RenderCollection, const FGraphRenderingState& State) override
@@ -484,13 +547,20 @@ namespace UE::Dataflow
 				bool bFoundGeometryAttributes = bFoundVertexStart && bFoundVertexCount && bFoundFaceStart && bFoundFaceCount
 					&& Collection.NumElements(FGeometryCollection::GeometryGroup) > 0;
 
-				if (bFoundRenderData && bFoundGeometryAttributes)
+				if (State.GetViewMode().GetName() == Dataflow::FDataflowConstruction3DViewMode::Name)
 				{
-					RenderMeshIndexedGeometryCollection(RenderCollection, State);
+					if (bFoundRenderData && bFoundGeometryAttributes)
+					{
+						RenderMeshIndexedGeometryCollection(RenderCollection, State);
+					}
+					else if (bFoundRenderData)
+					{
+						RenderBasicGeometryCollection(RenderCollection, State);
+					}
 				}
-				else if (bFoundRenderData)
+				else if (State.GetViewMode().GetName() == Dataflow::FDataflowConstructionUVViewMode::Name)
 				{
-					RenderBasicGeometryCollection(RenderCollection, State);
+					RenderGeometryCollectionUV(RenderCollection, State);
 				}
 			}
 		}
diff --git a/Engine/Plugins/Experimental/GeometryCollectionPlugin/Source/GeometryCollectionNodes/Private/Dataflow/GeometryCollectionTextureNodes.cpp b/Engine/Plugins/Experimental/GeometryCollectionPlugin/Source/GeometryCollectionNodes/Private/Dataflow/GeometryCollectionTextureNodes.cpp
--- a/Engine/Plugins/Experimental/GeometryCollectionPlugin/Source/GeometryCollectionNodes/Private/Dataflow/GeometryCollectionTextureNodes.cpp
+++ b/Engine/Plugins/Experimental/GeometryCollectionPlugin/Source/GeometryCollectionNodes/Private/Dataflow/GeometryCollectionTextureNodes.cpp
@@ -35,6 +35,7 @@ FBakeTextureFromCollectionDataflowNode::FBakeTextureFromCollectionDataflowNode(c
 	
 	RegisterOutputConnection(&Collection, &Collection);
 	RegisterOutputConnection(&Image);
+	RegisterOutputConnection(&UVChannel, &UVChannel);
 }
 
 void FBakeTextureFromCollectionDataflowNode::Evaluate(UE::Dataflow::FContext& Context, const FDataflowOutput* Out) const
@@ -43,6 +44,10 @@ void FBakeTextureFromCollectionDataflowNode::Evaluate(UE::Dataflow::FContext& Co
 	{
 		SafeForwardInput(Context, &Collection, &Collection);
 	}
+	else if (Out->IsA(&UVChannel))
+	{
+		SafeForwardInput(Context, &UVChannel, &UVChannel);
+	}
 	else if (Out->IsA(&Image))
 	{
 		FManagedArrayCollection InCollection = GetValue(Context, &Collection);
diff --git a/Engine/Plugins/Experimental/GeometryCollectionPlugin/Source/GeometryCollectionNodes/Private/Dataflow/GeometryCollectionUVNodes.cpp b/Engine/Plugins/Experimental/GeometryCollectionPlugin/Source/GeometryCollectionNodes/Private/Dataflow/GeometryCollectionUVNodes.cpp
--- a/Engine/Plugins/Experimental/GeometryCollectionPlugin/Source/GeometryCollectionNodes/Private/Dataflow/GeometryCollectionUVNodes.cpp
+++ b/Engine/Plugins/Experimental/GeometryCollectionPlugin/Source/GeometryCollectionNodes/Private/Dataflow/GeometryCollectionUVNodes.cpp
@@ -22,6 +22,7 @@ FAddUVChannelDataflowNode::FAddUVChannelDataflowNode(const UE::Dataflow::FNodePa
 	: FDataflowNode(InParam, InGuid)
 {
 	RegisterInputConnection(&Collection);
+
 	RegisterOutputConnection(&Collection, &Collection);
 	RegisterOutputConnection(&UVChannel);
 }
@@ -73,12 +74,18 @@ FMergeUVIslandsDataflowNode::FMergeUVIslandsDataflowNode(const UE::Dataflow::FNo
 	RegisterInputConnection(&MaxNormalDeviationDeg).SetCanHidePin(true).SetPinIsHidden(true);
 	RegisterInputConnection(&NormalSmoothingRounds).SetCanHidePin(true).SetPinIsHidden(true);
 	RegisterInputConnection(&NormalSmoothingAlpha).SetCanHidePin(true).SetPinIsHidden(true);
+
 	RegisterOutputConnection(&Collection, &Collection);
+	RegisterOutputConnection(&UVChannel, &UVChannel);
 }
 
 void FMergeUVIslandsDataflowNode::Evaluate(UE::Dataflow::FContext& Context, const FDataflowOutput* Out) const
 {
-	if (Out->IsA(&Collection))
+	if (Out->IsA(&UVChannel))
+	{
+		SafeForwardInput(Context, &UVChannel, &UVChannel);
+	}
+	else if (Out->IsA(&Collection))
 	{
 		const FManagedArrayCollection& InCollection = GetValue(Context, &Collection);
 		const int32 FaceCount = InCollection.NumElements(FGeometryCollection::FacesGroup);
@@ -136,12 +143,18 @@ FAutoUnwrapUVDataflowNode::FAutoUnwrapUVDataflowNode(const UE::Dataflow::FNodePa
 	RegisterInputConnection(&FaceSelection);
 	RegisterInputConnection(&UVChannel);
 	RegisterInputConnection(&GutterSize).SetCanHidePin(true).SetPinIsHidden(true);
+
 	RegisterOutputConnection(&Collection, &Collection);
+	RegisterOutputConnection(&UVChannel, &UVChannel);
 }
 
 void FAutoUnwrapUVDataflowNode::Evaluate(UE::Dataflow::FContext& Context, const FDataflowOutput* Out) const
 {
-	if (Out->IsA(&Collection))
+	if (Out->IsA(&UVChannel))
+	{
+		SafeForwardInput(Context, &UVChannel, &UVChannel);
+	}
+	else if (Out->IsA(&Collection))
 	{
 		const FManagedArrayCollection& InCollection = GetValue(Context, &Collection);
 		const int32 FaceCount = InCollection.NumElements(FGeometryCollection::FacesGroup);
@@ -203,12 +216,18 @@ FBoxProjectUVDataflowNode::FBoxProjectUVDataflowNode(const UE::Dataflow::FNodePa
 	RegisterInputConnection(&bCenterBoxAtPivot);
 	RegisterInputConnection(&bUniformProjectionScale);
 	RegisterInputConnection(&GutterSize).SetCanHidePin(true).SetPinIsHidden(true);
+
 	RegisterOutputConnection(&Collection, &Collection);
+	RegisterOutputConnection(&UVChannel, &UVChannel);
 }
 
 void FBoxProjectUVDataflowNode::Evaluate(UE::Dataflow::FContext& Context, const FDataflowOutput* Out) const
 {
-	if (Out->IsA(&Collection))
+	if (Out->IsA(&UVChannel))
+	{
+		SafeForwardInput(Context, &UVChannel, &UVChannel);
+	}
+	else if (Out->IsA(&Collection))
 	{
 		const FManagedArrayCollection& InCollection = GetValue(Context, &Collection);
 		const int32 FaceCount = InCollection.NumElements(FGeometryCollection::FacesGroup);
diff --git a/Engine/Plugins/Experimental/GeometryCollectionPlugin/Source/GeometryCollectionNodes/Public/Dataflow/GeometryCollectionTextureNodes.h b/Engine/Plugins/Experimental/GeometryCollectionPlugin/Source/GeometryCollectionNodes/Public/Dataflow/GeometryCollectionTextureNodes.h
--- a/Engine/Plugins/Experimental/GeometryCollectionPlugin/Source/GeometryCollectionNodes/Public/Dataflow/GeometryCollectionTextureNodes.h
+++ b/Engine/Plugins/Experimental/GeometryCollectionPlugin/Source/GeometryCollectionNodes/Public/Dataflow/GeometryCollectionTextureNodes.h
@@ -55,6 +55,7 @@ struct FBakeTextureFromCollectionDataflowNode : public FDataflowNode
 {
 	GENERATED_USTRUCT_BODY()
 	DATAFLOW_NODE_DEFINE_INTERNAL(FBakeTextureFromCollectionDataflowNode, "BakeTextureFromCollection", "GeometryCollection|Texture", "")
+	DATAFLOW_NODE_RENDER_TYPE("SurfaceRender", FGeometryCollection::StaticType(), "Collection", "UVChannel")
 
 private:
 	/** Target Collection */
@@ -78,7 +79,7 @@ private:
 	int32 GutterSize = 1;
 
 	/** Index of the added UV channel */
-	UPROPERTY(EditAnywhere, Category = ChannelOptions, meta = (DataflowInput, ClampMin = 0, ClampMax = 7, UIMin = 0, UIMax = 7))
+	UPROPERTY(EditAnywhere, Category = ChannelOptions, meta = (DataflowInput, DataflowOutput, DataflowPassthrough = "UVChannel", ClampMin = 0, ClampMax = 7, UIMin = 0, UIMax = 7))
 	int32 UVChannel = 0;
 
 	/** Attribute to bake in the red channel */
diff --git a/Engine/Plugins/Experimental/GeometryCollectionPlugin/Source/GeometryCollectionNodes/Public/Dataflow/GeometryCollectionUVNodes.h b/Engine/Plugins/Experimental/GeometryCollectionPlugin/Source/GeometryCollectionNodes/Public/Dataflow/GeometryCollectionUVNodes.h
--- a/Engine/Plugins/Experimental/GeometryCollectionPlugin/Source/GeometryCollectionNodes/Public/Dataflow/GeometryCollectionUVNodes.h
+++ b/Engine/Plugins/Experimental/GeometryCollectionPlugin/Source/GeometryCollectionNodes/Public/Dataflow/GeometryCollectionUVNodes.h
@@ -17,6 +17,7 @@ struct FAddUVChannelDataflowNode : public FDataflowNode
 {
 	GENERATED_USTRUCT_BODY()
 	DATAFLOW_NODE_DEFINE_INTERNAL(FAddUVChannelDataflowNode, "AddUVChannel", "GeometryCollection|UV", "")
+	DATAFLOW_NODE_RENDER_TYPE("SurfaceRender", FGeometryCollection::StaticType(), "Collection", "UVChannel")
 
 private:
 	/** Target Collection */
@@ -45,6 +46,7 @@ struct FAutoUnwrapUVDataflowNode : public FDataflowNode
 {
 	GENERATED_USTRUCT_BODY()
 	DATAFLOW_NODE_DEFINE_INTERNAL(FAutoUnwrapUVDataflowNode, "AutoUnwrapUV", "GeometryCollection|UV", "")
+	DATAFLOW_NODE_RENDER_TYPE("SurfaceRender", FGeometryCollection::StaticType(), "Collection", "UVChannel")
 
 private:
 	/** Target Collection */
@@ -56,7 +58,7 @@ private:
 	FDataflowFaceSelection FaceSelection;
 
 	/** UV channel to unwrap into ( 0 by default ) */
-	UPROPERTY(EditAnywhere, Category = Options, meta = (DataflowInput, ClampMin = 0, ClampMax = 7, UIMin = 0, UIMax = 7))
+	UPROPERTY(EditAnywhere, Category = Options, meta = (DataflowInput, DataflowOutput, DataflowPassthrough = "UVChannel", ClampMin = 0, ClampMax = 7, UIMin = 0, UIMax = 7))
 	int32 UVChannel = 0;
 
 	/** Approximate space to leave between UV islands, measured in texels for 512x512 texture */
@@ -77,6 +79,7 @@ struct FMergeUVIslandsDataflowNode : public FDataflowNode
 {
 	GENERATED_USTRUCT_BODY()
 	DATAFLOW_NODE_DEFINE_INTERNAL(FMergeUVIslandsDataflowNode, "MergeUVIslands", "GeometryCollection|UV", "")
+	DATAFLOW_NODE_RENDER_TYPE("SurfaceRender", FGeometryCollection::StaticType(), "Collection", "UVChannel")
 
 private:
 	/** Target Collection */
@@ -88,10 +91,9 @@ private:
 	FDataflowFaceSelection FaceSelection;
 
 	/** UV channel to unwrap into ( 0 by default ) */
-	UPROPERTY(EditAnywhere, Category = Options, meta = (DataflowInput, ClampMin = 0, ClampMax = 7, UIMin = 0, UIMax = 7))
+	UPROPERTY(EditAnywhere, Category = Options, meta = (DataflowInput, DataflowOutput, DataflowPassthrough = "UVChannel", ClampMin = 0, ClampMax = 7, UIMin = 0, UIMax = 7))
 	int32 UVChannel = 0;
 
-
 	// Threshold for allowed area distortion from merging islands (when we use ExpMap to compute new UVs for the merged island)
 	UPROPERTY(EditAnywhere, Category = Options, meta = (DataflowInput, ClampMin = 0, ClampMax = 10, UIMin = 0, UIMax = 10))
 	double AreaDistortionThreshold = 1.5;
@@ -120,6 +122,7 @@ struct FBoxProjectUVDataflowNode : public FDataflowNode
 {
 	GENERATED_USTRUCT_BODY()
 	DATAFLOW_NODE_DEFINE_INTERNAL(FBoxProjectUVDataflowNode, "BoxProjectUV", "GeometryCollection|UV", "")
+	DATAFLOW_NODE_RENDER_TYPE("SurfaceRender", FGeometryCollection::StaticType(), "Collection", "UVChannel")
 
 private:
 	/** Target Collection */
@@ -131,7 +134,7 @@ private:
 	//FDataflowFaceSelection FaceSelection;
 
 	/** UV channel to unwrap into ( 0 by default ) */
-	UPROPERTY(EditAnywhere, Category = Options, meta = (DataflowInput, ClampMin = 0, ClampMax = 7, UIMin = 0, UIMax = 7))
+	UPROPERTY(EditAnywhere, Category = Options, meta = (DataflowInput, DataflowOutput, DataflowPassthrough = "UVChannel", ClampMin = 0, ClampMax = 7, UIMin = 0, UIMax = 7))
 	int32 UVChannel = 0;
 
 	/** Approximate space to leave between UV islands, measured in texels for 512x512 texture */
