Dataflow Editor : update evaluation logic
- Graph does not evaluate upon loading an asset anymore ( only when something is changed )
- Invalid terminal node are triggering the actual evaluation of the graph ( all of them as long as they are active )
- Do not use the timestamp from the context to evaluate the terminal node ( because this conflict with evaluating multiple terminal )
- future work will involve removing terminal content related code probably and try to remove as much dependency on the context timestamp

diff --git a/Engine/Plugins/Experimental/Dataflow/Source/DataflowEditor/Private/Dataflow/DataflowEditorToolkit.cpp b/Engine/Plugins/Experimental/Dataflow/Source/DataflowEditor/Private/Dataflow/DataflowEditorToolkit.cpp
--- a/Engine/Plugins/Experimental/Dataflow/Source/DataflowEditor/Private/Dataflow/DataflowEditorToolkit.cpp
+++ b/Engine/Plugins/Experimental/Dataflow/Source/DataflowEditor/Private/Dataflow/DataflowEditorToolkit.cpp
@@ -398,6 +398,8 @@ void FDataflowEditorToolkit::RegisterContextHandlers()
 		OnContextHasWarningDelegateHandle = DataflowContext->OnContextHasWarning.AddSP(this, &FDataflowEditorToolkit::OnContextHasWarning);
 		OnContextHasErrorDelegateHandle = DataflowContext->OnContextHasError.AddSP(this, &FDataflowEditorToolkit::OnContextHasError);
 	}
+
+	FDataflowAssetDelegates::OnNodeInvalidated.AddSP(this, &FDataflowEditorToolkit::OnNodeInvalidated);
 }
 
 void FDataflowEditorToolkit::UnregisterContextHandlers()
@@ -412,8 +414,10 @@ void FDataflowEditorToolkit::UnregisterContextHandlers()
 
 		DataflowContext->OnContextHasInfo.Remove(OnContextHasInfoDelegateHandle);
 		DataflowContext->OnContextHasWarning.Remove(OnContextHasWarningDelegateHandle);
-		DataflowContext->OnContextHasError.Remove(OnContextHasErrorDelegateHandle);
+		DataflowContext->OnContextHasError.Remove(OnContextHasErrorDelegateHandle);	
 	}
+
+	FDataflowAssetDelegates::OnNodeInvalidated.RemoveAll(this);
 }
 
 void FDataflowEditorToolkit::CreateEditorModeManager()
@@ -1287,6 +1291,22 @@ void FDataflowEditorToolkit::OnNodeSelectionChanged(const TSet<UObject*>& InNewS
 	}
 }
 
+void FDataflowEditorToolkit::OnNodeInvalidated(UDataflow& DataflowAsset, FDataflowNode& Node)
+{
+	if (const TObjectPtr<UDataflowBaseContent>& EditorContent = GetEditorContent())
+	{
+		if (EditorContent->GetDataflowAsset() == &DataflowAsset)
+		{
+			// evaluate only active terminals nodes 
+			if (FDataflowTerminalNode* TerminalNode = Node.AsType<FDataflowTerminalNode>())
+			{
+				EvaluateTerminalNode(*TerminalNode);
+				bViewsNeedRefresh = true;
+			}
+		}
+	}
+}
+
 void FDataflowEditorToolkit::OnNodeSingleClicked(UObject* ClickedNode) const
 {
 	UDataflowEditorMode* DataflowMode = Cast<UDataflowEditorMode>(EditorModeManager->GetActiveScriptableMode(UDataflowEditorMode::EM_DataflowEditorModeId));
@@ -1527,60 +1547,39 @@ void FDataflowEditorToolkit::Tick(float DeltaTime)
 
 			// Update the list of dataflow terminal contents 
 			DataflowEditor->UpdateTerminalContents(InitTimeStamp);
+		}
+	}
 
-			// OnTick evaluation only pulls the terminal nodes. The other evaluations can be specific nodes.
-			// We only evaluate multiple terminal nodes if the dataflow owner is a UDataflow (Owner == Asset)
-			if (!GetTerminalContents().IsEmpty() && (EditorContent->GetDataflowOwner() == EditorContent->GetDataflowAsset()) && EditorContent->GetDataflowContext())
-			{
-				for (const TObjectPtr<UDataflowBaseContent>& TerminalContent : GetTerminalContents())
-				{
-					if (const UDataflow* const Dataflow = EditorContent->GetDataflowAsset())
-					{
-						if (const TSharedPtr<const UE::Dataflow::FGraph> Graph = Dataflow->GetDataflow())
-						{
-							const FName TerminalNodeName(TerminalContent->GetDataflowTerminal());
-							const FDataflowNode* const Node = Graph->FindBaseNode(TerminalNodeName).Get();
+	RefreshViewsIfNeeded();
+}
 
-							UE::Dataflow::FTimestamp TerminalNodeTimeStamp = InitTimeStamp;
-							EvaluateNode(Node, nullptr, TerminalNodeTimeStamp);  // When Node is null, EvaluateNode falls back on the EditorContent terminal node
+TStatId FDataflowEditorToolkit::GetStatId() const
+{
+	RETURN_QUICK_DECLARE_CYCLE_STAT(FDataflowEditorToolkit, STATGROUP_Tickables);
+}
 
-							// Take the Max of the existing time stamp, as other terminal nodes might have more recent invalidations
-							const UE::Dataflow::FTimestamp LastModifiedTimestamp = FMath::Max(EditorContent->GetLastModifiedTimestamp(), TerminalNodeTimeStamp);
-							
-							constexpr bool bDontMakeDirty = false;
-							EditorContent->SetLastModifiedTimestamp(LastModifiedTimestamp, bDontMakeDirty);
-						}
-					}
-				}
+void FDataflowEditorToolkit::EvaluateTerminalNode(const FDataflowTerminalNode& TerminalNode)
+{
+	using namespace UE::Dataflow;
 
-				const bool bMakeDirty = (EditorContent->GetLastModifiedTimestamp() != InitTimeStamp);
-				EditorContent->SetLastModifiedTimestamp(EditorContent->GetLastModifiedTimestamp(), bMakeDirty);
-			}
-			else
-			{
-				UE::Dataflow::FTimestamp TerminalNodeTimeStamp = InitTimeStamp;
-				EvaluateNode(nullptr, nullptr, TerminalNodeTimeStamp);
+	// do not evaluate disabled terminal nodes
+	if (!TerminalNode.IsActive())
+	{
+		return;
+	}
 
-				const bool bMakeDirty = (TerminalNodeTimeStamp != InitTimeStamp);
-				EditorContent->SetLastModifiedTimestamp(TerminalNodeTimeStamp, bMakeDirty);
-			}
+	UE_LOG(LogChaosDataflow, VeryVerbose, TEXT("FDataflowEditorToolkit::EvaluateTerminalNode(): Node [%s]"), *TerminalNode.GetName().ToString());
 
-			// Ensure the context object's selected node matches the selected node in the graph editor
-			// TODO: Create an Editor Context Object that can just hold a reference to the graph editor, rather than keeping these in sync
-			if (GraphEditor->GetNumberOfSelectedNodes() == 1)
-			{
-				const UDataflowEdNode* const ContextObjectSelectedNode = EditorContent->GetSelectedNode();
-				const UDataflowEdNode* const EditorSelectedNode = Cast<UDataflowEdNode>(GraphEditor->GetSingleSelectedNode());
-				ensure(EditorSelectedNode == ContextObjectSelectedNode);
-			}
+	if (const TObjectPtr<UDataflowBaseContent>& EditorContent = GetEditorContent())
+	{
+		FTimestamp TerminalNodeTimeStamp = FTimestamp::Invalid;
 
-		}
-	}
-}
+		OnBeginEvaluate();
+		EvaluateNode(&TerminalNode, nullptr, TerminalNodeTimeStamp);
+		OnFinishEvaluate();
 
-TStatId FDataflowEditorToolkit::GetStatId() const
-{
-	RETURN_QUICK_DECLARE_CYCLE_STAT(FDataflowEditorToolkit, STATGROUP_Tickables);
+		EditorContent->SetLastModifiedTimestamp(TerminalNodeTimeStamp, /*bMakeDirty*/true);
+	}
 }
 
 void FDataflowEditorToolkit::EvaluateNode(const FDataflowNode* Node, const FDataflowOutput* Output, UE::Dataflow::FTimestamp& InOutTimestamp)
@@ -1605,6 +1604,18 @@ void FDataflowEditorToolkit::EvaluateNode(const FDataflowNode* Node, const FData
 	}
 }
 
+void FDataflowEditorToolkit::RefreshViewsIfNeeded(bool bForce)
+{
+	if (bForce || bViewsNeedRefresh)
+	{
+		for (IDataflowViewListener* Listener : ViewListeners)
+		{
+			Listener->RefreshView();
+		}
+		bViewsNeedRefresh = false;
+	}
+}
+
 TSharedRef<SDataflowGraphEditor> FDataflowEditorToolkit::CreateGraphEditorWidget(UEdGraph* GraphToEdit, TSharedPtr<IStructureDetailsView> InNodeDetailsEditor)
 {
 	ensure(GraphToEdit);
@@ -1613,11 +1624,9 @@ TSharedRef<SDataflowGraphEditor> FDataflowEditorToolkit::CreateGraphEditorWidget
 	const FDataflowEditorCommands::FGraphEvaluationCallback Evaluate =
 		[this](const FDataflowNode* Node, const FDataflowOutput* Output)
 		{
-			//
-			// Graph evaluation begin
-			//
-			OnBeginEvaluate();
-
+			// evaluate may already have happen when the node was invalidated in the graph editor code 
+			// and cause on eof the terminal node to evaluate 
+			// but if the node is in a part of the graph that does not end by a terminal node , we need to explictly evaluate it 
 			if (const TObjectPtr<UDataflowBaseContent>& EditorContent = GetEditorContent())
 			{
 				UE::Dataflow::FTimestamp LastNodeTimestamp = EditorContent->GetLastModifiedTimestamp();
@@ -1626,11 +1635,6 @@ TSharedRef<SDataflowGraphEditor> FDataflowEditorToolkit::CreateGraphEditorWidget
 
 				EditorContent->SetLastModifiedTimestamp(LastNodeTimestamp);
 			}
-
-			//
-			// Graph evaluation done
-			//
-			OnFinishEvaluate();
 		};
 	
 	DataflowEditor->UpdateTerminalContents(FTimestamp::Invalid);
diff --git a/Engine/Plugins/Experimental/Dataflow/Source/DataflowEditor/Private/Dataflow/DataflowView.cpp b/Engine/Plugins/Experimental/Dataflow/Source/DataflowEditor/Private/Dataflow/DataflowView.cpp
--- a/Engine/Plugins/Experimental/Dataflow/Source/DataflowEditor/Private/Dataflow/DataflowView.cpp
+++ b/Engine/Plugins/Experimental/Dataflow/Source/DataflowEditor/Private/Dataflow/DataflowView.cpp
@@ -14,21 +14,6 @@ FDataflowNodeView::FDataflowNodeView(TObjectPtr<UDataflowBaseContent> InContent)
 {
 }
 
-
-FDataflowNodeView::~FDataflowNodeView()
-{
-	if (SelectedNode)
-	{
-		if (TSharedPtr<FDataflowNode> DataflowNode = SelectedNode->GetDataflowNode())
-		{
-			if (DataflowNode->GetOnNodeInvalidatedDelegate().IsBound() && OnNodeInvalidatedDelegateHandle.IsValid())
-			{
-				DataflowNode->GetOnNodeInvalidatedDelegate().Remove(OnNodeInvalidatedDelegateHandle);
-			}
-		}
-	}
-}
-
 TObjectPtr<UDataflowBaseContent> FDataflowNodeView::GetEditorContent()
 {
 	if (ensure(EditorContent))
@@ -76,20 +61,6 @@ void FDataflowNodeView::OnSelectedNodeChanged(UDataflowEdNode* InNode)
 {
 	if (!bIsPinnedDown)
 	{
-		//
-		// Remove from broadcast
-		//
-		if (SelectedNode)
-		{
-			if (TSharedPtr<FDataflowNode> DataflowNode = SelectedNode->GetDataflowNode())
-			{
-				if (DataflowNode->GetOnNodeInvalidatedDelegate().IsBound() && OnNodeInvalidatedDelegateHandle.IsValid())
-				{
-					DataflowNode->GetOnNodeInvalidatedDelegate().Remove(OnNodeInvalidatedDelegateHandle);
-				}
-			}
-		}
-
 		SelectedNode = nullptr;
 
 		if (InNode)  // nullptr is valid
@@ -98,34 +69,17 @@ void FDataflowNodeView::OnSelectedNodeChanged(UDataflowEdNode* InNode)
 			{
 				SelectedNode = InNode;
 			}
-
-			// 
-			// Bind OnNodeInvalidated() to new SelectedNode
-			// 
-			if (SelectedNode)
-			{
-				if (TSharedPtr<FDataflowNode> DataflowNode = SelectedNode->GetDataflowNode())
-				{
-					OnNodeInvalidatedDelegateHandle = DataflowNode->GetOnNodeInvalidatedDelegate().AddRaw(this, &FDataflowNodeView::OnNodeInvalidated);
-				}
-			}
 		}
 
 		UpdateViewData();
 	}
 }
 
-void FDataflowNodeView::OnNodeInvalidated(FDataflowNode* InvalidatedNode)
+void FDataflowNodeView::RefreshView()
 {
-	if (!bIsRefreshLocked)
+	if (!bIsRefreshLocked && SelectedNode)
 	{
-		if (TSharedPtr<FDataflowNode> DataflowNode = SelectedNode->GetDataflowNode())
-		{
-			if (InvalidatedNode == DataflowNode.Get())
-			{
-				UpdateViewData();
-			}
-		}
+		UpdateViewData();
 	}
 }
 
diff --git a/Engine/Plugins/Experimental/Dataflow/Source/DataflowEditor/Public/Dataflow/DataflowEditorToolkit.h b/Engine/Plugins/Experimental/Dataflow/Source/DataflowEditor/Public/Dataflow/DataflowEditorToolkit.h
--- a/Engine/Plugins/Experimental/Dataflow/Source/DataflowEditor/Public/Dataflow/DataflowEditorToolkit.h
+++ b/Engine/Plugins/Experimental/Dataflow/Source/DataflowEditor/Public/Dataflow/DataflowEditorToolkit.h
@@ -110,6 +110,8 @@ protected:
 	void OnNodeTitleCommitted(const FText& InNewText, ETextCommit::Type InCommitType, UEdGraphNode* GraphNode);
 	void OnNodeDoubleClicked(UEdGraphNode* ClickedNode);
 	void OnNodeSelectionChanged(const TSet<UObject*>& NewSelection);
+	void OnEvaluateSelectedNodes(const TSet<UObject*>& SelectedNodes);
+	void OnNodeInvalidated(UDataflow& DataflowAsset, FDataflowNode& Node);
 	void OnNodeDeleted(const TSet<UObject*>& NewSelection);
 	void OnNodeSingleClicked(UObject* ClickedNode) const;
 	void OnAssetPropertyValueChanged(const FPropertyChangedEvent& PropertyChangedEvent);
@@ -118,7 +120,9 @@ protected:
 	void OnTabClosed(TSharedRef<SDockTab> Tab);
 
 	// Node evaluation
+	void EvaluateTerminalNode(const FDataflowTerminalNode& TerminalNode);
 	void EvaluateNode(const FDataflowNode* Node, const FDataflowOutput* Output, UE::Dataflow::FTimestamp& InOutTimestamp);
+	void RefreshViewsIfNeeded(bool bForce = false);
 	void OnNodeBeginEvaluate(const FDataflowNode* Node, const FDataflowOutput* Output);
 	void OnNodeFinishEvaluate(const FDataflowNode* Node, const FDataflowOutput* Output);
 	void OnBeginEvaluate();
@@ -244,6 +248,7 @@ private:
 	// Graph delegates used to update the UI
 	FDelegateHandle OnSelectionChangedMulticastDelegateHandle;
     FDelegateHandle OnNodeDeletedMulticastDelegateHandle;
+	FDelegateHandle OnEvaluateSelectedNodesDelegateHandle;
     FDelegateHandle OnFinishedChangingPropertiesDelegateHandle;
 	FDelegateHandle OnFinishedChangingAssetPropertiesDelegateHandle;
 	FDelegateHandle OnConstructionSelectionChangedDelegateHandle;
@@ -305,6 +310,7 @@ private:
 
 	FDateTime GraphEvaluationBegin;
 	FDateTime GraphEvaluationFinished;
+	bool bViewsNeedRefresh = false;
 
 	/** Persistent component used for debug drawing */
 	TObjectPtr<class UDataflowDebugDrawComponent> DebugDrawComponent;
diff --git a/Engine/Plugins/Experimental/Dataflow/Source/DataflowEditor/Public/Dataflow/DataflowView.h b/Engine/Plugins/Experimental/Dataflow/Source/DataflowEditor/Public/Dataflow/DataflowView.h
--- a/Engine/Plugins/Experimental/Dataflow/Source/DataflowEditor/Public/Dataflow/DataflowView.h
+++ b/Engine/Plugins/Experimental/Dataflow/Source/DataflowEditor/Public/Dataflow/DataflowView.h
@@ -22,6 +22,7 @@ public:
 	virtual void OnConstructionViewSelectionChanged(const TArray<UPrimitiveComponent*>& InSelectedComponents) = 0;
 	virtual void OnSelectedNodeChanged(UDataflowEdNode* InNode) = 0;  // nullptr is valid
 	virtual void OnNodeInvalidated(FDataflowNode* InvalidatedNode) = 0;
+	virtual void RefreshView() = 0;
 };
 
 
@@ -34,7 +35,6 @@ class FDataflowNodeView : public FGCObject, public IDataflowViewListener
 {
 public:
 	FDataflowNodeView(TObjectPtr<UDataflowBaseContent> InContent = nullptr);
-	virtual ~FDataflowNodeView();
 
 	UDataflowEdNode* GetSelectedNode() const { return SelectedNode; }
 	bool SelectedNodeHaveSupportedOutputTypes(UDataflowEdNode* InNode);
@@ -65,8 +65,8 @@ public:
 	*/
 	virtual void OnConstructionViewSelectionChanged(const TArray<UPrimitiveComponent*>& InNode) override;
 	virtual void OnSelectedNodeChanged(UDataflowEdNode* InNode) override;  // nullptr is valid
-	virtual void OnNodeInvalidated(FDataflowNode* InvalidatedNode) override;
-
+	virtual void OnNodeInvalidated(FDataflowNode* InvalidatedNode) override {}
+	virtual void RefreshView() override;
 
 	/** FGCObject interface */
 	virtual void AddReferencedObjects(FReferenceCollector& Collector) override;
diff --git a/Engine/Source/Runtime/Experimental/Dataflow/Core/Private/Dataflow/DataflowNode.cpp b/Engine/Source/Runtime/Experimental/Dataflow/Core/Private/Dataflow/DataflowNode.cpp
--- a/Engine/Source/Runtime/Experimental/Dataflow/Core/Private/Dataflow/DataflowNode.cpp
+++ b/Engine/Source/Runtime/Experimental/Dataflow/Core/Private/Dataflow/DataflowNode.cpp
@@ -679,6 +679,13 @@ PRAGMA_DISABLE_DEPRECATION_WARNINGS  // Until LastModifiedTimestamp becomes priv
 	if (LastModifiedTimestamp < InModifiedTimestamp)
 	{
 		LastModifiedTimestamp = InModifiedTimestamp;
+
+		if (OnNodeInvalidatedDelegate.IsBound())
+		{
+			OnNodeInvalidatedDelegate.Broadcast(this);
+		}
+
+		// propagate to downstream
 PRAGMA_ENABLE_DEPRECATION_WARNINGS
 		for (TPair<UE::Dataflow::FConnectionKey, FDataflowOutput*>& Elem : ExpandedOutputs)
 		{
@@ -687,8 +694,6 @@ PRAGMA_ENABLE_DEPRECATION_WARNINGS
 		}
 
 		OnInvalidate();
-
-		OnNodeInvalidatedDelegate.Broadcast(this);
 	}
 }
 
diff --git a/Engine/Source/Runtime/Experimental/Dataflow/Engine/Private/Dataflow/DataflowEdNode.cpp b/Engine/Source/Runtime/Experimental/Dataflow/Engine/Private/Dataflow/DataflowEdNode.cpp
--- a/Engine/Source/Runtime/Experimental/Dataflow/Engine/Private/Dataflow/DataflowEdNode.cpp
+++ b/Engine/Source/Runtime/Experimental/Dataflow/Engine/Private/Dataflow/DataflowEdNode.cpp
@@ -981,7 +981,17 @@ void UDataflowEdNode::RegisterDelegateHandle()
 	{
 		OnNodeInvalidatedDelegateHandle = DataflowNode->GetOnNodeInvalidatedDelegate().AddWeakLambda(this, [this](FDataflowNode* InDataflowNode)
 			{
-				GetGraph()->NotifyNodeChanged(this);
+				if (UEdGraph* EdGraph = GetGraph())
+				{
+					GetGraph()->NotifyNodeChanged(this);
+					if (UDataflow* DataflowAsset = UDataflow::GetDataflowAssetFromEdGraph(EdGraph))
+					{
+						if (InDataflowNode)
+						{
+							FDataflowAssetDelegates::OnNodeInvalidated.Broadcast(*DataflowAsset, *InDataflowNode);
+						}
+					}
+				}
 			});
 	}
 #endif
