Dataflow : make dataflow views to read from outputs to avoid full evaluation of the graph

diff --git a/Engine/Plugins/Experimental/Dataflow/Source/DataflowEditor/Private/Dataflow/DataflowCollectionSpreadSheet.cpp b/Engine/Plugins/Experimental/Dataflow/Source/DataflowEditor/Private/Dataflow/DataflowCollectionSpreadSheet.cpp
--- a/Engine/Plugins/Experimental/Dataflow/Source/DataflowEditor/Private/Dataflow/DataflowCollectionSpreadSheet.cpp
+++ b/Engine/Plugins/Experimental/Dataflow/Source/DataflowEditor/Private/Dataflow/DataflowCollectionSpreadSheet.cpp
@@ -45,10 +45,10 @@ void FDataflowCollectionSpreadSheet::UpdateViewData()
 								FName Name = Output->GetName();
 								FName Type = Output->GetType();
 
-								if (Output->GetType() == "FManagedArrayCollection" && !Output->IsAnyType())
+								if (Output->GetType() == "FManagedArrayCollection")
 								{
-									const FManagedArrayCollection& Value = Output->GetValue<FManagedArrayCollection>(*Context, FManagedArrayCollection());
-
+									const FManagedArrayCollection DefaultCollection;
+									const FManagedArrayCollection& Value = Output->ReadValue(*Context, DefaultCollection);
 									CollectionSpreadSheet->GetCollectionTable()->GetCollectionInfoMap().Add(Name.ToString(), { Value });
 								}
 							}
diff --git a/Engine/Plugins/Experimental/Dataflow/Source/DataflowEditor/Private/Dataflow/DataflowSelectionView.cpp b/Engine/Plugins/Experimental/Dataflow/Source/DataflowEditor/Private/Dataflow/DataflowSelectionView.cpp
--- a/Engine/Plugins/Experimental/Dataflow/Source/DataflowEditor/Private/Dataflow/DataflowSelectionView.cpp
+++ b/Engine/Plugins/Experimental/Dataflow/Source/DataflowEditor/Private/Dataflow/DataflowSelectionView.cpp
@@ -55,19 +55,22 @@ void FDataflowSelectionView::UpdateViewData()
 
 								if (Output->GetType() == "FDataflowTransformSelection")
 								{
-									const FDataflowTransformSelection& Value = Output->GetValue<FDataflowTransformSelection>(*Context, FDataflowTransformSelection());
+									const FDataflowTransformSelection DefaultTransformSelection;
+									const FDataflowTransformSelection& Value = Output->ReadValue(*Context, DefaultTransformSelection);
 
 									SelectionView->GetSelectionTable()->GetSelectionInfoMap().Add(Name.ToString(), { Type.ToString(), TBitArray<>(Value.GetBitArray()) });
 								}
 								else if (Output->GetType() == "FDataflowVertexSelection")
 								{
-									const FDataflowVertexSelection& Value = Output->GetValue<FDataflowVertexSelection>(*Context, FDataflowVertexSelection());
+									const FDataflowVertexSelection DefaultVertexSelection;
+									const FDataflowVertexSelection& Value = Output->ReadValue(*Context, DefaultVertexSelection);
 
 									SelectionView->GetSelectionTable()->GetSelectionInfoMap().Add(Name.ToString(), { Type.ToString(), TBitArray<>(Value.GetBitArray()) });
 								}
 								else if (Output->GetType() == "FDataflowFaceSelection")
 								{
-									const FDataflowFaceSelection& Value = Output->GetValue<FDataflowFaceSelection>(*Context, FDataflowFaceSelection());
+									const FDataflowFaceSelection DefaultFaceSelection;
+									const FDataflowFaceSelection& Value = Output->ReadValue(*Context, DefaultFaceSelection);
 
 									SelectionView->GetSelectionTable()->GetSelectionInfoMap().Add(Name.ToString(), { Type.ToString(), TBitArray<>(Value.GetBitArray()) });
 								}
diff --git a/Engine/Plugins/Experimental/Dataflow/Source/DataflowEditor/Private/Dataflow/DataflowSkeletonView.cpp b/Engine/Plugins/Experimental/Dataflow/Source/DataflowEditor/Private/Dataflow/DataflowSkeletonView.cpp
--- a/Engine/Plugins/Experimental/Dataflow/Source/DataflowEditor/Private/Dataflow/DataflowSkeletonView.cpp
+++ b/Engine/Plugins/Experimental/Dataflow/Source/DataflowEditor/Private/Dataflow/DataflowSkeletonView.cpp
@@ -90,8 +90,8 @@ void FDataflowSkeletonView::UpdateViewData()
 					{
 						if (TSharedPtr<UE::Dataflow::FEngineContext> Context = SkeletalContent->GetDataflowContext())
 						{
-							FManagedArrayCollection DefaultCollection;
-							const FManagedArrayCollection& Result = Output->GetValue(*Context, DefaultCollection);
+							const FManagedArrayCollection DefaultCollection;
+							const FManagedArrayCollection& Result = Output->ReadValue(*Context, DefaultCollection);
 
 							FGeometryCollectionEngineConversion::ConvertCollectionToSkeleton(Result, Skeleton, CollectionIndexRemap);
 						}
