CIS - Fixed issue 799710, 797289, Sanitizer found error heap used after free.

#rnx

diff --git a/Engine/Plugins/Experimental/Dataflow/Source/DataflowEditor/Private/Dataflow/DataflowEditorModule.cpp b/Engine/Plugins/Experimental/Dataflow/Source/DataflowEditor/Private/Dataflow/DataflowEditorModule.cpp
--- a/Engine/Plugins/Experimental/Dataflow/Source/DataflowEditor/Private/Dataflow/DataflowEditorModule.cpp
+++ b/Engine/Plugins/Experimental/Dataflow/Source/DataflowEditor/Private/Dataflow/DataflowEditorModule.cpp
@@ -24,55 +24,62 @@
 
 #define LOCTEXT_NAMESPACE "DataflowEditor"
 
-const FColor FDataflowEditorModule::SurfaceColor = FLinearColor(0.6, 0.6, 0.6).ToRGBE();
-static const FName ScalarVertexPropertyGroupName = TEXT("ScalarVertexPropertyGroup");
-static const FName DataflowFunctionPropertyName = TEXT("DataflowFunctionProperty");
-
-class FDataflowEditorWeightMapPaintToolActionCommands : public TInteractiveToolCommands<FDataflowEditorWeightMapPaintToolActionCommands>
+namespace UE::Dataflow::Private
 {
-public:
-	FDataflowEditorWeightMapPaintToolActionCommands() : 
-		TInteractiveToolCommands<FDataflowEditorWeightMapPaintToolActionCommands>(
-			TEXT("DataflowEditorWeightMapPaintToolContext"),
-			LOCTEXT("DataflowEditorWeightMapPaintToolContext", "Dataflow Weight Map Paint Tool Context"),
-			NAME_None,
-			FAppStyle::GetAppStyleSetName())
-	{}
-
-	virtual void GetToolDefaultObjectList(TArray<UInteractiveTool*>& ToolCDOs) override
-	{
-		ToolCDOs.Add(GetMutableDefault<UDataflowEditorWeightMapPaintTool>());
-	}
-};
-
+	static const FName ScalarVertexPropertyGroupName = TEXT("ScalarVertexPropertyGroup");
+	static const FName DataflowFunctionPropertyName = TEXT("DataflowFunctionProperty");
+	static const FName DataflowVariableOverridesName = TEXT("DataflowVariableOverrides");
+	static const FName DataflowFreezeActionsName = TEXT("DataflowFreezeActions");
 
-class FDataflowToolActionCommandBindings : public UE::Dataflow::FDataflowToolRegistry::IDataflowToolActionCommands
-{
-public:
-	FDataflowToolActionCommandBindings()
+	class FDataflowEditorWeightMapPaintToolActionCommands : public TInteractiveToolCommands<FDataflowEditorWeightMapPaintToolActionCommands>
 	{
-		FDataflowEditorWeightMapPaintToolActionCommands::Register();
-	}
+	public:
+		FDataflowEditorWeightMapPaintToolActionCommands() : 
+			TInteractiveToolCommands<FDataflowEditorWeightMapPaintToolActionCommands>(
+				TEXT("DataflowEditorWeightMapPaintToolContext"),
+				LOCTEXT("DataflowEditorWeightMapPaintToolContext", "Dataflow Weight Map Paint Tool Context"),
+				NAME_None,
+				FAppStyle::GetAppStyleSetName())
+		{}
+
+		virtual void GetToolDefaultObjectList(TArray<UInteractiveTool*>& ToolCDOs) override
+		{
+			ToolCDOs.Add(GetMutableDefault<UDataflowEditorWeightMapPaintTool>());
+		}
+	};
 
-	virtual void UnbindActiveCommands(const TSharedPtr<FUICommandList>& UICommandList) const override
-	{
-		checkf(FDataflowEditorWeightMapPaintToolActionCommands::IsRegistered(), TEXT("Expected WeightMapPaintTool actions to have been registered"));
-		FDataflowEditorWeightMapPaintToolActionCommands::Get().UnbindActiveCommands(UICommandList);
-	}
 
-	virtual void BindCommandsForCurrentTool(const TSharedPtr<FUICommandList>& UICommandList, UInteractiveTool* Tool) const override
+	class FDataflowToolActionCommandBindings : public UE::Dataflow::FDataflowToolRegistry::IDataflowToolActionCommands
 	{
-		if (ExactCast<UDataflowEditorWeightMapPaintTool>(Tool))
+	public:
+		FDataflowToolActionCommandBindings()
+		{
+			FDataflowEditorWeightMapPaintToolActionCommands::Register();
+		}
+
+		virtual void UnbindActiveCommands(const TSharedPtr<FUICommandList>& UICommandList) const override
 		{
 			checkf(FDataflowEditorWeightMapPaintToolActionCommands::IsRegistered(), TEXT("Expected WeightMapPaintTool actions to have been registered"));
-			FDataflowEditorWeightMapPaintToolActionCommands::Get().BindCommandsForCurrentTool(UICommandList, Tool);
+			FDataflowEditorWeightMapPaintToolActionCommands::Get().UnbindActiveCommands(UICommandList);
 		}
-	}
-};
 
+		virtual void BindCommandsForCurrentTool(const TSharedPtr<FUICommandList>& UICommandList, UInteractiveTool* Tool) const override
+		{
+			if (ExactCast<UDataflowEditorWeightMapPaintTool>(Tool))
+			{
+				checkf(FDataflowEditorWeightMapPaintToolActionCommands::IsRegistered(), TEXT("Expected WeightMapPaintTool actions to have been registered"));
+				FDataflowEditorWeightMapPaintToolActionCommands::Get().BindCommandsForCurrentTool(UICommandList, Tool);
+			}
+		}
+	};
+}
+
+const FColor FDataflowEditorModule::SurfaceColor = FLinearColor(0.6, 0.6, 0.6).ToRGBE();
 
 void FDataflowEditorModule::StartupModule()
 {
+	using namespace UE::Dataflow::Private;
+
 	FDataflowEditorStyle::Get();
 	
 	// Register type customizations
@@ -80,8 +87,8 @@ void FDataflowEditorModule::StartupModule()
 	{
 		PropertyModule->RegisterCustomPropertyTypeLayout(ScalarVertexPropertyGroupName, FOnGetPropertyTypeCustomizationInstance::CreateStatic(&UE::Dataflow::FScalarVertexPropertyGroupCustomization::MakeInstance));
 		PropertyModule->RegisterCustomPropertyTypeLayout(DataflowFunctionPropertyName, FOnGetPropertyTypeCustomizationInstance::CreateStatic(&UE::Dataflow::FFunctionPropertyCustomization::MakeInstance));
-		PropertyModule->RegisterCustomPropertyTypeLayout("DataflowVariableOverrides", FOnGetPropertyTypeCustomizationInstance::CreateStatic(&FDataflowVariableOverridesDetails::MakeInstance));
-		PropertyModule->RegisterCustomPropertyTypeLayout(FDataflowFreezeActions::StaticStruct()->GetFName(), FOnGetPropertyTypeCustomizationInstance::CreateStatic(&UE::Dataflow::FFreezeActionsCustomization::MakeInstance));
+		PropertyModule->RegisterCustomPropertyTypeLayout(DataflowVariableOverridesName, FOnGetPropertyTypeCustomizationInstance::CreateStatic(&FDataflowVariableOverridesDetails::MakeInstance));
+		PropertyModule->RegisterCustomPropertyTypeLayout(DataflowFreezeActionsName, FOnGetPropertyTypeCustomizationInstance::CreateStatic(&UE::Dataflow::FFreezeActionsCustomization::MakeInstance));
 	}
 
 	UE::Dataflow::RenderingCallbacks();
@@ -99,7 +106,9 @@ void FDataflowEditorModule::StartupModule()
 }
 
 void FDataflowEditorModule::ShutdownModule()
-{	
+{
+	using namespace UE::Dataflow::Private;
+
 	FEditorModeRegistry::Get().UnregisterMode(UDataflowEditorMode::EM_DataflowEditorModeId);
 
 	// Deregister type customizations
@@ -107,7 +116,8 @@ void FDataflowEditorModule::ShutdownModule()
 	{
 		PropertyModule->UnregisterCustomPropertyTypeLayout(ScalarVertexPropertyGroupName);
 		PropertyModule->UnregisterCustomPropertyTypeLayout(DataflowFunctionPropertyName);
-		PropertyModule->UnregisterCustomPropertyTypeLayout(FDataflowFreezeActions::StaticStruct()->GetFName());
+		PropertyModule->UnregisterCustomPropertyTypeLayout(DataflowVariableOverridesName);
+		PropertyModule->UnregisterCustomPropertyTypeLayout(DataflowFreezeActionsName);
 	}
 
 	FDataflowEditorCommands::Unregister();
