Dataflow Editor: add a mesh vertex normal visualization to the construction viewport

diff --git a/Engine/Plugins/Experimental/Dataflow/Source/DataflowEditor/Private/Dataflow/DataflowEditorModule.cpp b/Engine/Plugins/Experimental/Dataflow/Source/DataflowEditor/Private/Dataflow/DataflowEditorModule.cpp
--- a/Engine/Plugins/Experimental/Dataflow/Source/DataflowEditor/Private/Dataflow/DataflowEditorModule.cpp
+++ b/Engine/Plugins/Experimental/Dataflow/Source/DataflowEditor/Private/Dataflow/DataflowEditorModule.cpp
@@ -13,6 +13,7 @@
 #include "Dataflow/DataflowFunctionPropertyCustomization.h"
 #include "Dataflow/DataflowSNodeFactories.h"
 #include "Dataflow/MeshStatsConstructionVisualization.h"
+#include "Dataflow/MeshConstructionVisualization.h"
 #include "Dataflow/ScalarVertexPropertyGroupCustomization.h"
 #include "Dataflow/DataflowToolRegistry.h"
 #include "DataflowEditorTools/DataflowEditorWeightMapPaintTool.h"
@@ -101,6 +102,7 @@ void FDataflowEditorModule::StartupModule()
 
 	UE::Dataflow::FDataflowConstructionVisualizationRegistry& ConstructionVisualizationRegistry = UE::Dataflow::FDataflowConstructionVisualizationRegistry::GetInstance();
 	ConstructionVisualizationRegistry.RegisterVisualization(MakeUnique<UE::Dataflow::FMeshStatsConstructionVisualization>());
+	ConstructionVisualizationRegistry.RegisterVisualization(MakeUnique<UE::Dataflow::FMeshConstructionVisualization>());
 
 	FDataflowEditorCommands::Register();
 }
@@ -127,6 +129,7 @@ void FDataflowEditorModule::ShutdownModule()
 
 	UE::Dataflow::FDataflowConstructionVisualizationRegistry& ConstructionVisualizationRegistry = UE::Dataflow::FDataflowConstructionVisualizationRegistry::GetInstance();
 	ConstructionVisualizationRegistry.DeregisterVisualization(UE::Dataflow::FMeshStatsConstructionVisualization::Name);
+	ConstructionVisualizationRegistry.DeregisterVisualization(UE::Dataflow::FMeshConstructionVisualization::Name);
 }
 
 IMPLEMENT_MODULE(FDataflowEditorModule, DataflowEditor)
diff --git a/Engine/Plugins/Experimental/Dataflow/Source/DataflowEditor/Private/Dataflow/MeshConstructionVisualization.cpp b/Engine/Plugins/Experimental/Dataflow/Source/DataflowEditor/Private/Dataflow/MeshConstructionVisualization.cpp
new file mode 100644
--- /dev/null
+++ b/Engine/Plugins/Experimental/Dataflow/Source/DataflowEditor/Private/Dataflow/MeshConstructionVisualization.cpp
@@ -0,0 +1,82 @@
+// Copyright Epic Games, Inc. All Rights Reserved.
+
+#include "Dataflow/MeshConstructionVisualization.h"
+#include "Components/DynamicMeshComponent.h"
+#include "Dataflow/DataflowConstructionViewportClient.h"
+#include "Dataflow/DataflowConstructionScene.h"
+#include "Framework/MultiBox/MultiBoxBuilder.h"
+
+#define LOCTEXT_NAMESPACE "MeshConstructionVisualization"
+
+namespace UE::Dataflow
+{
+	FName FMeshConstructionVisualization::Name = FName("MeshConstructionVisualization");
+
+	void FMeshConstructionVisualization::ExtendViewportShowMenu(const TSharedPtr<FDataflowConstructionViewportClient>& ViewportClient, FMenuBuilder& MenuBuilder)
+	{
+		MenuBuilder.BeginSection(TEXT("MeshVisualization"), LOCTEXT("MeshVisualizationSectionName", "Mesh"));
+		{
+			const FUIAction MeshNormalsToggleAction(
+				FExecuteAction::CreateLambda([this, ViewportClient]()
+					{
+						bMeshNormalsVisualizationEnabled = !bMeshNormalsVisualizationEnabled;
+						if (ViewportClient)
+						{
+							ViewportClient->Invalidate();
+						}
+					}),
+				FCanExecuteAction(),
+				FIsActionChecked::CreateLambda([this]()
+					{
+						return bMeshNormalsVisualizationEnabled;
+					}));
+
+			MenuBuilder.AddMenuEntry(LOCTEXT("MeshVisualization_NormalsEnabled", "Normals"),
+				LOCTEXT("MeshVisualization_NormalsEnabled_TooltipText", "Draw mesh vertex normals"),
+				FSlateIcon(),
+				MeshNormalsToggleAction,
+				NAME_None,
+				EUserInterfaceActionType::ToggleButton);
+		}
+		MenuBuilder.EndSection();
+	}
+
+	void FMeshConstructionVisualization::Draw(const FDataflowConstructionScene* ConstructionScene, FPrimitiveDrawInterface* PDI)
+	{
+		if (bMeshNormalsVisualizationEnabled && ConstructionScene && PDI)
+		{
+			for (const UDynamicMeshComponent* const DynamicMeshComponent : ConstructionScene->GetDynamicMeshComponents())
+			{
+				if (const UE::Geometry::FDynamicMesh3* const Mesh = DynamicMeshComponent->GetMesh())
+				{
+					if (Mesh->HasAttributes())
+					{
+						if (const UE::Geometry::FDynamicMeshNormalOverlay* const NormalOverlay = Mesh->Attributes()->PrimaryNormals())
+						{
+							PDI->AddReserveLines(SDPG_World, NormalOverlay->ElementCount());
+
+							for (int32 ElementIndex = 0; ElementIndex < NormalOverlay->MaxElementID(); ++ElementIndex)
+							{
+								if (NormalOverlay->IsElement(ElementIndex))
+								{
+									constexpr float LineLength = 2.0f;
+									constexpr float LineThickness = 0.01f;
+									constexpr FColor LineColor(15, 15, 240);
+
+									const int32 ParentVtx = NormalOverlay->GetParentVertex(ElementIndex);
+									const FVector Normal = (FVector)NormalOverlay->GetElement(ElementIndex);
+									const FVector Origin = (FVector)Mesh->GetVertex(ParentVtx);
+									PDI->DrawLine(Origin, Origin + LineLength * Normal, LineColor , SDPG_World, LineThickness);
+								}
+							}
+						}
+					}
+				}
+			}
+		}
+	}
+
+}
+
+#undef LOCTEXT_NAMESPACE 
+
diff --git a/Engine/Plugins/Experimental/Dataflow/Source/DataflowEditor/Private/Dataflow/MeshConstructionVisualization.h b/Engine/Plugins/Experimental/Dataflow/Source/DataflowEditor/Private/Dataflow/MeshConstructionVisualization.h
new file mode 100644
--- /dev/null
+++ b/Engine/Plugins/Experimental/Dataflow/Source/DataflowEditor/Private/Dataflow/MeshConstructionVisualization.h
@@ -0,0 +1,25 @@
+// Copyright Epic Games, Inc. All Rights Reserved.
+#pragma once
+
+#include "Dataflow/DataflowConstructionVisualization.h"
+
+namespace UE::Dataflow
+{
+	/** Visualization for drawing things on and around the mesh in the construction viewport */
+	class FMeshConstructionVisualization final : public IDataflowConstructionVisualization
+	{
+	public:
+		static FName Name;
+
+	private:
+		virtual FName GetName() const override
+		{
+			return Name;
+		}
+		virtual void ExtendViewportShowMenu(const TSharedPtr<FDataflowConstructionViewportClient>& ViewportClient, FMenuBuilder& MenuBuilder) override;
+		virtual void Draw(const FDataflowConstructionScene* ConstructionScene, FPrimitiveDrawInterface* PDI) override;
+
+		bool bMeshNormalsVisualizationEnabled = false;
+	};
+}
+
diff --git a/Engine/Plugins/Experimental/Dataflow/Source/DataflowEditor/Private/Dataflow/MeshStatsConstructionVisualization.cpp b/Engine/Plugins/Experimental/Dataflow/Source/DataflowEditor/Private/Dataflow/MeshStatsConstructionVisualization.cpp
--- a/Engine/Plugins/Experimental/Dataflow/Source/DataflowEditor/Private/Dataflow/MeshStatsConstructionVisualization.cpp
+++ b/Engine/Plugins/Experimental/Dataflow/Source/DataflowEditor/Private/Dataflow/MeshStatsConstructionVisualization.cpp
@@ -16,7 +16,7 @@ namespace UE::Dataflow
 
 	void FMeshStatsConstructionVisualization::ExtendViewportShowMenu(const TSharedPtr<FDataflowConstructionViewportClient>& ViewportClient, FMenuBuilder& MenuBuilder)
 	{
-		MenuBuilder.BeginSection(TEXT("MeshStatsVisualization"), LOCTEXT("MeshStatsVisualizationSectionName", "Mesh"));
+		MenuBuilder.BeginSection(TEXT("MeshStatsVisualization"), LOCTEXT("MeshStatsVisualizationSectionName", "Mesh Stats"));
 		{
 			const FUIAction MeshStatsToggleAction(
 				FExecuteAction::CreateLambda([this, ViewportClient]()
