Dataflow Editor: changes to both viewport toolbars
- remove options that do nothing/crash
- remove drop-down in the sim viewport that has no items
- add camera speed options to both viewports

diff --git a/Engine/Plugins/Experimental/Dataflow/Source/DataflowEditor/Private/Dataflow/DataflowConstructionViewportToolbar.cpp b/Engine/Plugins/Experimental/Dataflow/Source/DataflowEditor/Private/Dataflow/DataflowConstructionViewportToolbar.cpp
--- a/Engine/Plugins/Experimental/Dataflow/Source/DataflowEditor/Private/Dataflow/DataflowConstructionViewportToolbar.cpp
+++ b/Engine/Plugins/Experimental/Dataflow/Source/DataflowEditor/Private/Dataflow/DataflowConstructionViewportToolbar.cpp
@@ -107,7 +107,7 @@ void SDataflowConstructionViewportSelectionToolBar::Construct(const FArguments&
 				.OnGetMenuContent(this, &SDataflowConstructionViewportSelectionToolBar::GenerateShowMenu)
 		];
 
-
+	// Preview Profile selector
 	MainBoxPtr->AddSlot()
 		.AutoWidth()
 		.Padding(ToolbarSlotPadding)
@@ -149,36 +149,6 @@ TSharedRef<SWidget> SDataflowConstructionViewportSelectionToolBar::GenerateShowM
 	return MenuBuilder.MakeWidget();
 }
 
-// Mostly copied from the inaccessible SCommonEditorViewportToolbarBase::GenerateOptionsMenu()
-TSharedRef<SWidget> SDataflowConstructionViewportSelectionToolBar::GenerateOptionsMenu() const
-{
-	GetInfoProvider().OnFloatingButtonClicked();
-	TSharedRef<SEditorViewport> ViewportRef = GetInfoProvider().GetViewportWidget();
-
-	const bool bIsPerspective = GetViewportClient().GetViewportType() == LVT_Perspective;
-
-	constexpr bool bInShouldCloseWindowAfterMenuSelection = true;
-	FMenuBuilder OptionsMenuBuilder(bInShouldCloseWindowAfterMenuSelection, ViewportRef->GetCommandList());
-	{
-		OptionsMenuBuilder.BeginSection("LevelViewportViewportOptions", LOCTEXT("OptionsMenuHeader", "Viewport Options"));
-		{
-			OptionsMenuBuilder.AddMenuEntry(FEditorViewportCommands::Get().ToggleRealTime);
-			OptionsMenuBuilder.AddMenuEntry(FEditorViewportCommands::Get().ToggleStats);
-			OptionsMenuBuilder.AddMenuEntry(FEditorViewportCommands::Get().ToggleFPS);
-
-			if (bIsPerspective)
-			{
-				OptionsMenuBuilder.AddWidget(GenerateFOVMenu(), LOCTEXT("FOVAngle", "Field of View (H)"));
-			}
-		}
-		OptionsMenuBuilder.EndSection();
-
-		ExtendOptionsMenu(OptionsMenuBuilder);
-	}
-
-	return OptionsMenuBuilder.MakeWidget();
-}
-
 TSharedRef<SWidget> SDataflowConstructionViewportSelectionToolBar::MakeToolBar(const TSharedPtr<FExtender> InExtenders)
 {
 	// The following is modeled after portions of STransformViewportToolBar, which gets 
diff --git a/Engine/Plugins/Experimental/Dataflow/Source/DataflowEditor/Private/Dataflow/DataflowEditorViewportToolBarBase.cpp b/Engine/Plugins/Experimental/Dataflow/Source/DataflowEditor/Private/Dataflow/DataflowEditorViewportToolBarBase.cpp
new file mode 100644
--- /dev/null
+++ b/Engine/Plugins/Experimental/Dataflow/Source/DataflowEditor/Private/Dataflow/DataflowEditorViewportToolBarBase.cpp
@@ -0,0 +1,158 @@
+// Copyright Epic Games, Inc. All Rights Reserved.
+
+#include "Dataflow/DataflowEditorViewportToolBarBase.h"
+#include "SEditorViewport.h"
+#include "Widgets/Input/SSpinBox.h"
+#include "Widgets/Input/SSlider.h"
+#include "Widgets/Layout/SBorder.h"
+#include "EditorViewportClient.h"
+
+#define LOCTEXT_NAMESPACE "DataflowEditorViewportToolBarBase"
+
+TSharedRef<SWidget> SDataflowEditorViewportToolBarBase::GenerateCameraSpeedSettingsMenu() const
+{
+	// This comes from STransformViewportToolBar::FillCameraSpeedMenu
+	return SNew(SBorder)
+		.BorderImage(FAppStyle::GetBrush(TEXT("Menu.Background")))
+		[
+			SNew( SVerticalBox )
+			+SVerticalBox::Slot()
+			.AutoHeight()
+			.Padding( FMargin(8.0f, 2.0f, 60.0f, 2.0f) )
+			.HAlign( HAlign_Left )
+			[
+				SNew( STextBlock )
+				.Text( LOCTEXT("MouseSettingsCamSpeed", "Camera Speed")  )
+				.Font( FAppStyle::GetFontStyle( TEXT( "MenuItem.Font" ) ) )
+			]
+			+SVerticalBox::Slot()
+			.AutoHeight()
+			.Padding( FMargin(8.0f, 4.0f) )
+			[	
+				SNew( SHorizontalBox )
+				+SHorizontalBox::Slot()
+				.FillWidth(1)
+				.Padding( FMargin(0.0f, 2.0f) )
+				[
+					SNew( SBox )
+					.MinDesiredWidth(220)
+					[
+						SNew(SSlider)
+						.Value(this, &SDataflowEditorViewportToolBarBase::GetCamSpeedSliderPosition)
+						.OnValueChanged(this, &SDataflowEditorViewportToolBarBase::OnSetCamSpeed)
+					]
+				]
+				+SHorizontalBox::Slot()
+				.AutoWidth()
+				.Padding( 8.0f, 2.0f, 0.0f, 2.0f)
+				[
+					SNew( SBox )
+					.WidthOverride(40)
+					[
+						SNew( STextBlock )
+						.Text(this, &SDataflowEditorViewportToolBarBase::GetCameraSpeedLabel)
+						.Font( FAppStyle::GetFontStyle( TEXT( "MenuItem.Font" ) ) )
+					]
+				]
+			] // Camera Speed Scalar
+			+ SVerticalBox::Slot()
+			.AutoHeight()
+			.Padding(FMargin(8.0f, 2.0f, 60.0f, 2.0f))
+			.HAlign(HAlign_Left)
+			[
+				SNew(STextBlock)
+				.Text(LOCTEXT("MouseSettingsCamSpeedScalar", "Camera Speed Scalar"))
+				.Font(FAppStyle::GetFontStyle(TEXT("MenuItem.Font")))
+			]
+			+ SVerticalBox::Slot()
+			.AutoHeight()
+			.Padding(FMargin(8.0f, 4.0f))
+			[
+				SNew(SHorizontalBox)
+				+ SHorizontalBox::Slot()
+				.FillWidth(1)
+				.Padding(FMargin(0.0f, 2.0f))
+				[
+					SNew(SSpinBox<float>)
+					.MinValue(1.0f)
+ 					.MaxValue(std::numeric_limits<float>::max())
+					.MinSliderValue(1.0f)
+					.MaxSliderValue(128.0f)
+					.Value(this, &SDataflowEditorViewportToolBarBase::GetCamSpeedScalarBoxValue)
+					.OnValueChanged(this, &SDataflowEditorViewportToolBarBase::OnSetCamSpeedScalarBoxValue)
+					.ToolTipText(LOCTEXT("CameraSpeedScalar_ToolTip", "Scalar to increase camera movement range"))
+				]
+			]
+		];
+}
+
+FText SDataflowEditorViewportToolBarBase::GetCameraSpeedLabel() const
+{
+	const float CameraSpeed = GetViewportClient().GetCameraSpeed();
+	FNumberFormattingOptions FormattingOptions = FNumberFormattingOptions::DefaultNoGrouping();
+	FormattingOptions.MaximumFractionalDigits = CameraSpeed > 1 ? 1 : 3;
+	return FText::AsNumber(CameraSpeed, &FormattingOptions);
+}
+
+float SDataflowEditorViewportToolBarBase::GetCamSpeedSliderPosition() const
+{
+	return (GetViewportClient().GetCameraSpeedSetting() - 1) / ((float)FEditorViewportClient::MaxCameraSpeeds - 1);
+}
+
+void SDataflowEditorViewportToolBarBase::OnSetCamSpeed(float NewValue) const
+{
+	const int32 OldSpeedSetting = GetViewportClient().GetCameraSpeedSetting();
+	const int32 NewSpeedSetting = NewValue * ((float)FEditorViewportClient::MaxCameraSpeeds - 1) + 1;
+
+	if (OldSpeedSetting != NewSpeedSetting)
+	{
+		GetViewportClient().SetCameraSpeedSetting(NewSpeedSetting);
+	}
+}
+
+float SDataflowEditorViewportToolBarBase::GetCamSpeedScalarBoxValue() const
+{
+	return GetViewportClient().GetCameraSpeedScalar();
+}
+
+void SDataflowEditorViewportToolBarBase::OnSetCamSpeedScalarBoxValue(float NewValue) const
+{
+	GetViewportClient().SetCameraSpeedScalar(NewValue);
+}
+
+TSharedRef<SWidget> SDataflowEditorViewportToolBarBase::GenerateOptionsMenu() const
+{
+	GetInfoProvider().OnFloatingButtonClicked();
+	TSharedRef<SEditorViewport> ViewportRef = GetInfoProvider().GetViewportWidget();
+
+	const bool bIsPerspective = GetViewportClient().GetViewportType() == LVT_Perspective;
+
+	constexpr bool bInShouldCloseWindowAfterMenuSelection = true;
+	FMenuBuilder OptionsMenuBuilder(bInShouldCloseWindowAfterMenuSelection, ViewportRef->GetCommandList());
+	{
+		OptionsMenuBuilder.BeginSection("DataflowEditorViewportViewportOptions", LOCTEXT("OptionsMenuHeader", "Viewport Options"));
+		{
+			if (bIsPerspective)
+			{
+				OptionsMenuBuilder.AddWidget(GenerateFOVMenu(), LOCTEXT("FOVAngle", "Field of View (H)"));
+			}
+
+			OptionsMenuBuilder.AddSubMenu(
+				LOCTEXT("CameraSpeedSettings", "Camera Speed Settings"),
+				LOCTEXT("CameraSpeedSettingsToolTip", "Adjust camera speed settings"),
+				FNewMenuDelegate::CreateLambda([this](FMenuBuilder& MenuBuilder)
+					{
+						MenuBuilder.AddWidget(GenerateCameraSpeedSettingsMenu(), FText());
+					}
+				));
+
+		}
+		OptionsMenuBuilder.EndSection();
+
+		ExtendOptionsMenu(OptionsMenuBuilder);
+	}
+
+	return OptionsMenuBuilder.MakeWidget();
+}
+
+#undef LOCTEXT_NAMESPACE
diff --git a/Engine/Plugins/Experimental/Dataflow/Source/DataflowEditor/Private/Dataflow/DataflowEditorViewportToolBarBase.h b/Engine/Plugins/Experimental/Dataflow/Source/DataflowEditor/Private/Dataflow/DataflowEditorViewportToolBarBase.h
new file mode 100644
--- /dev/null
+++ b/Engine/Plugins/Experimental/Dataflow/Source/DataflowEditor/Private/Dataflow/DataflowEditorViewportToolBarBase.h
@@ -0,0 +1,19 @@
+// Copyright Epic Games, Inc. All Rights Reserved.
+
+#pragma once
+
+#include "SCommonEditorViewportToolbarBase.h"
+
+/// Common code for toolbars in the Dataflow Editor
+class SDataflowEditorViewportToolBarBase : public SCommonEditorViewportToolbarBase
+{
+protected:
+	TSharedRef<SWidget> GenerateCameraSpeedSettingsMenu() const;
+	FText GetCameraSpeedLabel() const;
+	float GetCamSpeedSliderPosition() const;
+	void OnSetCamSpeed(float NewValue) const;
+	float GetCamSpeedScalarBoxValue() const;
+	void OnSetCamSpeedScalarBoxValue(float NewValue) const;
+	TSharedRef<SWidget> GenerateOptionsMenu() const;
+};
+
diff --git a/Engine/Plugins/Experimental/Dataflow/Source/DataflowEditor/Private/Dataflow/DataflowSimulationViewportToolbar.cpp b/Engine/Plugins/Experimental/Dataflow/Source/DataflowEditor/Private/Dataflow/DataflowSimulationViewportToolbar.cpp
--- a/Engine/Plugins/Experimental/Dataflow/Source/DataflowEditor/Private/Dataflow/DataflowSimulationViewportToolbar.cpp
+++ b/Engine/Plugins/Experimental/Dataflow/Source/DataflowEditor/Private/Dataflow/DataflowSimulationViewportToolbar.cpp
@@ -13,6 +13,7 @@
 #include "Widgets/Layout/SBorder.h"
 #include "Widgets/SBoxPanel.h"
 #include "LODSyncInterface.h"
+#include "SEditorViewportViewMenu.h"
 
 #define LOCTEXT_NAMESPACE "DataflowSimulationViewportToolBar"
 
@@ -47,6 +48,61 @@ void SDataflowSimulationViewportToolBar::Construct(const FArguments& InArgs, TSh
 		PreviewProfileController(MakeShared<FDataflowPreviewProfileController>(ProfileIndexStorage));
 
 	SCommonEditorViewportToolbarBase::Construct(Args, InDataflowViewport);
+
+
+	const FMargin ToolbarSlotPadding(4.0f, 1.0f);
+	TSharedPtr<SHorizontalBox> MainBox;
+
+	ChildSlot
+	[
+		SNew(SBorder)
+			.BorderImage(FAppStyle::Get().GetBrush("EditorViewportToolBar.Background"))
+			.Cursor(EMouseCursor::Default)
+			[
+				SNew(SVerticalBox)
+					+ SVerticalBox::Slot()
+					.AutoHeight()
+					[
+						SAssignNew(MainBox, SHorizontalBox)
+					]
+			]
+	];
+
+	// Options menu
+	MainBox->AddSlot()
+		.AutoWidth()
+		.Padding(ToolbarSlotPadding)
+		[
+			SNew(SEditorViewportToolbarMenu)
+				.ParentToolBar(SharedThis(this))
+				.Cursor(EMouseCursor::Default)
+				.Image("EditorViewportToolBar.OptionsDropdown")
+				.OnGetMenuContent(this, &SDataflowSimulationViewportToolBar::GenerateOptionsMenu)
+		];
+
+	// Display (Lit, Unlit, Wireframe, etc.)
+	const TSharedRef<SEditorViewport> ViewportRef = StaticCastSharedPtr<SEditorViewport>(EditorViewport.Pin()).ToSharedRef();
+	MainBox->AddSlot()
+		.AutoWidth()
+		.Padding(ToolbarSlotPadding)
+		[
+			SNew(SEditorViewportViewMenu, ViewportRef, SharedThis(this))
+				.Cursor(EMouseCursor::Default)
+				.MenuExtenders(Extenders)
+		];
+
+	// Preview Profile selector
+	MainBox->AddSlot()
+		.AutoWidth()
+		.Padding(ToolbarSlotPadding)
+		[
+			SNew(SPreviewSceneProfileSelector).PreviewProfileController(
+				MakeShared<FDataflowPreviewProfileController>(ProfileIndexStorage))
+		];
+
+	// Add optional toolbar slots to be added by child classes inherited from this common viewport toolbar
+	ExtendLeftAlignedToolbarSlots(MainBox, SharedThis(this));
+
 }
 
 void SDataflowSimulationViewportToolBar::ExtendLeftAlignedToolbarSlots(TSharedPtr<SHorizontalBox> MainBoxPtr, TSharedPtr<SViewportToolBar> ParentToolBarPtr) const
diff --git a/Engine/Plugins/Experimental/Dataflow/Source/DataflowEditor/Public/Dataflow/DataflowConstructionViewportToolbar.h b/Engine/Plugins/Experimental/Dataflow/Source/DataflowEditor/Public/Dataflow/DataflowConstructionViewportToolbar.h
--- a/Engine/Plugins/Experimental/Dataflow/Source/DataflowEditor/Public/Dataflow/DataflowConstructionViewportToolbar.h
+++ b/Engine/Plugins/Experimental/Dataflow/Source/DataflowEditor/Public/Dataflow/DataflowConstructionViewportToolbar.h
@@ -2,14 +2,14 @@
 #pragma once
 
 #include "CoreMinimal.h"
-#include "SCommonEditorViewportToolbarBase.h"
+#include "Dataflow/DataflowEditorViewportToolBarBase.h"
 
 // ----------------------------------------------------------------------------------
 class SDataflowConstructionViewport;
 class SEditorViewportToolbarMenu;
 
 /** Base toolbar for the dataflow. Should be extended to add more features */
-class SDataflowConstructionViewportSelectionToolBar : public SCommonEditorViewportToolbarBase
+class SDataflowConstructionViewportSelectionToolBar : public SDataflowEditorViewportToolBarBase
 {
 public:
 	SLATE_BEGIN_ARGS(SDataflowConstructionViewportSelectionToolBar) {}
@@ -26,8 +26,6 @@ private:
 	// SCommonEditorViewportToolbarBase
 	virtual TSharedRef<SWidget> GenerateShowMenu() const override;
 
-	TSharedRef<SWidget> GenerateOptionsMenu() const;
-
 	TSharedRef<SWidget> MakeToolBar(const TSharedPtr<FExtender> InExtenders);
 
 	// View mode
diff --git a/Engine/Plugins/Experimental/Dataflow/Source/DataflowEditor/Public/Dataflow/DataflowSimulationViewportToolbar.h b/Engine/Plugins/Experimental/Dataflow/Source/DataflowEditor/Public/Dataflow/DataflowSimulationViewportToolbar.h
--- a/Engine/Plugins/Experimental/Dataflow/Source/DataflowEditor/Public/Dataflow/DataflowSimulationViewportToolbar.h
+++ b/Engine/Plugins/Experimental/Dataflow/Source/DataflowEditor/Public/Dataflow/DataflowSimulationViewportToolbar.h
@@ -2,7 +2,7 @@
 #pragma once
 
 #include "CoreMinimal.h"
-#include "SCommonEditorViewportToolbarBase.h"
+#include "Dataflow/DataflowEditorViewportToolBarBase.h"
 
 // ----------------------------------------------------------------------------------
 class SDataflowSimulationViewport;
@@ -10,7 +10,7 @@ class FExtender;
 class FUICommandList;
 
 /** Base toolbar for the dataflow. Should be extended to add more features */
-class SDataflowSimulationViewportToolBar : public SCommonEditorViewportToolbarBase
+class SDataflowSimulationViewportToolBar : public SDataflowEditorViewportToolBarBase
 {
 public:
 	SLATE_BEGIN_ARGS(SDataflowSimulationViewportToolBar) {}
@@ -35,4 +35,4 @@ private:
 
 	TSharedPtr<FUICommandList> CommandList;
 	TSharedPtr<FExtender> Extenders;
-};
\ No newline at end of file
+};
