Dataflow Editor, Simulation Visualization: fix for 'Weight Map' pulldown being empty

diff --git a/Engine/Plugins/ChaosClothAssetEditor/Source/ChaosClothAssetEditor/Private/ChaosClothAsset/ClothDataflowSimulationVisualization.cpp b/Engine/Plugins/ChaosClothAssetEditor/Source/ChaosClothAssetEditor/Private/ChaosClothAsset/ClothDataflowSimulationVisualization.cpp
--- a/Engine/Plugins/ChaosClothAssetEditor/Source/ChaosClothAssetEditor/Private/ChaosClothAsset/ClothDataflowSimulationVisualization.cpp
+++ b/Engine/Plugins/ChaosClothAssetEditor/Source/ChaosClothAssetEditor/Private/ChaosClothAsset/ClothDataflowSimulationVisualization.cpp
@@ -61,6 +61,11 @@ namespace UE::Chaos::ClothAsset
 		return ClothEditorSimulationVisualization.GetDisplayString(GetClothComponent(SimulationScene));
 	}
 
+	void FClothDataflowSimulationVisualization::SimulationSceneUpdated(const FDataflowSimulationScene* SimulationScene)
+	{
+		ClothEditorSimulationVisualization.RefreshMenusForClothComponent(GetClothComponent(SimulationScene));
+	}
+
 }  // End namespace UE::Chaos::ClothAsset
 
 #undef LOCTEXT_NAMESPACE 
diff --git a/Engine/Plugins/ChaosClothAssetEditor/Source/ChaosClothAssetEditor/Private/ChaosClothAsset/ClothDataflowSimulationVisualization.h b/Engine/Plugins/ChaosClothAssetEditor/Source/ChaosClothAssetEditor/Private/ChaosClothAsset/ClothDataflowSimulationVisualization.h
--- a/Engine/Plugins/ChaosClothAssetEditor/Source/ChaosClothAssetEditor/Private/ChaosClothAsset/ClothDataflowSimulationVisualization.h
+++ b/Engine/Plugins/ChaosClothAssetEditor/Source/ChaosClothAssetEditor/Private/ChaosClothAsset/ClothDataflowSimulationVisualization.h
@@ -21,6 +21,7 @@ namespace UE::Chaos::ClothAsset
 		virtual void Draw(const FDataflowSimulationScene* SimulationScene, FPrimitiveDrawInterface* PDI) override;
 		virtual void DrawCanvas(const FDataflowSimulationScene* SimulationScene, FCanvas* Canvas, const FSceneView* SceneView) override;
 		virtual FText GetDisplayString(const FDataflowSimulationScene* SimulationScene) const override;
+		virtual void SimulationSceneUpdated(const FDataflowSimulationScene* SimulationScene) override;
 
 		const UChaosClothComponent* GetClothComponent(const FDataflowSimulationScene* SimulationScene) const;
 
diff --git a/Engine/Plugins/Experimental/Dataflow/Source/DataflowEditor/Private/Dataflow/DataflowSimulationScene.cpp b/Engine/Plugins/Experimental/Dataflow/Source/DataflowEditor/Private/Dataflow/DataflowSimulationScene.cpp
--- a/Engine/Plugins/Experimental/Dataflow/Source/DataflowEditor/Private/Dataflow/DataflowSimulationScene.cpp
+++ b/Engine/Plugins/Experimental/Dataflow/Source/DataflowEditor/Private/Dataflow/DataflowSimulationScene.cpp
@@ -3,6 +3,7 @@
 #include "Dataflow/DataflowSimulationScene.h"
 #include "Dataflow/DataflowSimulationManager.h"
 #include "Dataflow/DataflowSimulationControls.h"
+#include "Dataflow/DataflowSimulationVisualization.h"
 #include "Dataflow/DataflowEditor.h"
 #include "Components/PrimitiveComponent.h"
 #include "Chaos/CacheManagerActor.h"
@@ -338,6 +339,15 @@ void FDataflowSimulationScene::TickDataflowScene(const float DeltaSeconds)
 
 				// Reset the dirty flag
 				EditorContent->SetSimulationDirty(false);
+
+				// Update simulation visualizations
+				for (const TPair<FName, TUniquePtr<UE::Dataflow::IDataflowSimulationVisualization>>& Visualization : UE::Dataflow::FDataflowSimulationVisualizationRegistry::GetInstance().GetVisualizations())
+				{
+					if (Visualization.Value.IsValid())
+					{
+						Visualization.Value->SimulationSceneUpdated(this);
+					}
+				}
 			}
 		}
 
diff --git a/Engine/Plugins/Experimental/Dataflow/Source/DataflowEditor/Public/Dataflow/DataflowSimulationVisualization.h b/Engine/Plugins/Experimental/Dataflow/Source/DataflowEditor/Public/Dataflow/DataflowSimulationVisualization.h
--- a/Engine/Plugins/Experimental/Dataflow/Source/DataflowEditor/Public/Dataflow/DataflowSimulationVisualization.h
+++ b/Engine/Plugins/Experimental/Dataflow/Source/DataflowEditor/Public/Dataflow/DataflowSimulationVisualization.h
@@ -26,6 +26,7 @@ namespace UE::Dataflow
 		DATAFLOWEDITOR_API virtual void Draw(const FDataflowSimulationScene* SimulationScene, FPrimitiveDrawInterface* PDI) {};
 		DATAFLOWEDITOR_API virtual void DrawCanvas(const FDataflowSimulationScene* SimulationScene, FCanvas* Canvas, const FSceneView* SceneView) {};
 		DATAFLOWEDITOR_API virtual FText GetDisplayString(const FDataflowSimulationScene* SimulationScene) const { return FText(); };
+		DATAFLOWEDITOR_API virtual void SimulationSceneUpdated(const FDataflowSimulationScene* SimulationScene) {};
 	};
 
 	class FDataflowSimulationVisualizationRegistry
