Volumetric brush for dataflow paint tool

diff --git a/Engine/Plugins/Experimental/Dataflow/Source/DataflowEditor/Private/DataflowEditorTools/DataflowEditorWeightMapPaintTool.cpp b/Engine/Plugins/Experimental/Dataflow/Source/DataflowEditor/Private/DataflowEditorTools/DataflowEditorWeightMapPaintTool.cpp
--- a/Engine/Plugins/Experimental/Dataflow/Source/DataflowEditor/Private/DataflowEditorTools/DataflowEditorWeightMapPaintTool.cpp
+++ b/Engine/Plugins/Experimental/Dataflow/Source/DataflowEditor/Private/DataflowEditorTools/DataflowEditorWeightMapPaintTool.cpp
@@ -636,7 +636,8 @@ void UDataflowEditorWeightMapPaintTool::UpdateROI(const FSculptBrushStamp& Brush
 	}
 
 	FVector3d CenterNormal = Mesh->IsTriangle(CenterTID) ? TriNormals[CenterTID] : FVector3d::One();		// One so that normal check always passes
-	
+
+	bool bVolumetric = (FilterProperties->BrushAreaMode == EMeshVertexPaintBrushAreaType::Volumetric);
 	bool bUseAngleThreshold = FilterProperties->AngleThreshold < 180.0f;
 	double DotAngleThreshold = FMathd::Cos(FilterProperties->AngleThreshold * FMathd::DegToRad);
 	bool bStopAtUVSeams = FilterProperties->bUVSeams;
@@ -657,24 +658,35 @@ void UDataflowEditorWeightMapPaintTool::UpdateROI(const FSculptBrushStamp& Brush
 		}
 		return false;
 	};
-
 	bool bFill = (FilterProperties->SubToolType == EDataflowEditorWeightMapPaintInteractionType::Fill);
-
-	if (Mesh->IsTriangle(CenterTID))
+	
+	if (bVolumetric)
 	{
-		TArray<int32> StartROI;
-		StartROI.Add(CenterTID);
-		FMeshConnectedComponents::GrowToConnectedTriangles(Mesh, StartROI, TriangleROI, &TempROIBuffer,
-			[&](int t1, int t2) 
-		{ 
-			if ((Mesh->GetTriCentroid(t2) - BrushPos).SquaredLength() < RadiusSqr)
+		Octree.RangeQuery(BrushBox,
+			[&](int TriIdx) {
+			if ((Mesh->GetTriCentroid(TriIdx) - BrushPos).SquaredLength() < RadiusSqr)
 			{
-				return CheckEdgeCriteria(t1, t2);
+				TriangleROI.Add(TriIdx);
 			}
-			return false;
 		});
 	}
-
+	else
+	{
+		if (Mesh->IsTriangle(CenterTID))
+		{
+			TArray<int32> StartROI;
+			StartROI.Add(CenterTID);
+			FMeshConnectedComponents::GrowToConnectedTriangles(Mesh, StartROI, TriangleROI, &TempROIBuffer,
+				[&](int t1, int t2) 
+			{ 
+				if ((Mesh->GetTriCentroid(t2) - BrushPos).SquaredLength() < RadiusSqr)
+				{
+					return CheckEdgeCriteria(t1, t2);
+				}
+				return false;
+			});
+		}
+	}
 	if (bFill)
 	{
 		TArray<int32> StartROI;
diff --git a/Engine/Plugins/Experimental/Dataflow/Source/DataflowEditor/Public/DataflowEditorTools/DataflowEditorWeightMapPaintTool.h b/Engine/Plugins/Experimental/Dataflow/Source/DataflowEditor/Public/DataflowEditorTools/DataflowEditorWeightMapPaintTool.h
--- a/Engine/Plugins/Experimental/Dataflow/Source/DataflowEditor/Public/DataflowEditorTools/DataflowEditorWeightMapPaintTool.h
+++ b/Engine/Plugins/Experimental/Dataflow/Source/DataflowEditor/Public/DataflowEditorTools/DataflowEditorWeightMapPaintTool.h
@@ -22,6 +22,7 @@
 #include "TransformTypes.h"
 #include "ToolDataVisualizer.h"
 #include "GroupTopology.h"
+#include "MeshVertexPaintTool.h"
 #include "DataflowEditorTools/DataflowEditorToolBuilder.h"
 
 #include "DataflowEditorWeightMapPaintTool.generated.h"
@@ -169,7 +170,12 @@ public:
 		HideEditConditionToggle, EditConditionHides, EditCondition = "SubToolType == EDataflowEditorWeightMapPaintInteractionType::Gradient"))
 	double GradientLowValue = 0.0;
 
-
+	/** Area Mode specifies the shape of the brush and which triangles will be included relative to the cursor */
+	UPROPERTY(EditAnywhere, Category = Filters, meta = (DisplayName = "Brush Area Mode",
+		HideEditConditionToggle, EditConditionHides, EditCondition =
+		"SubToolType == EDataflowEditorWeightMapPaintInteractionType::Brush || SubToolType == EDataflowEditorWeightMapPaintInteractionType::Fill"))
+	EMeshVertexPaintBrushAreaType BrushAreaMode = EMeshVertexPaintBrushAreaType::Connected;
+	
 	/** The Region affected by the current operation will be bounded by edge angles larger than this threshold */
 	UPROPERTY(EditAnywhere, Category = Filters, meta = (UIMin = "0.0", ClampMin = "0.0", UIMax = "180.0", ClampMax = "180.0",
 		HideEditConditionToggle, EditConditionHides, EditCondition = 
