Dataflow Editor preview profile controller: if we don't find a preview profile by name, don't then set the current profile index to -1. The name lookup can fail if the user creates a new profile but the editor quits unexpectedly before it is saved to disk.

diff --git a/Engine/Plugins/Experimental/Dataflow/Source/DataflowEditor/Private/Dataflow/DataflowSceneProfileIndexStorage.cpp b/Engine/Plugins/Experimental/Dataflow/Source/DataflowEditor/Private/Dataflow/DataflowSceneProfileIndexStorage.cpp
--- a/Engine/Plugins/Experimental/Dataflow/Source/DataflowEditor/Private/Dataflow/DataflowSceneProfileIndexStorage.cpp
+++ b/Engine/Plugins/Experimental/Dataflow/Source/DataflowEditor/Private/Dataflow/DataflowSceneProfileIndexStorage.cpp
@@ -53,7 +53,7 @@ int32 FDataflowConstructionSceneProfileIndexStorage::RetrieveProfileIndex()
 	if (const UDataflowEditorOptions* const Options = UDataflowEditorOptions::StaticClass()->GetDefaultObject<UDataflowEditorOptions>())
 	{
 		const int32 FoundIndex = UE::Dataflow::Private::GetProfileIndex(Options->ConstructionProfileName);
-		if (FoundIndex != ConstructionScene->GetCurrentProfileIndex())
+		if (FoundIndex != ConstructionScene->GetCurrentProfileIndex() && (FoundIndex != INDEX_NONE))
 		{
 			ConstructionScene->SetCurrentProfileIndex(FoundIndex);
 		}
@@ -91,7 +91,7 @@ int32 FDataflowSimulationSceneProfileIndexStorage::RetrieveProfileIndex()
 	if (const UDataflowEditorOptions* const Options = UDataflowEditorOptions::StaticClass()->GetDefaultObject<UDataflowEditorOptions>())
 	{
 		const int32 FoundIndex = UE::Dataflow::Private::GetProfileIndex(Options->SimulationProfileName);
-		if (FoundIndex != SimulationScene->GetCurrentProfileIndex())
+		if (FoundIndex != SimulationScene->GetCurrentProfileIndex() && (FoundIndex != INDEX_NONE))
 		{
 			SimulationScene->SetCurrentProfileIndex(FoundIndex);
 		}
