Dataflow : add an add pin button for optional inputs
- Place the add pin button like they are in blueprint
- Refactor part of Dataflow SNode that deals with the show/Hide input and now the add pin button
- Remove the top margin

diff --git a/Engine/Plugins/Experimental/Dataflow/Source/DataflowEditor/Private/Dataflow/DataflowSNode.cpp b/Engine/Plugins/Experimental/Dataflow/Source/DataflowEditor/Private/Dataflow/DataflowSNode.cpp
--- a/Engine/Plugins/Experimental/Dataflow/Source/DataflowEditor/Private/Dataflow/DataflowSNode.cpp
+++ b/Engine/Plugins/Experimental/Dataflow/Source/DataflowEditor/Private/Dataflow/DataflowSNode.cpp
@@ -19,6 +19,7 @@
 #include "Widgets/Text/STextBlock.h"
 #include "Widgets/SBoxPanel.h"
 #include "Widgets/Images/SImage.h"
+#include "Widgets/Input/SButton.h"
 #include "Editor/Transactor.h"
 #include "GraphEditorSettings.h"
 #include "Dataflow/DataflowSchema.h"
@@ -382,27 +383,85 @@ void SDataflowEdNode::AddReferencedObjects(FReferenceCollector& Collector)
 	}
 }
 
+void SDataflowEdNode::CreateOutputSideAddButton(TSharedPtr<SVerticalBox> OutputBox)
+{
+	if (DataflowGraphNode)
+	{
+		if (TSharedPtr<FDataflowNode> DataflowNode = DataflowGraphNode->GetDataflowNode())
+		{
+			if (DataflowNode->CanAddPin())
+			{
+				TSharedRef<SWidget> AddPinButton = AddPinButtonContent(
+					LOCTEXT("AddPinButton_Text", "Add Pin"),
+					LOCTEXT("AddPinButton_Tooltip", "Add an optional input pin"),
+					true
+				);
+
+				const FMargin AddPinPadding = Settings->GetOutputPinPadding();
+
+				OutputBox->AddSlot()
+					.AutoHeight()
+					.VAlign(VAlign_Center)
+					.Padding(AddPinPadding)
+					[
+						AddPinButton
+					];
+			}
+		}
+	}
+}
+
 void SDataflowEdNode::CreateInputSideAddButton(TSharedPtr<SVerticalBox> InputBox)
 {
-	TSharedRef<SWidget> AddPinButton = AddPinButtonContent(
-		LOCTEXT("AddPinInputButton", "Show/Hide Inputs"),
-		LOCTEXT("AddPinInputButton_Tooltip", "Show/Hide input pins."),
-		false
-	);
-
-	FMargin AddPinPadding = Settings->GetOutputPinPadding();
-	AddPinPadding.Top += 6.0f;
-
-	InputBox->AddSlot()
-		.AutoHeight()
-		.VAlign(VAlign_Center)
-		.Padding(AddPinPadding)
-		[
-			AddPinButton
-		];
+	if (DataflowGraphNode)
+	{
+		if (TSharedPtr<FDataflowNode> DataflowNode = DataflowGraphNode->GetDataflowNode())
+		{
+			if (DataflowNode->HasHideableInputs())
+			{
+				TSharedRef<SWidget> ShowHideInputsButton = AddPinButtonContent(
+					LOCTEXT("ShowHideInputsButton_Text", "Show/Hide Inputs"),
+					LOCTEXT("ShowHideInputsButton_Tooltip", "Show/Hide input pins."),
+					false
+				);
+
+				// override the on clicked function
+				if (ShowHideInputsButton->GetWidgetClass().GetWidgetType() == SButton::StaticWidgetClass().GetWidgetType())
+				{
+					TSharedRef<SButton> TypedButtonWidget = StaticCastSharedRef<SButton>(ShowHideInputsButton);
+					TypedButtonWidget->SetOnClicked(FOnClicked::CreateSP(this, &SDataflowEdNode::OnShowHideInputs));
+				}
+
+				const FMargin AddPinPadding = Settings->GetInputPinPadding();
+
+				InputBox->AddSlot()
+					.AutoHeight()
+					.VAlign(VAlign_Center)
+					.Padding(AddPinPadding)
+					[
+						ShowHideInputsButton
+					];
+			}
+		}
+	}
 }
 
 FReply SDataflowEdNode::OnAddPin()
+{
+	if (DataflowGraphNode)
+	{
+		if (TSharedPtr<FDataflowNode> DataflowNode = DataflowGraphNode->GetDataflowNode())
+		{
+			if (DataflowNode->CanAddPin())
+			{
+				DataflowGraphNode->AddOptionPin();
+			}
+		}
+	}
+	return FReply::Handled();
+}
+
+FReply SDataflowEdNode::OnShowHideInputs()
 {
 	if (DataflowGraphNode)
 	{
@@ -457,7 +516,7 @@ EVisibility SDataflowEdNode::IsAddPinButtonVisible() const
 	{
 		if (const TSharedPtr<FDataflowNode> DataflowNode = DataflowGraphNode->GetDataflowNode())
 		{
-			if (DataflowNode->HasHideableInputs())
+			if (DataflowNode->HasHideableInputs() || DataflowNode->CanAddPin())
 			{
 				return Visibility;
 			}
diff --git a/Engine/Plugins/Experimental/Dataflow/Source/DataflowEditor/Public/Dataflow/DataflowSNode.h b/Engine/Plugins/Experimental/Dataflow/Source/DataflowEditor/Public/Dataflow/DataflowSNode.h
--- a/Engine/Plugins/Experimental/Dataflow/Source/DataflowEditor/Public/Dataflow/DataflowSNode.h
+++ b/Engine/Plugins/Experimental/Dataflow/Source/DataflowEditor/Public/Dataflow/DataflowSNode.h
@@ -54,11 +54,16 @@ private:
 	// SGraphNode interface
 	/** Override this to create a button to add pins on the input side of the node */
 	virtual void CreateInputSideAddButton(TSharedPtr<SVerticalBox> InputBox) override;
+	/** Override this to create a button to add pins on the output side of the node */
+	virtual void CreateOutputSideAddButton(TSharedPtr<SVerticalBox> OutputBox) override;
 	/** Callback function executed when Add pin button is clicked */
 	virtual FReply OnAddPin() override;
 	/** Checks whether Add pin button should currently be visible */
 	virtual EVisibility IsAddPinButtonVisible() const;
 
+	/** when clicking on the show/hide input pins button */
+	FReply OnShowHideInputs();
+
 	const FSlateBrush* GetPinButtonImage() const;
 
 	virtual TSharedPtr<SGraphPin> CreatePinWidget(UEdGraphPin* Pin) const override;
