Dataflow : add support for watching output
- Output have now a Watch output menu entry
- Save watched output GUIDs in the dataflow ed node
- Extend dataflow node widget to display watch output under it
- move some static utility method from the schema class to the ed node one

diff --git a/Engine/Plugins/Experimental/Dataflow/Source/DataflowEditor/Private/Dataflow/DataflowGraphEditor.cpp b/Engine/Plugins/Experimental/Dataflow/Source/DataflowEditor/Private/Dataflow/DataflowGraphEditor.cpp
--- a/Engine/Plugins/Experimental/Dataflow/Source/DataflowEditor/Private/Dataflow/DataflowGraphEditor.cpp
+++ b/Engine/Plugins/Experimental/Dataflow/Source/DataflowEditor/Private/Dataflow/DataflowGraphEditor.cpp
@@ -261,6 +261,15 @@ void SDataflowGraphEditor::Construct(const FArguments& InArgs, UObject* InAssetO
 					FDataflowEditorCommands::Get().AddNewSubGraph,
 					FExecuteAction::CreateSP(this, &SDataflowGraphEditor::AddNewSubGraph)
 			);
+			GraphEditorCommands->MapAction(FGraphEditorCommands::Get().StartWatchingPin,
+				FExecuteAction::CreateSP(this, &SDataflowGraphEditor::OnStartWatchingPin),
+				FCanExecuteAction::CreateSP(this, &SDataflowGraphEditor::CanStartWatchingPin)
+			);
+
+			GraphEditorCommands->MapAction(FGraphEditorCommands::Get().StopWatchingPin,
+				FExecuteAction::CreateSP(this, &SDataflowGraphEditor::OnStopWatchingPin),
+				FCanExecuteAction::CreateSP(this, &SDataflowGraphEditor::CanStopWatchingPin)
+			);
 
 		}
 	}
@@ -912,6 +921,45 @@ void SDataflowGraphEditor::PasteSelectedNodes()
 	}
 }
 
+void SDataflowGraphEditor::OnStartWatchingPin()
+{
+	if (const UEdGraphPin* Pin = GetGraphPinForMenu())
+	{
+		if (UDataflowEdNode* EdNode = Cast<UDataflowEdNode>(Pin->GetOwningNode()))
+		{
+			EdNode->WatchPin(Pin, true);
+		}
+	}
+}
+
+bool SDataflowGraphEditor::CanStartWatchingPin() const
+{
+	return !CanStopWatchingPin();
+}
+
+void SDataflowGraphEditor::OnStopWatchingPin()
+{
+	if (const UEdGraphPin* Pin = GetGraphPinForMenu())
+	{
+		if (UDataflowEdNode* EdNode = Cast<UDataflowEdNode>(Pin->GetOwningNode()))
+		{
+			EdNode->WatchPin(Pin, false);
+		}
+	}
+}
+
+bool SDataflowGraphEditor::CanStopWatchingPin() const
+{
+	if (UEdGraphPin* Pin = const_cast<SDataflowGraphEditor*>(this)->GetGraphPinForMenu())
+	{
+		if (const UDataflowEdNode* EdNode = Cast<UDataflowEdNode>(Pin->GetOwningNode()))
+		{
+			return EdNode->IsPinWatched(Pin);
+		}
+	}
+	return false;
+}
+
 void SDataflowGraphEditor::InitGraphEditorMessageBar()
 {
 	SAssignNew(MessageBar, SHorizontalBox)
diff --git a/Engine/Plugins/Experimental/Dataflow/Source/DataflowEditor/Private/Dataflow/DataflowSNode.cpp b/Engine/Plugins/Experimental/Dataflow/Source/DataflowEditor/Private/Dataflow/DataflowSNode.cpp
--- a/Engine/Plugins/Experimental/Dataflow/Source/DataflowEditor/Private/Dataflow/DataflowSNode.cpp
+++ b/Engine/Plugins/Experimental/Dataflow/Source/DataflowEditor/Private/Dataflow/DataflowSNode.cpp
@@ -8,6 +8,7 @@
 #include "Dataflow/DataflowGraphEditor.h"
 #include "Dataflow/DataflowNodeFactory.h"
 #include "Dataflow/DataflowObject.h"
+#include "Dataflow/DataflowTypePolicy.h"
 #include "Framework/Application/SlateApplication.h"
 #include "Logging/LogMacros.h"
 #include "SourceCodeNavigation.h"
@@ -31,6 +32,44 @@ namespace UE::Dataflow::Private
 {
 	static bool bShowPerfData = true;
 	static FAutoConsoleVariableRef CVarChaosBoxCalcBoundsISPCEnabled(TEXT("p.Dataflow.ShowPerfData"), bShowPerfData, TEXT("Whether to display perf data on top of dataflow nodes"));
+
+	void AddOutputValueToString(const FDataflowOutput& Output, FContext& Context, FString& OutString)
+	{
+		using namespace UE::Dataflow::Private;
+
+		if (!Output.HasCachedValue(Context))
+		{
+			OutString += TEXT("/!\\ ");
+		}
+		if (FDataflowSelectionTypePolicy::SupportsTypeStatic(Output.GetType()))
+		{
+			const FDataflowSelection DefaultValue;
+			const FDataflowSelection& Value = Output.ReadValue(Context, DefaultValue);
+			OutString += FString::Printf(TEXT("%s : %d/%d\n"), *Output.GetName().ToString(), Value.NumSelected(), Value.Num());
+		}
+		else if (FDataflowStringConvertibleTypePolicy::SupportsTypeStatic(Output.GetType()))
+		{
+			// Try converting to a string using FDataflowStringConvertibleTypes
+			const FString DefaultValue;
+			const FString Value = Output.ReadValue<FDataflowStringConvertibleTypes>(Context, DefaultValue);
+			OutString += FString::Printf(TEXT("%s : %s\n"), *Output.GetName().ToString(), *Value);
+		}
+		else if (FDataflowUObjectConvertibleTypePolicy::SupportsTypeStatic(Output.GetType()))
+		{
+			static const FName NullObject(TEXT("(null)"));
+
+			// Try converting to a string using FDataflowUObjectConvertibleTypes
+			const TObjectPtr<UObject> DefaultValue;
+			const TObjectPtr<UObject> Value = Output.ReadValue<FDataflowUObjectConvertibleTypes>(Context, DefaultValue);
+			const FName ObjectName = Value ? Value->GetFName() : NullObject;
+			OutString += FString::Printf(TEXT("%s : %s\n"), *Output.GetName().ToString(), *ObjectName.ToString());
+		}
+		else
+		{
+			OutString += FString::Printf(TEXT("%s : (Cannot watch this type)\n"), *Output.GetName().ToString());
+		}
+	}
+
 }
 
 const FSlateBrush* SDataflowEdNode::GetPinButtonImage() const
@@ -72,37 +111,7 @@ void SDataflowEdNode::Construct(const FArguments& InArgs, UDataflowEdNode* InNod
 			});
 
 	PerfWidget = SNew(STextBlock);
-
-
-	//
-	//  Cached
-	//
-	/*
-	const FSlateBrush* CachedFalseBrush = FDataflowEditorStyle::Get().GetBrush(TEXT("Dataflow.Cached.False"));
-	const FSlateBrush* CachedTrueBrush = FDataflowEditorStyle::Get().GetBrush(TEXT("Dataflow.Cached.True"));
-
-	CacheStatusStyle = FCheckBoxStyle()
-		.SetCheckBoxType(ESlateCheckBoxType::CheckBox)
-		.SetUncheckedImage(*CachedFalseBrush)
-		.SetUncheckedHoveredImage(*CachedFalseBrush)
-		.SetUncheckedPressedImage(*CachedFalseBrush)
-		.SetCheckedImage(*CachedTrueBrush)
-		.SetCheckedHoveredImage(*CachedTrueBrush)
-		.SetCheckedPressedImage(*CachedTrueBrush)
-		.SetPadding(FMargin(0, 0, 0, 1));
-
-		CacheStatus = SNew(SCheckBox)
-		.Style(&CacheStatusStyle)
-		.IsChecked_Lambda([this]()-> ECheckBoxState
-		{
-			if (DataflowGraphNode)
-			{
-				return ECheckBoxState::Checked;
-				// @todo(dataflow) : Missing connection to FToolkit
-			}
-			return ECheckBoxState::Unchecked;
-		});
-	*/
+	WatchWidget = SNew(STextBlock);
 }
 
 TSharedPtr<SGraphPin> SDataflowEdNode::CreatePinWidget(UEdGraphPin* Pin) const
@@ -270,21 +279,34 @@ TArray<FOverlayWidgetInfo> SDataflowEdNode::GetOverlayWidgets(bool bSelected, co
 				}
 			}
 		}
-	}
 
-	/*
-	if (DataflowGraphNode )
-	{
-		// @todo(dataflow) : Need to bump the title box over 6px
-		const FVector2D ImageSize = CacheStatus->GetDesiredSize();
+		FString WatchString;
+		if (DataflowGraphNode && DataflowGraphNode->HasAnyWatchedConnection())
+		{
+			if (TSharedPtr<UE::Dataflow::FContext> Context = DataflowInterface->GetDataflowContext())
+			{
+				if (TSharedPtr<const FDataflowNode> DataflowNode = DataflowGraphNode->GetDataflowNode())
+				{
+					for (const FDataflowOutput* Output : DataflowNode->GetOutputs())
+					{
+						if (Output && DataflowGraphNode->IsConnectionWatched(*Output))
+						{
+							UE::Dataflow::Private::AddOutputValueToString(*Output, *Context, WatchString);
+						}
+					}
+				}
 
-		FOverlayWidgetInfo Info;
-		Info.OverlayOffset = FVector2D(6.f, 6.f);
-		Info.Widget = CacheStatus;
+				WatchWidget->SetVisibility(WatchString.IsEmpty() ? EVisibility::Collapsed : EVisibility::Visible);
+				WatchWidget->SetText(FText::FromString(WatchString));
 
-		Widgets.Add(Info);
+				FOverlayWidgetInfo WatchWidgetInfo;
+				WatchWidgetInfo.OverlayOffset = FVector2D{ 0, (double)WidgetSize.Y + 4.0 };
+				WatchWidgetInfo.Widget = WatchWidget;
+				Widgets.Add(WatchWidgetInfo);
+			}
+		}
 	}
-	*/
+
 	return Widgets;
 }
 
diff --git a/Engine/Plugins/Experimental/Dataflow/Source/DataflowEditor/Private/Dataflow/DataflowSchema.cpp b/Engine/Plugins/Experimental/Dataflow/Source/DataflowEditor/Private/Dataflow/DataflowSchema.cpp
--- a/Engine/Plugins/Experimental/Dataflow/Source/DataflowEditor/Private/Dataflow/DataflowSchema.cpp
+++ b/Engine/Plugins/Experimental/Dataflow/Source/DataflowEditor/Private/Dataflow/DataflowSchema.cpp
@@ -95,6 +95,19 @@ void UDataflowSchema::GetContextMenuActions(class UToolMenu* Menu, class UGraphN
 		}
 	}
 	Super::GetContextMenuActions(Menu, Context);
+
+	if (const UDataflowEdNode* EdNode = Cast<UDataflowEdNode>(Context->Node))
+	{
+		if (Context->Pin && Context->Pin->Direction == EGPD_Output)
+		{
+			if (FDataflowConnection* Connection = UDataflowEdNode::GetConnectionFromPin(Context->Pin))
+			{
+				FToolMenuSection& Section = Menu->AddSection("DataflowSchema_PinContextMenu_SectionDebug", LOCTEXT("DataflowSchema_PinContextMenu_SectionDebug_Text", "Debug"));
+				const bool bConnectionWatched = EdNode->IsConnectionWatched(*Connection);
+				Section.AddMenuEntry(bConnectionWatched? FGraphEditorCommands::Get().StopWatchingPin: FGraphEditorCommands::Get().StartWatchingPin);
+			}
+		}
+	}
 }
 bool UDataflowSchema::CanPinBeConnectedToNode(const UEdGraphPin* Pin, const UE::Dataflow::FFactoryParameters& NodeParameters)
 {
@@ -345,44 +358,6 @@ const FPinConnectionResponse UDataflowSchema::CanCreateConnection(const UEdGraph
 	return FPinConnectionResponse(CONNECT_RESPONSE_DISALLOW, NoConnectionResponse[FMath::RandRange(0, NoConnectionResponse.Num()-1)]);
 }
 
-TSharedPtr<FDataflowNode> UDataflowSchema::GetDataflowNodeFromEdNode(UEdGraphNode* EdNode)
-{
-	if (UDataflowEdNode* DataflowEdNode = Cast<UDataflowEdNode>(EdNode))
-	{
-		return DataflowEdNode->GetDataflowNode();
-	}
-	return {};
-}
-
-TSharedPtr<const FDataflowNode> UDataflowSchema::GetDataflowNodeFromEdNode(const UEdGraphNode* EdNode)
-{
-	if (const UDataflowEdNode* DataflowEdNode = Cast<UDataflowEdNode>(EdNode))
-	{
-		return DataflowEdNode->GetDataflowNode();
-	}
-	return {};
-}
-
-FDataflowConnection* UDataflowSchema::GetConnectionFromPin(const UEdGraphPin* Pin)
-{
-	FDataflowConnection* DataflowConnection = nullptr;
-	if (Pin)
-	{
-		if (TSharedPtr<FDataflowNode> DataflowNode = GetDataflowNodeFromEdNode(Pin->GetOwningNode()))
-		{
-			if (Pin->Direction == EEdGraphPinDirection::EGPD_Input)
-			{
-				DataflowConnection = DataflowNode->FindInput(Pin->PinName);
-			}
-			else if (Pin->Direction == EEdGraphPinDirection::EGPD_Output)
-			{
-				DataflowConnection = DataflowNode->FindOutput(Pin->PinName);
-			}
-		}
-	}
-	return DataflowConnection;
-}
-
 UE::Dataflow::FPin::EDirection UDataflowSchema::GetDirectionFromPinDirection(EEdGraphPinDirection InPinDirection)
 {
 	switch (InPinDirection)
@@ -447,7 +422,7 @@ void UDataflowSchema::GetAssetsNodeHoverMessage(const TArray<struct FAssetData>&
 {
 	OutOkIcon = false;
 	OutTooltipText = FString(TEXT("Asset type unsupported by this node"));
-	if (TSharedPtr<const FDataflowNode> DataflowNode = GetDataflowNodeFromEdNode(HoverNode))
+	if (TSharedPtr<const FDataflowNode> DataflowNode = UDataflowEdNode::GetDataflowNodeFromEdNode(HoverNode))
 	{
 		bool bNeedInvalidation = false;
 		for (const FAssetData& AssetData : Assets)
@@ -467,7 +442,7 @@ void UDataflowSchema::GetAssetsNodeHoverMessage(const TArray<struct FAssetData>&
 
 void UDataflowSchema::DroppedAssetsOnNode(const TArray<struct FAssetData>& Assets, const FVector2D& GraphPosition, UEdGraphNode* Node) const
 {
-	if (TSharedPtr<FDataflowNode> DataflowNode = GetDataflowNodeFromEdNode(Node))
+	if (TSharedPtr<FDataflowNode> DataflowNode = UDataflowEdNode::GetDataflowNodeFromEdNode(Node))
 	{
 		bool bNeedInvalidation = false;
 		for (const FAssetData& AssetData : Assets)
@@ -490,7 +465,7 @@ void UDataflowSchema::DroppedAssetsOnNode(const TArray<struct FAssetData>& Asset
 
 bool UDataflowSchema::SupportsDropPinOnNode(UEdGraphNode* InTargetNode, const FEdGraphPinType& InSourcePinType, EEdGraphPinDirection InSourcePinDirection, FText& OutMessage) const
 {
-	if (TSharedPtr<FDataflowNode> TargetDataflowNode = GetDataflowNodeFromEdNode(InTargetNode))
+	if (TSharedPtr<FDataflowNode> TargetDataflowNode = UDataflowEdNode::GetDataflowNodeFromEdNode(InTargetNode))
 	{
 		const UE::Dataflow::FPin::EDirection SourceDirection = GetDirectionFromPinDirection(InSourcePinDirection);
 		if (TargetDataflowNode->SupportsDropConnectionOnNode(InSourcePinType.PinCategory, SourceDirection))
@@ -507,9 +482,9 @@ UEdGraphPin* UDataflowSchema::DropPinOnNode(UEdGraphNode* InTargetNode, const FN
 {
 	UEdGraphPin* ResultPin = nullptr;
 
-	if (TSharedPtr<FDataflowNode> TargetDataflowNode = GetDataflowNodeFromEdNode(InTargetNode))
+	if (TSharedPtr<FDataflowNode> TargetDataflowNode = UDataflowEdNode::GetDataflowNodeFromEdNode(InTargetNode))
 	{
-		if (const FDataflowConnection* SourceConnection = GetConnectionFromPin(PinBeingDropped))
+		if (const FDataflowConnection* SourceConnection = UDataflowEdNode::GetConnectionFromPin(PinBeingDropped))
 		{
 			if (const FDataflowConnection* NewConnection = TargetDataflowNode->OnDropConnectionOnNode(*SourceConnection))
 			{
diff --git a/Engine/Plugins/Experimental/Dataflow/Source/DataflowEditor/Public/Dataflow/DataflowGraphEditor.h b/Engine/Plugins/Experimental/Dataflow/Source/DataflowEditor/Public/Dataflow/DataflowGraphEditor.h
--- a/Engine/Plugins/Experimental/Dataflow/Source/DataflowEditor/Public/Dataflow/DataflowGraphEditor.h
+++ b/Engine/Plugins/Experimental/Dataflow/Source/DataflowEditor/Public/Dataflow/DataflowGraphEditor.h
@@ -176,6 +176,11 @@ private:
 	/** Return whether all currently selected Dataflow nodes can execute the RemovePin function. */
 	bool CanRemoveOptionPin() const;
 
+	void OnStartWatchingPin();
+	bool CanStartWatchingPin() const;
+	void OnStopWatchingPin();
+	bool CanStopWatchingPin() const;
+
 	bool GetPinVisibility(SGraphEditor::EPinVisibility InVisibility) const;
 
 	/** Create a widget to display an overlaid message in the graph editor panel */
diff --git a/Engine/Plugins/Experimental/Dataflow/Source/DataflowEditor/Public/Dataflow/DataflowSNode.h b/Engine/Plugins/Experimental/Dataflow/Source/DataflowEditor/Public/Dataflow/DataflowSNode.h
--- a/Engine/Plugins/Experimental/Dataflow/Source/DataflowEditor/Public/Dataflow/DataflowSNode.h
+++ b/Engine/Plugins/Experimental/Dataflow/Source/DataflowEditor/Public/Dataflow/DataflowSNode.h
@@ -75,6 +75,7 @@ private:
 	TSharedPtr<SCheckBox> RenderCheckBoxWidget;
 	TSharedPtr<SImage> FreezeImageWidget;
 	TSharedPtr<STextBlock> PerfWidget;
+	TSharedPtr<STextBlock> WatchWidget;
 
 	//FCheckBoxStyle CacheStatusStyle;
 	//TSharedPtr<SCheckBox> CacheStatus;
diff --git a/Engine/Plugins/Experimental/Dataflow/Source/DataflowEditor/Public/Dataflow/DataflowSchema.h b/Engine/Plugins/Experimental/Dataflow/Source/DataflowEditor/Public/Dataflow/DataflowSchema.h
--- a/Engine/Plugins/Experimental/Dataflow/Source/DataflowEditor/Public/Dataflow/DataflowSchema.h
+++ b/Engine/Plugins/Experimental/Dataflow/Source/DataflowEditor/Public/Dataflow/DataflowSchema.h
@@ -48,13 +48,11 @@ public:
 	float GetPinTypeWireThickness(const FName& Type) const;
 	TOptional<FLinearColor> GetPinColorOverride(TSharedPtr<FDataflowNode> DataflowNode, UEdGraphPin* Pin) const;
 
+
 private:
 	static bool CanPinBeConnectedToNode(const UEdGraphPin* Pin, const UE::Dataflow::FFactoryParameters& NodeParameters);
 	static bool CanConnectPins(const UEdGraphPin& OutputPin, const UEdGraphPin& InputPin);
 	
-	static TSharedPtr<FDataflowNode> GetDataflowNodeFromEdNode(UEdGraphNode* EdNode);
-	static TSharedPtr<const FDataflowNode> GetDataflowNodeFromEdNode(const UEdGraphNode* EdNode);
-	static FDataflowConnection* GetConnectionFromPin(const UEdGraphPin* Pin);
 	static UE::Dataflow::FPin::EDirection GetDirectionFromPinDirection(EEdGraphPinDirection InPinDirection);
 
 	// used by SupportsDropPinOnNode to know which original connection is being dropped
diff --git a/Engine/Source/Runtime/Experimental/Dataflow/Engine/Private/Dataflow/DataflowEdNode.cpp b/Engine/Source/Runtime/Experimental/Dataflow/Engine/Private/Dataflow/DataflowEdNode.cpp
--- a/Engine/Source/Runtime/Experimental/Dataflow/Engine/Private/Dataflow/DataflowEdNode.cpp
+++ b/Engine/Source/Runtime/Experimental/Dataflow/Engine/Private/Dataflow/DataflowEdNode.cpp
@@ -447,6 +447,45 @@ bool UDataflowEdNode::PinIsCompatibleWithType(const UEdGraphPin& Pin, const FEdG
 	return false;
 }
 
+bool UDataflowEdNode::HasAnyWatchedConnection() const
+{
+	return (WatchedConnections.Num() > 0);
+}
+
+bool UDataflowEdNode::IsConnectionWatched(const FDataflowConnection& Connection) const
+{
+	return WatchedConnections.Contains(Connection.GetGuid());
+}
+
+void UDataflowEdNode::WatchConnection(const FDataflowConnection& Connection, bool Value)
+{
+	if (Value)
+	{
+		WatchedConnections.AddUnique(Connection.GetGuid());
+	}
+	else
+	{
+		WatchedConnections.Remove(Connection.GetGuid());
+	}
+}
+
+bool UDataflowEdNode::IsPinWatched(const UEdGraphPin* Pin) const
+{
+	if (FDataflowConnection* Connection = UDataflowEdNode::GetConnectionFromPin(Pin))
+	{
+		return IsConnectionWatched(*Connection);
+	}
+	return false;
+}
+
+void UDataflowEdNode::WatchPin(const UEdGraphPin* Pin, bool bWatch)
+{
+	if (FDataflowConnection* Connection = UDataflowEdNode::GetConnectionFromPin(Pin))
+	{
+		WatchConnection(*Connection, bWatch);
+	}
+}
+
 FText UDataflowEdNode::GetNodeTitle(ENodeTitleType::Type TitleType) const
 {
 	return FText::FromString(GetName());
@@ -1017,5 +1056,48 @@ void UDataflowEdNode::SetDataflowNodeGuid(FGuid InGuid)
 	RegisterDelegateHandle();
 }
 
+TSharedPtr<FDataflowNode> UDataflowEdNode::GetDataflowNodeFromEdNode(UEdGraphNode* EdNode)
+{
+	if (UDataflowEdNode* DataflowEdNode = Cast<UDataflowEdNode>(EdNode))
+	{
+		return DataflowEdNode->GetDataflowNode();
+	}
+	return {};
+}
+
+TSharedPtr<const FDataflowNode> UDataflowEdNode::GetDataflowNodeFromEdNode(const UEdGraphNode* EdNode)
+{
+	if (const UDataflowEdNode* DataflowEdNode = Cast<UDataflowEdNode>(EdNode))
+	{
+		return DataflowEdNode->GetDataflowNode();
+	}
+	return {};
+}
+
+FDataflowConnection* UDataflowEdNode::GetConnectionFromPin(const UEdGraphPin* Pin)
+{
+#if WITH_EDITOR
+	FDataflowConnection* DataflowConnection = nullptr;
+	if (Pin)
+	{
+		if (TSharedPtr<FDataflowNode> DataflowNode = GetDataflowNodeFromEdNode(Pin->GetOwningNode()))
+		{
+			if (Pin->Direction == EEdGraphPinDirection::EGPD_Input)
+			{
+				DataflowConnection = DataflowNode->FindInput(Pin->PinName);
+			}
+			else if (Pin->Direction == EEdGraphPinDirection::EGPD_Output)
+			{
+				DataflowConnection = DataflowNode->FindOutput(Pin->PinName);
+			}
+		}
+	}
+	return DataflowConnection;
+#else
+	ensure(false); // shoudl always be called in editor
+	return nullptr;
+#endif
+}
+
 #undef LOCTEXT_NAMESPACE
 
diff --git a/Engine/Source/Runtime/Experimental/Dataflow/Engine/Public/Dataflow/DataflowEdNode.h b/Engine/Source/Runtime/Experimental/Dataflow/Engine/Public/Dataflow/DataflowEdNode.h
--- a/Engine/Source/Runtime/Experimental/Dataflow/Engine/Public/Dataflow/DataflowEdNode.h
+++ b/Engine/Source/Runtime/Experimental/Dataflow/Engine/Public/Dataflow/DataflowEdNode.h
@@ -106,6 +106,17 @@ public:
 	DATAFLOWENGINE_API void RegisterDelegateHandle();
 	DATAFLOWENGINE_API void UnRegisterDelegateHandle();
 
+	DATAFLOWENGINE_API bool HasAnyWatchedConnection() const;
+	DATAFLOWENGINE_API bool IsConnectionWatched(const FDataflowConnection& Connection) const;
+	DATAFLOWENGINE_API void WatchConnection(const FDataflowConnection& Connection, bool Value);
+
+	DATAFLOWENGINE_API bool IsPinWatched(const UEdGraphPin* Pin) const;
+	DATAFLOWENGINE_API void WatchPin(const UEdGraphPin* Pin, bool bWatch);
+
+	DATAFLOWENGINE_API static FDataflowConnection* GetConnectionFromPin(const UEdGraphPin* Pin);
+	DATAFLOWENGINE_API static TSharedPtr<FDataflowNode> GetDataflowNodeFromEdNode(UEdGraphNode* EdNode);
+	DATAFLOWENGINE_API static TSharedPtr<const FDataflowNode> GetDataflowNodeFromEdNode(const UEdGraphNode* EdNode);
+
 private:
 
 	UPROPERTY(Transient)
@@ -117,6 +128,9 @@ private:
 	UPROPERTY(Transient)
 	bool bCanEnableRenderWireframe = true;
 
+	// Store Guids from connection being watched (display values of corresponding output) 
+	TArray<FGuid> WatchedConnections;
+
 	void RemoveSpacesInAllPinTypes();
 
 	FDelegateHandle OnNodeInvalidatedDelegateHandle;
