Dataflow : fix copy and paste issues with anytypes
- Add extra type information to clipboard connection data to reconstruct the missing information during the paste operation

diff --git a/Engine/Plugins/Experimental/Dataflow/Source/DataflowEditor/Private/Dataflow/AssetDefinition_DataflowAsset.cpp b/Engine/Plugins/Experimental/Dataflow/Source/DataflowEditor/Private/Dataflow/AssetDefinition_DataflowAsset.cpp
--- a/Engine/Plugins/Experimental/Dataflow/Source/DataflowEditor/Private/Dataflow/AssetDefinition_DataflowAsset.cpp
+++ b/Engine/Plugins/Experimental/Dataflow/Source/DataflowEditor/Private/Dataflow/AssetDefinition_DataflowAsset.cpp
@@ -214,6 +214,17 @@ EAssetCommandResult UAssetDefinition_DataflowAsset::OpenAssets(const FAssetOpenA
 	return EAssetCommandResult::Unhandled;
 }
 
+void FDataflowConnectionData::Set(const FDataflowOutput& Output, const FDataflowInput& Input)
+{
+	static constexpr TCHAR Format[] = TEXT("/{0}:{1}|{2}");
+
+	const FDataflowNode* OutputNode = Output.GetOwningNode();
+	const FDataflowNode* InputNode = Input.GetOwningNode();
+
+	Out = FString::Format(Format, { OutputNode ? OutputNode->GetName().ToString() : FString(), Output.GetName().ToString(), Output.GetType().ToString() });
+	In = FString::Format(Format, { InputNode ? InputNode->GetName().ToString() : FString(), Input.GetName().ToString(), Input.GetType().ToString() });
+}
+
 FString FDataflowConnectionData::GetNode(const FString InConnection)
 {
 	FString Left, Right;
@@ -226,19 +237,46 @@ FString FDataflowConnectionData::GetNode(const FString InConnection)
 		}
 	}
 
-	return FString("");
+	return {};
 }
 
 FString FDataflowConnectionData::GetProperty(const FString InConnection)
 {
-	FString Right;
+	FString Left, Right;
 
 	if (InConnection.Split(TEXT(":"), nullptr, &Right, ESearchCase::IgnoreCase, ESearchDir::FromEnd))
 	{
-		return Right;
+		if (Right.Split(TEXT("|"), &Left, nullptr, ESearchCase::IgnoreCase, ESearchDir::FromEnd))
+		{
+			return Left;
+		}
+		else
+		{
+			return Right; // no types
+		}
 	}
 
-	return FString("");
+	return {};
+}
+
+void FDataflowConnectionData::GetNodePropertyAndType(const FString InConnection, FString& OutNode, FString& OutProperty, FString& OutType)
+{
+	FString Left, Right;
+
+	// String should look like "/NodeName:Property|Type"
+	if (InConnection.Split(TEXT(":"), &Left, &Right, ESearchCase::IgnoreCase, ESearchDir::FromEnd))
+	{
+		// get the node name from "/NodeName"
+		Left.Split(TEXT("/"), nullptr, &OutNode, ESearchCase::IgnoreCase, ESearchDir::FromEnd);
+
+		// get the property and type from "Property|Type"
+		const bool bHasType = Right.Split(TEXT("|"), &OutProperty, &OutType, ESearchCase::IgnoreCase, ESearchDir::FromEnd);
+		if (!bHasType)
+		{
+			OutProperty = Right;
+			OutType = {};
+		}
+	}
 }
 
 #undef LOCTEXT_NAMESPACE
diff --git a/Engine/Plugins/Experimental/Dataflow/Source/DataflowEditor/Private/Dataflow/DataflowAssetEditUtils.cpp b/Engine/Plugins/Experimental/Dataflow/Source/DataflowEditor/Private/Dataflow/DataflowAssetEditUtils.cpp
--- a/Engine/Plugins/Experimental/Dataflow/Source/DataflowEditor/Private/Dataflow/DataflowAssetEditUtils.cpp
+++ b/Engine/Plugins/Experimental/Dataflow/Source/DataflowEditor/Private/Dataflow/DataflowAssetEditUtils.cpp
@@ -638,24 +638,21 @@ namespace UE::Dataflow
 		// Build connection data
 		for (const FDataflowInput* Input : NodeInputsToSave)
 		{
-			if (!Input) continue;
-
-			const FDataflowNode* InputNode = Input->GetOwningNode();
-			if (!InputNode) continue;
+			if (!Input)
+			{
+				continue;
+			}
 
 			const FDataflowOutput* Output = Input->GetConnection();
-			if (!Output) continue;
-
-			const FDataflowNode* OutputNode = Output->GetOwningNode();
-			if (!OutputNode) continue;
-
-			if (NodeGuids.Contains(OutputNode->GetGuid()))
+			if (!Output)
 			{
-				FDataflowConnectionData DataflowConnectionData
-				{
-					.Out = FString::Format(TEXT("/{0}:{1}"), { OutputNode->GetName().ToString(), Output->GetName().ToString() }),
-					.In = FString::Format(TEXT("/{0}:{1}"), { InputNode->GetName().ToString(), Input->GetName().ToString() }),
-				};
+				continue;
+			}
+			
+			if (NodeGuids.Contains(Output->GetOwningNodeGuid()))
+			{
+				FDataflowConnectionData DataflowConnectionData;
+				DataflowConnectionData.Set(*Output, *Input);
 				CopyPasteContent.ConnectionData.Emplace(DataflowConnectionData);
 			}
 		}
@@ -769,18 +766,21 @@ namespace UE::Dataflow
 			// Recreate connections
 			for (const FDataflowConnectionData& Connection : CopyPasteContent.ConnectionData)
 			{
-				const FString NodeIn = FDataflowConnectionData::GetNode(Connection.In);
-				const UDataflowEdNode* EdNodeIn = OriginalNodeNameToEdNode[NodeIn];
-				const FGuid GuidIn = EdNodeIn? EdNodeIn->DataflowNodeGuid: FGuid();
+				FString NodeIn, PropertyIn, TypeIn;
+				FDataflowConnectionData::GetNodePropertyAndType(Connection.In, NodeIn, PropertyIn, TypeIn);
 
-				const FString PropertyIn = FDataflowConnectionData::GetProperty(Connection.In);
-				const FName InputputName = *PropertyIn;
+				FString NodeOut, PropertyOut, TypeOut;
+				FDataflowConnectionData::GetNodePropertyAndType(Connection.Out, NodeOut, PropertyOut, TypeOut);
 
-				const FString NodeOut = FDataflowConnectionData::GetNode(Connection.Out);
+				ensure(TypeIn == TypeOut);
+
+				const UDataflowEdNode* EdNodeIn = OriginalNodeNameToEdNode[NodeIn];
 				const UDataflowEdNode* EdNodeOut = OriginalNodeNameToEdNode[NodeOut];
-				const FGuid GuidOut = EdNodeOut? EdNodeOut->DataflowNodeGuid: FGuid();
 
-				const FString PropertyOut = FDataflowConnectionData::GetProperty(Connection.Out);
+				const FGuid GuidIn = EdNodeIn? EdNodeIn->DataflowNodeGuid: FGuid();
+				const FGuid GuidOut = EdNodeOut ? EdNodeOut->DataflowNodeGuid : FGuid();
+
+				const FName InputputName = *PropertyIn;
 				const FName OutputputName = *PropertyOut;
 
 				if (TSharedPtr<UE::Dataflow::FGraph> DataflowGraph = DataflowAsset.GetDataflow())
@@ -789,6 +789,13 @@ namespace UE::Dataflow
 					{
 						if (TSharedPtr<FDataflowNode> DataflowNodeTo = DataflowGraph->FindBaseNode(GuidIn))
 						{
+							FDataflowInput* InputConnection = DataflowNodeTo->FindInput(InputputName);
+							FDataflowOutput* OutputConnection = DataflowNodeFrom->FindOutput(OutputputName);
+
+							// make sure we set the right type before attempting any connection
+							DataflowNodeFrom->TrySetConnectionType(OutputConnection, FName(TypeOut));
+							DataflowNodeTo->TrySetConnectionType(InputConnection, FName(TypeIn));
+
 							// first connect the edgraph as this may affect the dataflow inputs ( for AnyType )
 							UEdGraphPin* OutputPin = FindPin(EdNodeOut, EEdGraphPinDirection::EGPD_Output, OutputputName);
 							UEdGraphPin* InputPin = FindPin(EdNodeIn, EEdGraphPinDirection::EGPD_Input, InputputName);
@@ -798,10 +805,7 @@ namespace UE::Dataflow
 							}
 
 							// now connect the dataflow
-							FDataflowInput* InputConnection = DataflowNodeTo->FindInput(InputputName);
-							FDataflowOutput* OutputConnection = DataflowNodeFrom->FindOutput(OutputputName);
 							DataflowGraph->Connect(OutputConnection, InputConnection);
-
 						}
 					}
 				}
diff --git a/Engine/Plugins/Experimental/Dataflow/Source/DataflowEditor/Public/Dataflow/AssetDefinition_DataflowAsset.h b/Engine/Plugins/Experimental/Dataflow/Source/DataflowEditor/Public/Dataflow/AssetDefinition_DataflowAsset.h
--- a/Engine/Plugins/Experimental/Dataflow/Source/DataflowEditor/Public/Dataflow/AssetDefinition_DataflowAsset.h
+++ b/Engine/Plugins/Experimental/Dataflow/Source/DataflowEditor/Public/Dataflow/AssetDefinition_DataflowAsset.h
@@ -22,6 +22,9 @@ namespace UE::DataflowAssetDefinitionHelpers
 	DATAFLOWEDITOR_API UObject* NewOrOpenDataflowAsset(const UObject* Asset);
 }
 
+struct FDataflowInput;
+struct FDataflowOutput;
+
 //
 // Copy/Paste structs
 //
@@ -75,8 +78,11 @@ struct DATAFLOWEDITOR_API FDataflowConnectionData
 	UPROPERTY()
 	FString In;
 
+	void Set(const FDataflowOutput& Output, const FDataflowInput& Input);
+
 	static FString GetNode(const FString InConnection);
 	static FString GetProperty(const FString InConnection);
+	static void GetNodePropertyAndType(const FString InConnection, FString& OutNode, FString& OutProperty, FString& OutType);
 };
 
 USTRUCT()
