ChaosFlesh: added runtime info in dataflow async simulation notification

diff --git a/Engine/Plugins/Experimental/Dataflow/Source/DataflowEditor/Private/Dataflow/DataflowSimulationGenerator.cpp b/Engine/Plugins/Experimental/Dataflow/Source/DataflowEditor/Private/Dataflow/DataflowSimulationGenerator.cpp
--- a/Engine/Plugins/Experimental/Dataflow/Source/DataflowEditor/Private/Dataflow/DataflowSimulationGenerator.cpp
+++ b/Engine/Plugins/Experimental/Dataflow/Source/DataflowEditor/Private/Dataflow/DataflowSimulationGenerator.cpp
@@ -94,6 +94,7 @@ bool FDataflowTaskManager::AllocateSimulationResource(const FVector2f& TimeRange
 
 	SimulationResource = MakeShared<FDataflowSimulationResource>();
 	SimulationResource->NumSimulatedFrames = &NumSimulatedFrames;
+	SimulationResource->RecentDateTimeTicks = &RecentDateTimeTicks;
 	SimulationResource->bCancelled = &bCancelled;
 	
 	NumFrames = FMath::Floor((TimeRange[1] - TimeRange[0] + UE_SMALL_NUMBER) * FrameRate);
@@ -116,6 +117,7 @@ bool FDataflowTaskManager::AllocateSimulationResource(const FVector2f& TimeRange
 			// Init simulation proxies from interface
 			DataflowManager->InitSimulationInterfaces();
 		}
+		RecentDateTimeTicks.store(FDateTime::UtcNow().GetTicks());
 		DataflowManager->SetSimulationEnabled(!SimulationTask->GetTask().bAsyncCaching);
 	}
 
@@ -202,6 +204,7 @@ void FDataflowSimulationGenerator::StartGenerateSimulation()
 	TaskManager->SimulationTask->GetTask().bAsyncCaching = CacheParams.bAsyncCaching;
 	TaskManager->SimulationTask->GetTask().bRestartSimulation = CacheParams.bRestartSimulation;
 	FVector2f TimeRange = CacheParams.bRestartSimulation ? CacheParams.RestartTimeRange : CacheParams.TimeRange;
+	TaskManager->StartTime = FDateTime::UtcNow();
 	TaskManager->AllocateSimulationResource(TimeRange, CacheParams.FrameRate, CacheAsset, BlueprintClass, DataflowContent, BlueprintTransform, bSkeletalMeshVisibility, GeneratorDeltaTime);
 
 	if(CacheParams.bAsyncCaching)
@@ -220,7 +223,6 @@ void FDataflowSimulationGenerator::StartGenerateSimulation()
 	NotificationConfig.bKeepOpenOnSuccess = true;
 	NotificationConfig.bKeepOpenOnFailure = true;
 	TaskManager->AsyncNotification = MakeUnique<FAsyncTaskNotification>(NotificationConfig);
-	TaskManager->StartTime = FDateTime::UtcNow();
 	TaskManager->LastUpdateTime = TaskManager->StartTime;
 
 	PendingAction = EDataflowGeneratorActions::TickGenerate;
@@ -253,9 +255,12 @@ void FDataflowSimulationGenerator::TickGenerateSimulation()
 		
 		const int32 NumSimulatedFrames = TaskManager->NumSimulatedFrames.load() + 1; //Counting the preroll
 		const int32 NumTotalFrames = TaskManager->NumFrames + 1;
-		const FText ProgressMessage = FText::FromString(FString::Printf(TEXT("Finished %d/%d, %.1f%%"),
-			NumSimulatedFrames, NumTotalFrames, 100.0 * NumSimulatedFrames / NumTotalFrames));
-		TaskManager->AsyncNotification->SetProgressText(ProgressMessage);
+		const FDateTime RecentFrameTime = FDateTime(TaskManager->RecentDateTimeTicks.load());
+		const double AverageFrameTime = (RecentFrameTime - TaskManager->StartTime).GetTotalSeconds() / NumSimulatedFrames;
+		const double EstimatedTime = FMath::Max(0, AverageFrameTime * (NumTotalFrames - NumSimulatedFrames) - (CurrentTime - RecentFrameTime).GetTotalSeconds());
+		FString ProgressMessage = FString::Printf(TEXT("Finished %d/%d, %.1f%%\nAverage time: %.1f seconds/frame\nEstimated finish time: %s"),
+			NumSimulatedFrames, NumTotalFrames, 100.0 * NumSimulatedFrames / NumTotalFrames, AverageFrameTime, *FText::AsTimespan(FTimespan::FromSeconds(EstimatedTime)).ToString());
+		TaskManager->AsyncNotification->SetProgressText(FText::FromString(ProgressMessage));
 		TaskManager->LastUpdateTime = CurrentTime;
 	}
 	else
diff --git a/Engine/Plugins/Experimental/Dataflow/Source/DataflowEditor/Public/Dataflow/DataflowSimulationGenerator.h b/Engine/Plugins/Experimental/Dataflow/Source/DataflowEditor/Public/Dataflow/DataflowSimulationGenerator.h
--- a/Engine/Plugins/Experimental/Dataflow/Source/DataflowEditor/Public/Dataflow/DataflowSimulationGenerator.h
+++ b/Engine/Plugins/Experimental/Dataflow/Source/DataflowEditor/Public/Dataflow/DataflowSimulationGenerator.h
@@ -66,6 +66,9 @@ namespace UE::Dataflow
 		/** Number of simulated frames */
 		std::atomic<int32>* NumSimulatedFrames = nullptr;
 
+		/** Simulation time of current frames */
+		std::atomic<int64>* RecentDateTimeTicks = nullptr;
+
 		/** Async cancel boolean */
 		std::atomic<bool>* bCancelled = nullptr;
 
@@ -82,6 +85,10 @@ namespace UE::Dataflow
 			{
 				++(*NumSimulatedFrames);
 			}
+			if (RecentDateTimeTicks)
+			{
+				RecentDateTimeTicks->store(FDateTime::UtcNow().GetTicks());
+			}
 		}
 	};
 
@@ -120,6 +127,9 @@ namespace UE::Dataflow
 		/** Number of simulated frames */
 		std::atomic<int32> NumSimulatedFrames = 0;
 		
+		/** Simulation time of current frames */
+		std::atomic<int64> RecentDateTimeTicks = 0;
+
 		/** Boolean to check if the simulation has been cancelled */
 		std::atomic<bool> bCancelled = false;
 		
