Dataflow : add Evaluation button to the UI
- Show evaluation status of the graph
- Allow to toggle between automatic and manual evaluation mode
- Allow to clear the data cache as well

diff --git a/Engine/Plugins/Experimental/Dataflow/Source/DataflowEditor/Private/Dataflow/DataflowEditorCommands.cpp b/Engine/Plugins/Experimental/Dataflow/Source/DataflowEditor/Private/Dataflow/DataflowEditorCommands.cpp
--- a/Engine/Plugins/Experimental/Dataflow/Source/DataflowEditor/Private/Dataflow/DataflowEditorCommands.cpp
+++ b/Engine/Plugins/Experimental/Dataflow/Source/DataflowEditor/Private/Dataflow/DataflowEditorCommands.cpp
@@ -58,6 +58,10 @@ void FDataflowEditorCommandsImpl::RegisterCommands()
 	TBaseCharacterFXEditorCommands::RegisterCommands();
 	
 	UI_COMMAND(EvaluateNode, "Evaluate", "Trigger an evaluation of the selected node.", EUserInterfaceActionType::Button, FInputChord());
+	UI_COMMAND(EvaluateGraph, "Evaluate Dataflow Graph", "Trigger an evaluation of the graph.", EUserInterfaceActionType::Button, FInputChord(EKeys::F5));
+	UI_COMMAND(EvaluateGraphAutomatic, "Automatic Graph Evaluation", "Set the evaluation of the graph to be automatic.", EUserInterfaceActionType::RadioButton, FInputChord());
+	UI_COMMAND(EvaluateGraphManual, "Manual Graph Evaluation", "Set the evaluation of the graph to be manual.", EUserInterfaceActionType::RadioButton, FInputChord());
+	UI_COMMAND(ClearGraphCache, "Clear Graph Cache", "Clear the graph ata cache.", EUserInterfaceActionType::Button, FInputChord());
 	UI_COMMAND(FreezeNodes, "FreezeNodes", "Freeze the evaluation of the selected nodes.", EUserInterfaceActionType::Button, FInputChord());
 	UI_COMMAND(UnfreezeNodes, "UnfreezeNodes", "Unfreeze the evaluation of the selected nodes.", EUserInterfaceActionType::Button, FInputChord());
 	UI_COMMAND(CreateComment, "CreateComment", "Create a Comment node.", EUserInterfaceActionType::None, FInputChord());
diff --git a/Engine/Plugins/Experimental/Dataflow/Source/DataflowEditor/Private/Dataflow/DataflowEditorToolkit.cpp b/Engine/Plugins/Experimental/Dataflow/Source/DataflowEditor/Private/Dataflow/DataflowEditorToolkit.cpp
--- a/Engine/Plugins/Experimental/Dataflow/Source/DataflowEditor/Private/Dataflow/DataflowEditorToolkit.cpp
+++ b/Engine/Plugins/Experimental/Dataflow/Source/DataflowEditor/Private/Dataflow/DataflowEditorToolkit.cpp
@@ -596,6 +596,8 @@ void FDataflowEditorToolkit::PostInitAssetEditor()
 	UToolMenu* const AssetToolbar = UToolMenus::Get()->ExtendMenu(ToolBarName);
 	FToolMenuSection& Section = AssetToolbar->FindOrAddSection("DataflowTools");
 
+	AddEValuationWidget(Section);
+
 	for (const TPair<FName, TSharedPtr<const FUICommandInfo>>& NodeAndAddCommand : DataflowMode->NodeTypeToAddNodeCommandMap)
 	{
 		ToolkitCommands->MapAction(NodeAndAddCommand.Value,
@@ -605,11 +607,132 @@ void FDataflowEditorToolkit::PostInitAssetEditor()
 		Section.AddEntry(FToolMenuEntry::InitToolBarButton(NodeAndAddCommand.Value));
 	}
 
-
 	// Force scenes to update loaded asset viewer settings
 	UAssetViewerSettings::Get()->OnAssetViewerSettingsChanged().Broadcast(NAME_None);
 }
 
+void FDataflowEditorToolkit::SetEvaluateGraphMode(FDataflowEditorToolkit::EEvaluationMode Mode)
+{
+	EvaluationMode = Mode;
+}
+
+void FDataflowEditorToolkit::ClearGraphCache()
+{
+	if (TObjectPtr<UDataflowBaseContent>& EditorContent = GetEditorContent())
+	{
+		if (TSharedPtr<UE::Dataflow::FEngineContext> DataflowContext = EditorContent->GetDataflowContext())
+		{
+			DataflowContext->ClearAllData();
+		}
+	}
+}
+
+bool FDataflowEditorToolkit::CanClearGraphCache() const
+{
+	if (const TObjectPtr<UDataflowBaseContent>& EditorContent = GetEditorContent())
+	{
+		if (TSharedPtr<UE::Dataflow::FEngineContext> DataflowContext = EditorContent->GetDataflowContext())
+		{
+			return !DataflowContext->IsEmpty();
+		}
+	}
+	return false;
+}
+
+FSlateIcon FDataflowEditorToolkit::GetEvaluationStatusImage() const
+{
+	static const FName CompileStatusBackground("Blueprint.CompileStatus.Background");
+	static const FName CompileStatusUnknown("Blueprint.CompileStatus.Overlay.Unknown");
+	static const FName CompileStatusGood("Blueprint.CompileStatus.Overlay.Good");
+	static const FName CompileStatusWarning("Blueprint.CompileStatus.Overlay.Warning");
+	static const FName CompileStatusError("Blueprint.CompileStatus.Overlay.Error");
+
+	bool bHasWarning = false;
+	bool bHasError = false;
+	if (const TObjectPtr<UDataflowBaseContent>& EditorContent = GetEditorContent())
+	{
+		if (TSharedPtr<UE::Dataflow::FEngineContext> DataflowContext = EditorContent->GetDataflowContext())
+		{
+			bHasWarning = (DataflowContext->GetNumWarnings() > 0);
+			bHasError = (DataflowContext->GetNumErrors() > 0);
+		}
+	}
+
+	FName OverlayIcon;
+	if (IsGraphDirty())
+	{
+		OverlayIcon = CompileStatusUnknown;
+	}
+	else
+	{
+		if (bHasError)
+		{
+			OverlayIcon = CompileStatusError;
+		}
+		else if (bHasWarning)
+		{
+			OverlayIcon = CompileStatusWarning;
+		}
+		else
+		{
+			OverlayIcon = CompileStatusGood;
+		}
+	}
+	return FSlateIcon(FAppStyle::GetAppStyleSetName(), CompileStatusBackground, NAME_None, OverlayIcon);
+}
+
+void FDataflowEditorToolkit::AddEValuationWidget(FToolMenuSection& InSection)
+{
+	ToolkitCommands->MapAction(FDataflowEditorCommands::Get().EvaluateGraph,
+		FExecuteAction::CreateSP(this, &FDataflowEditorToolkit::EvaluateGraph)
+	);
+	ToolkitCommands->MapAction(FDataflowEditorCommands::Get().EvaluateGraphAutomatic,
+		FExecuteAction::CreateSP(this, &FDataflowEditorToolkit::SetEvaluateGraphMode, EEvaluationMode::Automatic),
+		FCanExecuteAction(),
+		FIsActionChecked::CreateLambda([this]() { return EvaluationMode == EEvaluationMode::Automatic; })
+	);
+	ToolkitCommands->MapAction(FDataflowEditorCommands::Get().EvaluateGraphManual,
+		FExecuteAction::CreateSP(this, &FDataflowEditorToolkit::SetEvaluateGraphMode, EEvaluationMode::Manual),
+		FCanExecuteAction(),
+		FIsActionChecked::CreateLambda([this]() { return EvaluationMode == EEvaluationMode::Manual; })
+	);
+	ToolkitCommands->MapAction(FDataflowEditorCommands::Get().ClearGraphCache,
+		FExecuteAction::CreateSP(this, &FDataflowEditorToolkit::ClearGraphCache),
+		FCanExecuteAction::CreateSP(this, &FDataflowEditorToolkit::CanClearGraphCache)
+	);
+
+	InSection.AddEntry(FToolMenuEntry::InitToolBarButton(
+		FDataflowEditorCommands::Get().EvaluateGraph,
+		TAttribute<FText>(),
+		TAttribute<FText>(),
+		TAttribute<FSlateIcon>(this, &FDataflowEditorToolkit::GetEvaluationStatusImage)
+	));
+
+	FToolMenuEntry EvaluationOptions = FToolMenuEntry::InitComboButton(
+		"DataflowEvaluationOptions",
+		FUIAction(),
+		FOnGetContent::CreateSP(this, &FDataflowEditorToolkit::GenerateEvaluationOptionsMenu),
+		LOCTEXT("DataflowEvaluationOptions", ""),
+		LOCTEXT("DataflowEvaluationOptions_ToolbarTooltip", "Options to customize how Dataflow evaluate"),
+		TAttribute<FSlateIcon>(),
+		true
+	);
+	EvaluationOptions.StyleNameOverride = "SlimToolBar";
+	InSection.AddEntry(EvaluationOptions);
+}
+
+TSharedRef<SWidget> FDataflowEditorToolkit::GenerateEvaluationOptionsMenu()
+{
+	FMenuBuilder MenuBuilder(true, GetToolkitCommands());
+	MenuBuilder.BeginSection(TEXT("Section"));
+	MenuBuilder.AddMenuEntry(FDataflowEditorCommands::Get().EvaluateGraphAutomatic);
+	MenuBuilder.AddMenuEntry(FDataflowEditorCommands::Get().EvaluateGraphManual);
+	MenuBuilder.AddSeparator();
+	MenuBuilder.AddMenuEntry(FDataflowEditorCommands::Get().ClearGraphCache);
+	MenuBuilder.EndSection();
+	return MenuBuilder.MakeWidget();
+}
+
 void FDataflowEditorToolkit::InitializeEdMode(UBaseCharacterFXEditorMode* EdMode)
 {
 	UDataflowEditorMode* DataflowMode = Cast<UDataflowEditorMode>(EdMode);
@@ -1293,17 +1416,78 @@ void FDataflowEditorToolkit::OnNodeSelectionChanged(const TSet<UObject*>& InNewS
 	}
 }
 
-void FDataflowEditorToolkit::OnNodeInvalidated(UDataflow& DataflowAsset, FDataflowNode& Node)
+bool FDataflowEditorToolkit::IsGraphDirty() const
 {
+	using namespace UE::Dataflow;
+
 	if (const TObjectPtr<UDataflowBaseContent>& EditorContent = GetEditorContent())
 	{
-		if (EditorContent->GetDataflowAsset() == &DataflowAsset)
+		if (TSharedPtr<UE::Dataflow::FEngineContext> DataflowContext = EditorContent->GetDataflowContext())
 		{
-			// evaluate only active terminals nodes 
-			if (FDataflowTerminalNode* TerminalNode = Node.AsType<FDataflowTerminalNode>())
+			if (UDataflow* DataflowAsset = EditorContent->GetDataflowAsset())
 			{
-				EvaluateTerminalNode(*TerminalNode);
-				bViewsNeedRefresh = true;
+				if (TSharedPtr<UE::Dataflow::FGraph> DataflowGraph = DataflowAsset->GetDataflow())
+				{
+					const FTimestamp LastModifiedTimestamp = EditorContent->GetLastModifiedTimestamp();
+
+					for (TSharedPtr<FDataflowNode> Node : DataflowGraph->GetNodes())
+					{
+						if (Node)
+						{
+							if (FDataflowTerminalNode* TerminalNode = Node->AsType<FDataflowTerminalNode>())
+							{
+								if (LastModifiedTimestamp < TerminalNode->GetTimestamp())
+								{
+									return true;
+								}
+							}
+						}
+					}
+				}
+			}
+		}
+	}
+	return false;
+}
+
+void FDataflowEditorToolkit::EvaluateGraph()
+{
+	if (const TObjectPtr<UDataflowBaseContent>& EditorContent = GetEditorContent())
+	{
+		if (UDataflow* DataflowAsset = EditorContent->GetDataflowAsset())
+		{
+			if (TSharedPtr<UE::Dataflow::FGraph> DataflowGraph = DataflowAsset->GetDataflow())
+			{
+				for (TSharedPtr<FDataflowNode> Node : DataflowGraph->GetNodes())
+				{
+					if (Node)
+					{
+						if (FDataflowTerminalNode* TerminalNode = Node->AsType<FDataflowTerminalNode>())
+						{
+							EvaluateTerminalNode(*TerminalNode);
+							bViewsNeedRefresh = true;
+						}
+					}
+				}
+			}
+		}
+	}
+}
+
+void FDataflowEditorToolkit::OnNodeInvalidated(UDataflow& DataflowAsset, FDataflowNode& Node)
+{
+	if (EvaluationMode == EEvaluationMode::Automatic)
+	{
+		if (const TObjectPtr<UDataflowBaseContent>& EditorContent = GetEditorContent())
+		{
+			if (EditorContent->GetDataflowAsset() == &DataflowAsset)
+			{
+				// evaluate only active terminals nodes 
+				if (FDataflowTerminalNode* TerminalNode = Node.AsType<FDataflowTerminalNode>())
+				{
+					EvaluateTerminalNode(*TerminalNode);
+					bViewsNeedRefresh = true;
+				}
 			}
 		}
 	}
diff --git a/Engine/Plugins/Experimental/Dataflow/Source/DataflowEditor/Public/Dataflow/DataflowEditorCommands.h b/Engine/Plugins/Experimental/Dataflow/Source/DataflowEditor/Public/Dataflow/DataflowEditorCommands.h
--- a/Engine/Plugins/Experimental/Dataflow/Source/DataflowEditor/Public/Dataflow/DataflowEditorCommands.h
+++ b/Engine/Plugins/Experimental/Dataflow/Source/DataflowEditor/Public/Dataflow/DataflowEditorCommands.h
@@ -56,6 +56,10 @@ public:
 	 static void UpdateToolCommandBinding(UInteractiveTool* Tool, TSharedPtr<FUICommandList> UICommandList, bool bUnbind = false);
 	
 	TSharedPtr<FUICommandInfo> EvaluateNode;
+	TSharedPtr<FUICommandInfo> EvaluateGraph;
+	TSharedPtr<FUICommandInfo> EvaluateGraphAutomatic;
+	TSharedPtr<FUICommandInfo> EvaluateGraphManual;
+	TSharedPtr<FUICommandInfo> ClearGraphCache;
 	TSharedPtr<FUICommandInfo> FreezeNodes;
 	TSharedPtr<FUICommandInfo> UnfreezeNodes;
 	TSharedPtr<FUICommandInfo> CreateComment;
diff --git a/Engine/Plugins/Experimental/Dataflow/Source/DataflowEditor/Public/Dataflow/DataflowEditorToolkit.h b/Engine/Plugins/Experimental/Dataflow/Source/DataflowEditor/Public/Dataflow/DataflowEditorToolkit.h
--- a/Engine/Plugins/Experimental/Dataflow/Source/DataflowEditor/Public/Dataflow/DataflowEditorToolkit.h
+++ b/Engine/Plugins/Experimental/Dataflow/Source/DataflowEditor/Public/Dataflow/DataflowEditorToolkit.h
@@ -45,6 +45,12 @@ class DATAFLOWEDITOR_API FDataflowEditorToolkit final : public FBaseCharacterFXE
 {
 	using FBaseCharacterFXEditorToolkit::ObjectScene;
 
+	enum class EEvaluationMode: uint8
+	{
+		Automatic,
+		Manual,
+	};
+
 public:
 
 	explicit FDataflowEditorToolkit(UAssetEditor* InOwningAssetEditor);
@@ -124,6 +130,7 @@ protected:
 	// Node evaluation
 	void EvaluateTerminalNode(const FDataflowTerminalNode& TerminalNode);
 	void EvaluateNode(const FDataflowNode* Node, const FDataflowOutput* Output, UE::Dataflow::FTimestamp& InOutTimestamp);
+	void EvaluateGraph();
 	void RefreshViewsIfNeeded(bool bForce = false);
 	void OnNodeBeginEvaluate(const FDataflowNode* Node, const FDataflowOutput* Output);
 	void OnNodeFinishEvaluate(const FDataflowNode* Node, const FDataflowOutput* Output);
@@ -227,6 +234,15 @@ private:
 	TSharedPtr<SDataflowMembersWidget> CreateDataflowMembersWidget();
 	TSharedRef<SGraphEditor> CreateSubGraphEditorWidget(TSharedRef<FTabInfo> InTabInfo, UDataflowSubGraph* InGraph);
 
+	void AddEValuationWidget(FToolMenuSection& Section);
+	TSharedRef<SWidget> GenerateEvaluationOptionsMenu();
+	FSlateIcon GetEvaluationStatusImage() const;
+	bool IsGraphDirty() const;
+
+	void SetEvaluateGraphMode(EEvaluationMode Mode);
+	void ClearGraphCache();
+	bool CanClearGraphCache() const;
+
 	/** Create the simulation viewport client */
 	void CreateSimulationViewportClient();
 
@@ -302,6 +318,8 @@ private:
 	/** Cached value of the p.Dataflow.EnableGraphEval cvar, to avoid calling FindConsoleVariable too often */
 	bool bDataflowEnableGraphEval;
 
+	EEvaluationMode EvaluationMode = EEvaluationMode::Automatic;
+
 	/** Delegate for updating the cached value of p.Dataflow.EnableGraphEval */
 	FDelegateHandle GraphEvalCVarChangedDelegateHandle;
 
