Dataflow : fix typo and adjust some tooltips

diff --git a/Engine/Plugins/Experimental/Dataflow/Source/DataflowEditor/Private/Dataflow/DataflowEditorCommands.cpp b/Engine/Plugins/Experimental/Dataflow/Source/DataflowEditor/Private/Dataflow/DataflowEditorCommands.cpp
--- a/Engine/Plugins/Experimental/Dataflow/Source/DataflowEditor/Private/Dataflow/DataflowEditorCommands.cpp
+++ b/Engine/Plugins/Experimental/Dataflow/Source/DataflowEditor/Private/Dataflow/DataflowEditorCommands.cpp
@@ -59,9 +59,9 @@ void FDataflowEditorCommandsImpl::RegisterCommands()
 	
 	UI_COMMAND(EvaluateNode, "Evaluate", "Trigger an evaluation of the selected node.", EUserInterfaceActionType::Button, FInputChord());
 	UI_COMMAND(EvaluateGraph, "Evaluate Dataflow Graph", "Trigger an evaluation of the graph.", EUserInterfaceActionType::Button, FInputChord(EKeys::F5));
-	UI_COMMAND(EvaluateGraphAutomatic, "Automatic Graph Evaluation", "Set the evaluation of the graph to be automatic.", EUserInterfaceActionType::RadioButton, FInputChord());
-	UI_COMMAND(EvaluateGraphManual, "Manual Graph Evaluation", "Set the evaluation of the graph to be manual.", EUserInterfaceActionType::RadioButton, FInputChord());
-	UI_COMMAND(ClearGraphCache, "Clear Graph Cache", "Clear the graph ata cache.", EUserInterfaceActionType::Button, FInputChord());
+	UI_COMMAND(EvaluateGraphAutomatic, "Automatic Graph Evaluation", "Set the evaluation mode of the graph to Automatic.", EUserInterfaceActionType::RadioButton, FInputChord());
+	UI_COMMAND(EvaluateGraphManual, "Manual Graph Evaluation", "Set the evaluation mode of the graph to Manual.", EUserInterfaceActionType::RadioButton, FInputChord());
+	UI_COMMAND(ClearGraphCache, "Clear Graph Cache", "Clear the graph Data cache.", EUserInterfaceActionType::Button, FInputChord());
 	UI_COMMAND(FreezeNodes, "FreezeNodes", "Freeze the evaluation of the selected nodes.", EUserInterfaceActionType::Button, FInputChord());
 	UI_COMMAND(UnfreezeNodes, "UnfreezeNodes", "Unfreeze the evaluation of the selected nodes.", EUserInterfaceActionType::Button, FInputChord());
 	UI_COMMAND(CreateComment, "CreateComment", "Create a Comment node.", EUserInterfaceActionType::None, FInputChord());
