Dataflow updates to support debug draw only nodes. Without this change, nodes with DebugDraw but no DATAFLOW_NODE_RENDER_TYPE could not be pinned or report whether or not they support a view mode.

* Make FDataflowNode::DebugDraw editor-only.
* Add new virtual method FDataflowNode::CanDebugDraw (used to determine if node can be pinned)
* Add new virtual method FDataflowNode::CanDebugDrawViewMode (used to determine which view modes the node supports and gates whether or not DebugDraw is called for a node)

Updated all existing nodes that implemented DebugDraw, including adding a DataflowEditor private dependency in editor builds.

diff --git a/Engine/Plugins/Experimental/Dataflow/Source/DataflowEditor/Private/Dataflow/DataflowEditorToolkit.cpp b/Engine/Plugins/Experimental/Dataflow/Source/DataflowEditor/Private/Dataflow/DataflowEditorToolkit.cpp
--- a/Engine/Plugins/Experimental/Dataflow/Source/DataflowEditor/Private/Dataflow/DataflowEditorToolkit.cpp
+++ b/Engine/Plugins/Experimental/Dataflow/Source/DataflowEditor/Private/Dataflow/DataflowEditorToolkit.cpp
@@ -1004,16 +1004,30 @@ void FDataflowEditorToolkit::UpdateDebugDraw()
 	{
 		if (const TSharedPtr<UE::Dataflow::FEngineContext> DataflowContext = EditorContent->GetDataflowContext())
 		{
+			FName CurrentViewMode;
+			if (UDataflowEditorMode* const DataflowMode = Cast<UDataflowEditorMode>(EditorModeManager->GetActiveScriptableMode(UDataflowEditorMode::EM_DataflowEditorModeId)))
+			{
+				if (const UE::Dataflow::IDataflowConstructionViewMode* const ConstructionViewMode = DataflowMode->GetConstructionViewMode())
+				{
+					CurrentViewMode = ConstructionViewMode->GetName();
+				}
+			}
+
 			if (const UDataflowEdNode* const SelectedNode = EditorContent->GetSelectedNode())
 			{
 				if (const TSharedPtr<const FDataflowNode> SelectedDataflowNode = SelectedNode->GetDataflowNode())
 				{
-					DebugDrawObject->ResetAllState();
+					if (SelectedDataflowNode->CanDebugDrawViewMode(CurrentViewMode))
+					{
+						DebugDrawObject->ResetAllState();
 
-					FDataflowNode::FDebugDrawParameters DebugDrawParameters;
-					DebugDrawParameters.bNodeIsSelected = true;
-					DebugDrawParameters.bNodeIsPinned = false;
-					SelectedDataflowNode->DebugDraw(*DataflowContext, *DebugDrawObject, DebugDrawParameters);
+						FDataflowNode::FDebugDrawParameters DebugDrawParameters;
+						DebugDrawParameters.bNodeIsSelected = true;
+						DebugDrawParameters.bNodeIsPinned = false;
+						DebugDrawParameters.CurrentViewMode = CurrentViewMode;
+
+						SelectedDataflowNode->DebugDraw(*DataflowContext, *DebugDrawObject, DebugDrawParameters);
+					}
 				}
 			}
 
@@ -1021,12 +1035,16 @@ void FDataflowEditorToolkit::UpdateDebugDraw()
 			{
 				if (const TSharedPtr<const FDataflowNode> PinnedDataflowNode = PinnedNode->GetDataflowNode())
 				{
-					DebugDrawObject->ResetAllState();
+					if (PinnedDataflowNode->CanDebugDrawViewMode(CurrentViewMode))
+					{
+						DebugDrawObject->ResetAllState();
 
-					FDataflowNode::FDebugDrawParameters DebugDrawParameters;
-					DebugDrawParameters.bNodeIsSelected = false;
-					DebugDrawParameters.bNodeIsPinned = true;
-					PinnedDataflowNode->DebugDraw(*DataflowContext, *DebugDrawObject, DebugDrawParameters);
+						FDataflowNode::FDebugDrawParameters DebugDrawParameters;
+						DebugDrawParameters.bNodeIsSelected = false;
+						DebugDrawParameters.bNodeIsPinned = true;
+						DebugDrawParameters.CurrentViewMode = CurrentViewMode;
+						PinnedDataflowNode->DebugDraw(*DataflowContext, *DebugDrawObject, DebugDrawParameters);
+					}
 				}
 			}
 		}
@@ -1372,23 +1390,17 @@ void FDataflowEditorToolkit::OnNodeSelectionChanged(const TSet<UObject*>& InNewS
 			// Get the first view mode that this node can use
 			auto FindAnyValidViewMode = [&](FName& OutViewModeName) -> bool
 			{
-				const TArray<UE::Dataflow::FRenderingParameter> RenderingParameters = PrimarySelection->GetRenderParameters();
-				for (const UE::Dataflow::FRenderingParameter& Param : RenderingParameters)
+				for (const TPair<FName, TUniquePtr<UE::Dataflow::IDataflowConstructionViewMode>>& ViewMode : UE::Dataflow::FRenderingViewModeFactory::GetInstance().GetViewModes())
 				{
-					const FName NodeOutputTypeName = Param.Type;
-
-					for (const TPair<FName, TUniquePtr<UE::Dataflow::IDataflowConstructionViewMode>>& ViewMode : UE::Dataflow::FRenderingViewModeFactory::GetInstance().GetViewModes())
+					check(ViewMode.Value.IsValid());
+					const bool bCanRender = UE::Dataflow::CanRenderNodeOutput(*PrimarySelection, *GetEditorContent(), *ViewMode.Value);
+					if (bCanRender)
 					{
-						check(ViewMode.Value.IsValid());
-						const bool bCanRender = UE::Dataflow::CanRenderNodeOutput(*PrimarySelection, *GetEditorContent(), *ViewMode.Value);
-
-						if (bCanRender)
-						{
-							OutViewModeName = ViewMode.Key;
-							return true;
-						}
+						OutViewModeName = ViewMode.Key;
+						return true;
 					}
 				}
+			
 				return false;
 			};
 
diff --git a/Engine/Plugins/Experimental/Dataflow/Source/DataflowEditor/Private/Dataflow/DataflowEditorUtil.cpp b/Engine/Plugins/Experimental/Dataflow/Source/DataflowEditor/Private/Dataflow/DataflowEditorUtil.cpp
--- a/Engine/Plugins/Experimental/Dataflow/Source/DataflowEditor/Private/Dataflow/DataflowEditorUtil.cpp
+++ b/Engine/Plugins/Experimental/Dataflow/Source/DataflowEditor/Private/Dataflow/DataflowEditorUtil.cpp
@@ -10,6 +10,7 @@
 #include "Dataflow/DataflowInstance.h"
 #include "Dataflow/DataflowObject.h"
 #include "Dataflow/DataflowRenderingFactory.h"
+#include "Dataflow/DataflowRenderingViewMode.h"
 #include "DynamicMesh/MeshNormals.h"
 #include "Engine/SkeletalMesh.h"
 #include "Animation/Skeleton.h"
@@ -130,6 +131,10 @@ namespace UE::Dataflow
 						}
 					}
 				}
+				if (NodeTarget->CanDebugDrawViewMode(ViewMode.GetName()))
+				{
+					return true;
+				}
 			}
 		}
 
diff --git a/Engine/Plugins/Experimental/Dataflow/Source/DataflowEditor/Private/Dataflow/DataflowSNode.cpp b/Engine/Plugins/Experimental/Dataflow/Source/DataflowEditor/Private/Dataflow/DataflowSNode.cpp
--- a/Engine/Plugins/Experimental/Dataflow/Source/DataflowEditor/Private/Dataflow/DataflowSNode.cpp
+++ b/Engine/Plugins/Experimental/Dataflow/Source/DataflowEditor/Private/Dataflow/DataflowSNode.cpp
@@ -213,7 +213,7 @@ TSharedRef<SWidget> SDataflowEdNode::CreateTitleRightWidget()
 	const bool bHasRenderingCheckBox = (
 		DataflowInterface->NodesHaveToggleWidget() && 
 		DataflowGraphNode->GetDataflowNode() &&
-		DataflowGraphNode->GetDataflowNode()->GetRenderParameters().Num()
+		(DataflowGraphNode->GetDataflowNode()->GetRenderParameters().Num() || DataflowGraphNode->GetDataflowNode()->CanDebugDraw())
 		);
 
 	if (bHasRenderingCheckBox)
diff --git a/Engine/Plugins/Experimental/GeometryCollectionPlugin/Source/GeometryCollectionNodes/GeometryCollectionNodes.Build.cs b/Engine/Plugins/Experimental/GeometryCollectionPlugin/Source/GeometryCollectionNodes/GeometryCollectionNodes.Build.cs
--- a/Engine/Plugins/Experimental/GeometryCollectionPlugin/Source/GeometryCollectionNodes/GeometryCollectionNodes.Build.cs
+++ b/Engine/Plugins/Experimental/GeometryCollectionPlugin/Source/GeometryCollectionNodes/GeometryCollectionNodes.Build.cs
@@ -38,6 +38,16 @@ namespace UnrealBuildTool.Rules
 					"FractureEngine"
 				}
 			);
+
+			if (Target.bCompileAgainstEditor)
+			{
+				PrivateDependencyModuleNames.AddRange(
+					new string[]
+					{
+					"DataflowEditor",
+					}
+					);
+			}
 		}
 	}
 }
diff --git a/Engine/Plugins/Experimental/GeometryCollectionPlugin/Source/GeometryCollectionNodes/Private/Dataflow/GeometryCollectionFracturingNodes.cpp b/Engine/Plugins/Experimental/GeometryCollectionPlugin/Source/GeometryCollectionNodes/Private/Dataflow/GeometryCollectionFracturingNodes.cpp
--- a/Engine/Plugins/Experimental/GeometryCollectionPlugin/Source/GeometryCollectionNodes/Private/Dataflow/GeometryCollectionFracturingNodes.cpp
+++ b/Engine/Plugins/Experimental/GeometryCollectionPlugin/Source/GeometryCollectionNodes/Private/Dataflow/GeometryCollectionFracturingNodes.cpp
@@ -2,6 +2,9 @@
 
 #include "Dataflow/GeometryCollectionFracturingNodes.h"
 #include "Dataflow/DataflowCore.h"
+#if WITH_EDITOR
+#include "Dataflow/DataflowRenderingViewMode.h"
+#endif
 
 #include "Engine/StaticMesh.h"
 #include "GeometryCollection/GeometryCollectionObject.h"
@@ -605,6 +608,12 @@ void FPlaneCutterDataflowNode_v2::Evaluate(UE::Dataflow::FContext& Context, cons
 	}
 }
 
+#if WITH_EDITOR
+bool FPlaneCutterDataflowNode_v2::CanDebugDrawViewMode(const FName& ViewModeName) const
+{
+	return ViewModeName == UE::Dataflow::FDataflowConstruction3DViewMode::Name;
+}
+
 void FPlaneCutterDataflowNode_v2::DebugDraw(UE::Dataflow::FContext& Context, IDataflowDebugDrawInterface& DataflowRenderingInterface, const FDebugDrawParameters& DebugDrawParameters) const
 {
 	if ((DebugDrawParameters.bNodeIsSelected || DebugDrawParameters.bNodeIsPinned))
@@ -702,6 +711,7 @@ void FPlaneCutterDataflowNode_v2::DebugDraw(UE::Dataflow::FContext& Context, IDa
 		}
 	}
 }
+#endif
 
 /*--------------------------------------------------------------------------------------------------------*/
 
diff --git a/Engine/Plugins/Experimental/GeometryCollectionPlugin/Source/GeometryCollectionNodes/Private/Dataflow/GeometryCollectionMakeNodes.cpp b/Engine/Plugins/Experimental/GeometryCollectionPlugin/Source/GeometryCollectionNodes/Private/Dataflow/GeometryCollectionMakeNodes.cpp
--- a/Engine/Plugins/Experimental/GeometryCollectionPlugin/Source/GeometryCollectionNodes/Private/Dataflow/GeometryCollectionMakeNodes.cpp
+++ b/Engine/Plugins/Experimental/GeometryCollectionPlugin/Source/GeometryCollectionNodes/Private/Dataflow/GeometryCollectionMakeNodes.cpp
@@ -2,6 +2,9 @@
 
 #include "Dataflow/GeometryCollectionMakeNodes.h"
 #include "Dataflow/DataflowCore.h"
+#if WITH_EDITOR
+#include "Dataflow/DataflowRenderingViewMode.h"
+#endif
 
 #include "Engine/Engine.h"
 #include "Engine/StaticMesh.h"
@@ -106,6 +109,12 @@ void FMakeLiteralStringDataflowNode_v2::Evaluate(UE::Dataflow::FContext& Context
 
 /*--------------------------------------------------------------------------------------------------------*/
 
+#if WITH_EDITOR
+bool FMakePointsDataflowNode::CanDebugDrawViewMode(const FName& ViewModeName) const
+{
+	return ViewModeName == UE::Dataflow::FDataflowConstruction3DViewMode::Name;
+}
+
 void FMakePointsDataflowNode::DebugDraw(UE::Dataflow::FContext& Context, IDataflowDebugDrawInterface& DataflowRenderingInterface, const FDebugDrawParameters& DebugDrawParameters) const
 {
 	if ((DebugDrawParameters.bNodeIsSelected || DebugDrawParameters.bNodeIsPinned))
@@ -128,6 +137,7 @@ void FMakePointsDataflowNode::DebugDraw(UE::Dataflow::FContext& Context, IDatafl
 		DataflowRenderingInterface.DrawBox(0.5 * FVector(Max - Min), FQuat::Identity, 0.5 * FVector(Min + Max), 1.0);
 	}
 }
+#endif
 
 void FMakePointsDataflowNode::Evaluate(UE::Dataflow::FContext& Context, const FDataflowOutput* Out) const
 {
@@ -173,6 +183,12 @@ void FMakeSphereDataflowNode::Evaluate(UE::Dataflow::FContext& Context, const FD
 	}
 }
 
+#if WITH_EDITOR
+bool FMakeSphereDataflowNode::CanDebugDrawViewMode(const FName& ViewModeName) const
+{
+	return ViewModeName == UE::Dataflow::FDataflowConstruction3DViewMode::Name;
+}
+
 void FMakeSphereDataflowNode::DebugDraw(UE::Dataflow::FContext& Context, IDataflowDebugDrawInterface& DataflowRenderingInterface, const FDebugDrawParameters& DebugDrawParameters) const
 {
 	if ((DebugDrawParameters.bNodeIsSelected || DebugDrawParameters.bNodeIsPinned))
@@ -188,6 +204,7 @@ void FMakeSphereDataflowNode::DebugDraw(UE::Dataflow::FContext& Context, IDatafl
 		DataflowRenderingInterface.DrawBox(0.5 * FVector(Max - Min), FQuat::Identity, 0.5 * FVector(Min + Max), 1.0);
 	}
 }
+#endif
 
 /*--------------------------------------------------------------------------------------------------------*/
 
@@ -562,6 +579,12 @@ void FMakePlaneDataflowNode::Evaluate(UE::Dataflow::FContext& Context, const FDa
 	}
 }
 
+#if WITH_EDITOR
+bool FMakePlaneDataflowNode::CanDebugDrawViewMode(const FName& ViewModeName) const
+{
+	return ViewModeName == UE::Dataflow::FDataflowConstruction3DViewMode::Name;
+}
+
 void FMakePlaneDataflowNode::DebugDraw(UE::Dataflow::FContext& Context, IDataflowDebugDrawInterface& DataflowRenderingInterface, const FDebugDrawParameters& DebugDrawParameters) const
 {
 	if ((DebugDrawParameters.bNodeIsSelected || DebugDrawParameters.bNodeIsPinned))
@@ -593,6 +616,7 @@ void FMakePlaneDataflowNode::DebugDraw(UE::Dataflow::FContext& Context, IDataflo
 		DataflowRenderingInterface.DrawLine(InBasePoint, InBasePoint + InNormal * 2.f);
 	}
 }
+#endif
 
 /* -------------------------------------------------------------------------------- */
 
diff --git a/Engine/Plugins/Experimental/GeometryCollectionPlugin/Source/GeometryCollectionNodes/Private/Dataflow/GeometryCollectionUtilityNodes.cpp b/Engine/Plugins/Experimental/GeometryCollectionPlugin/Source/GeometryCollectionNodes/Private/Dataflow/GeometryCollectionUtilityNodes.cpp
--- a/Engine/Plugins/Experimental/GeometryCollectionPlugin/Source/GeometryCollectionNodes/Private/Dataflow/GeometryCollectionUtilityNodes.cpp
+++ b/Engine/Plugins/Experimental/GeometryCollectionPlugin/Source/GeometryCollectionNodes/Private/Dataflow/GeometryCollectionUtilityNodes.cpp
@@ -2,6 +2,9 @@
 
 #include "Dataflow/GeometryCollectionUtilityNodes.h"
 #include "Dataflow/DataflowCore.h"
+#if WITH_EDITOR
+#include "Dataflow/DataflowRenderingViewMode.h"
+#endif
 
 #include "GeometryCollection/ManagedArrayCollection.h"
 #include "GeometryCollection/GeometryCollection.h"
@@ -158,6 +161,12 @@ void FCreateLeafConvexHullsDataflowNode::Evaluate(UE::Dataflow::FContext& Contex
 	}
 }
 
+#if WITH_EDITOR
+bool FCreateLeafConvexHullsDataflowNode::CanDebugDrawViewMode(const FName& ViewModeName) const
+{
+	return ViewModeName == UE::Dataflow::FDataflowConstruction3DViewMode::Name;
+}
+
 void FCreateLeafConvexHullsDataflowNode::DebugDraw(UE::Dataflow::FContext& Context, IDataflowDebugDrawInterface& DataflowRenderingInterface, const FDebugDrawParameters& DebugDrawParameters) const
 {
 	using namespace UE::Geometry;
@@ -187,6 +196,7 @@ void FCreateLeafConvexHullsDataflowNode::DebugDraw(UE::Dataflow::FContext& Conte
 		}
 	}
 }
+#endif
 
 /* --------------------------------------------------------------------------------------------------------------------------- */
 
@@ -383,6 +393,12 @@ void FGenerateClusterConvexHullsFromLeafHullsDataflowNode::Evaluate(UE::Dataflow
 	}
 }
 
+#if WITH_EDITOR
+bool FGenerateClusterConvexHullsFromLeafHullsDataflowNode::CanDebugDrawViewMode(const FName& ViewModeName) const
+{
+	return ViewModeName == UE::Dataflow::FDataflowConstruction3DViewMode::Name;
+}
+
 void FGenerateClusterConvexHullsFromLeafHullsDataflowNode::DebugDraw(UE::Dataflow::FContext& Context, IDataflowDebugDrawInterface& DataflowRenderingInterface, const FDebugDrawParameters& DebugDrawParameters) const
 {
 	using namespace UE::Geometry;
@@ -412,6 +428,7 @@ void FGenerateClusterConvexHullsFromLeafHullsDataflowNode::DebugDraw(UE::Dataflo
 		}
 	}
 }
+#endif
 
 /* --------------------------------------------------------------------------------------------------------------------------- */
 
@@ -511,6 +528,12 @@ void FGenerateClusterConvexHullsFromChildrenHullsDataflowNode::Evaluate(UE::Data
 	}
 }
 
+#if WITH_EDITOR
+bool FGenerateClusterConvexHullsFromChildrenHullsDataflowNode::CanDebugDrawViewMode(const FName& ViewModeName) const
+{
+	return ViewModeName == UE::Dataflow::FDataflowConstruction3DViewMode::Name;
+}
+
 void FGenerateClusterConvexHullsFromChildrenHullsDataflowNode::DebugDraw(UE::Dataflow::FContext& Context, IDataflowDebugDrawInterface& DataflowRenderingInterface, const FDebugDrawParameters& DebugDrawParameters) const
 {
 	using namespace UE::Geometry;
@@ -540,6 +563,7 @@ void FGenerateClusterConvexHullsFromChildrenHullsDataflowNode::DebugDraw(UE::Dat
 		}
 	}
 }
+#endif
 
 /* --------------------------------------------------------------------------------------------------------------------------- */
 
diff --git a/Engine/Plugins/Experimental/GeometryCollectionPlugin/Source/GeometryCollectionNodes/Public/Dataflow/GeometryCollectionFracturingNodes.h b/Engine/Plugins/Experimental/GeometryCollectionPlugin/Source/GeometryCollectionNodes/Public/Dataflow/GeometryCollectionFracturingNodes.h
--- a/Engine/Plugins/Experimental/GeometryCollectionPlugin/Source/GeometryCollectionNodes/Public/Dataflow/GeometryCollectionFracturingNodes.h
+++ b/Engine/Plugins/Experimental/GeometryCollectionPlugin/Source/GeometryCollectionNodes/Public/Dataflow/GeometryCollectionFracturingNodes.h
@@ -727,7 +727,14 @@ public:
 
 	virtual void Evaluate(UE::Dataflow::FContext& Context, const FDataflowOutput* Out) const override;
 
+#if WITH_EDITOR
+	virtual bool CanDebugDraw() const override
+	{
+		return true;
+	}
+	virtual bool CanDebugDrawViewMode(const FName& ViewModeName) const override;
 	virtual void DebugDraw(UE::Dataflow::FContext& Context, IDataflowDebugDrawInterface& DataflowRenderingInterface, const FDebugDrawParameters& DebugDrawParameters) const override;
+#endif
 };
 
 /**
diff --git a/Engine/Plugins/Experimental/GeometryCollectionPlugin/Source/GeometryCollectionNodes/Public/Dataflow/GeometryCollectionMakeNodes.h b/Engine/Plugins/Experimental/GeometryCollectionPlugin/Source/GeometryCollectionNodes/Public/Dataflow/GeometryCollectionMakeNodes.h
--- a/Engine/Plugins/Experimental/GeometryCollectionPlugin/Source/GeometryCollectionNodes/Public/Dataflow/GeometryCollectionMakeNodes.h
+++ b/Engine/Plugins/Experimental/GeometryCollectionPlugin/Source/GeometryCollectionNodes/Public/Dataflow/GeometryCollectionMakeNodes.h
@@ -95,7 +95,14 @@ public:
 
 	virtual void Evaluate(UE::Dataflow::FContext& Context, const FDataflowOutput* Out) const override;
 
+#if WITH_EDITOR
+	virtual bool CanDebugDraw() const override
+	{
+		return true;
+	}
+	virtual bool CanDebugDrawViewMode(const FName& ViewModeName) const override;
 	virtual void DebugDraw(UE::Dataflow::FContext& Context, IDataflowDebugDrawInterface& DataflowRenderingInterface, const FDebugDrawParameters& DebugDrawParameters) const override;
+#endif
 };
 
 
@@ -184,7 +191,14 @@ public:
 
 	virtual void Evaluate(UE::Dataflow::FContext& Context, const FDataflowOutput* Out) const override;
 
+#if WITH_EDITOR
+	virtual bool CanDebugDraw() const override
+	{
+		return true;
+	}
+	virtual bool CanDebugDrawViewMode(const FName& ViewModeName) const override;
 	virtual void DebugDraw(UE::Dataflow::FContext& Context, IDataflowDebugDrawInterface& DataflowRenderingInterface, const FDebugDrawParameters& DebugDrawParameters) const override;
+#endif
 };
 
 /**
@@ -848,7 +862,14 @@ public:
 
 private:
 	virtual void Evaluate(UE::Dataflow::FContext& Context, const FDataflowOutput* Out) const override;
+#if WITH_EDITOR
+	virtual bool CanDebugDraw() const override
+	{
+		return true;
+	}
+	virtual bool CanDebugDrawViewMode(const FName& ViewModeName) const override;
 	virtual void DebugDraw(UE::Dataflow::FContext& Context, IDataflowDebugDrawInterface& DataflowRenderingInterface, const FDebugDrawParameters& DebugDrawParameters) const override;
+#endif
 };
 
 /**
diff --git a/Engine/Plugins/Experimental/GeometryCollectionPlugin/Source/GeometryCollectionNodes/Public/Dataflow/GeometryCollectionUtilityNodes.h b/Engine/Plugins/Experimental/GeometryCollectionPlugin/Source/GeometryCollectionNodes/Public/Dataflow/GeometryCollectionUtilityNodes.h
--- a/Engine/Plugins/Experimental/GeometryCollectionPlugin/Source/GeometryCollectionNodes/Public/Dataflow/GeometryCollectionUtilityNodes.h
+++ b/Engine/Plugins/Experimental/GeometryCollectionPlugin/Source/GeometryCollectionNodes/Public/Dataflow/GeometryCollectionUtilityNodes.h
@@ -192,7 +192,14 @@ private:
 	int32 ColorRandomSeed = 0;
 
 	virtual void Evaluate(UE::Dataflow::FContext& Context, const FDataflowOutput* Out) const override;
+#if WITH_EDITOR
+	virtual bool CanDebugDraw() const override
+	{
+		return true;
+	}
+	virtual bool CanDebugDrawViewMode(const FName& ViewModeName) const override;
 	virtual void DebugDraw(UE::Dataflow::FContext& Context, IDataflowDebugDrawInterface& DataflowRenderingInterface, const FDebugDrawParameters& DebugDrawParameters) const override;
+#endif
 
 public:
 	FCreateLeafConvexHullsDataflowNode(const UE::Dataflow::FNodeParameters& InParam, FGuid InGuid = FGuid::NewGuid());
@@ -402,7 +409,14 @@ private:
 	int32 ColorRandomSeed = 0;
 
 	virtual void Evaluate(UE::Dataflow::FContext& Context, const FDataflowOutput* Out) const override;
+#if WITH_EDITOR
+	virtual bool CanDebugDraw() const override
+	{
+		return true;
+	}
+	virtual bool CanDebugDrawViewMode(const FName& ViewModeName) const override;
 	virtual void DebugDraw(UE::Dataflow::FContext& Context, IDataflowDebugDrawInterface& DataflowRenderingInterface, const FDebugDrawParameters& DebugDrawParameters) const override;
+#endif
 
 public:
 	FGenerateClusterConvexHullsFromLeafHullsDataflowNode(const UE::Dataflow::FNodeParameters& InParam, FGuid InGuid = FGuid::NewGuid());
@@ -506,7 +520,14 @@ private:
 	int32 ColorRandomSeed = 0;
 
 	virtual void Evaluate(UE::Dataflow::FContext& Context, const FDataflowOutput* Out) const override;
+#if WITH_EDITOR
+	virtual bool CanDebugDraw() const override
+	{
+		return true;
+	}
+	virtual bool CanDebugDrawViewMode(const FName& ViewModeName) const override;
 	virtual void DebugDraw(UE::Dataflow::FContext& Context, IDataflowDebugDrawInterface& DataflowRenderingInterface, const FDebugDrawParameters& DebugDrawParameters) const override;
+#endif
 
 public:
 	FGenerateClusterConvexHullsFromChildrenHullsDataflowNode(const UE::Dataflow::FNodeParameters& InParam, FGuid InGuid = FGuid::NewGuid());
diff --git a/Engine/Plugins/Experimental/MeshResizing/Source/MeshResizingNodes/MeshResizingDataflowNodes.Build.cs b/Engine/Plugins/Experimental/MeshResizing/Source/MeshResizingNodes/MeshResizingDataflowNodes.Build.cs
--- a/Engine/Plugins/Experimental/MeshResizing/Source/MeshResizingNodes/MeshResizingDataflowNodes.Build.cs
+++ b/Engine/Plugins/Experimental/MeshResizing/Source/MeshResizingNodes/MeshResizingDataflowNodes.Build.cs
@@ -32,5 +32,15 @@ public class MeshResizingDataflowNodes : ModuleRules
 				"DynamicMesh"
 			}
 			);
+
+		if(Target.bCompileAgainstEditor)
+		{
+			PrivateDependencyModuleNames.AddRange(
+				new string[]
+				{
+					"DataflowEditor",
+				}
+				);
+		}
 	}
 }
diff --git a/Engine/Source/Runtime/Experimental/Dataflow/Core/Public/Dataflow/DataflowDebugDrawInterface.h b/Engine/Source/Runtime/Experimental/Dataflow/Core/Public/Dataflow/DataflowDebugDrawInterface.h
--- a/Engine/Source/Runtime/Experimental/Dataflow/Core/Public/Dataflow/DataflowDebugDrawInterface.h
+++ b/Engine/Source/Runtime/Experimental/Dataflow/Core/Public/Dataflow/DataflowDebugDrawInterface.h
@@ -23,6 +23,7 @@ public:
 	DATAFLOWCORE_API virtual void SetWorldPriority() = 0;
 	DATAFLOWCORE_API virtual void ResetAllState() = 0;
 
+	DATAFLOWCORE_API virtual void ReservePoints(int32 NumAdditionalPoints) = 0;
 	DATAFLOWCORE_API virtual void DrawPoint(const FVector& Position) = 0;
 	DATAFLOWCORE_API virtual void DrawLine(const FVector& Start, const FVector& End) const = 0;
 
diff --git a/Engine/Source/Runtime/Experimental/Dataflow/Core/Public/Dataflow/DataflowNode.h b/Engine/Source/Runtime/Experimental/Dataflow/Core/Public/Dataflow/DataflowNode.h
--- a/Engine/Source/Runtime/Experimental/Dataflow/Core/Public/Dataflow/DataflowNode.h
+++ b/Engine/Source/Runtime/Experimental/Dataflow/Core/Public/Dataflow/DataflowNode.h
@@ -24,6 +24,7 @@ namespace UE::Dataflow {
 	class FGraph;
 }
 class IDataflowDebugDrawInterface;
+class IDataflowConstructionViewMode;
 
 /** Empty class for the freeze UI customization. */
 USTRUCT()
@@ -111,15 +112,29 @@ PRAGMA_ENABLE_DEPRECATION_WARNINGS
 	UE_DEPRECATED(5.4, "FDataflowNode::IsExperimental is deprecated.")
 	virtual bool IsExperimental() { return false; }
 
+#if WITH_EDITOR
 	struct FDebugDrawParameters
 	{
 		bool bNodeIsSelected;
 		bool bNodeIsPinned;
+		FName CurrentViewMode;
 	};
 
+	/** Override if DebugDraw has been implemented.*/
+	DATAFLOWCORE_API virtual bool CanDebugDraw() const
+	{
+		return false;
+	}
+	/** Specify whether or not DebugDraw supports this ViewMode */
+	DATAFLOWCORE_API virtual bool CanDebugDrawViewMode(const FName& ViewModeName) const
+	{
+		return false;
+	}
+	/** Debug drawing callback. This will only be called for nodes which CanDebugDrawViewMode for the current view mode.*/
 	DATAFLOWCORE_API virtual void DebugDraw(UE::Dataflow::FContext& Context,
 		IDataflowDebugDrawInterface& DataflowRenderingInterface,
 		const FDebugDrawParameters& DebugDrawParameters) const {}
+#endif
 
 	//
 	// Connections
@@ -492,6 +507,15 @@ PRAGMA_ENABLE_DEPRECATION_WARNINGS
 		return GetValue(Context, UE::Dataflow::TConnectionReference<T>(Reference), Default);
 	}
 
+	template<class T> const T& GetOutputValue(UE::Dataflow::FContext& Context, const T* Reference, const T& Default) const
+	{
+		if (const FDataflowOutput* Output = FindOutput(Reference))
+		{
+			return Output->GetValue(Context, Default);
+		}
+		return Default;
+	}
+
 	/**
 	*   SetValue(...)
 	*
diff --git a/Engine/Source/Runtime/Experimental/Dataflow/Engine/Private/Dataflow/DataflowDebugDraw.cpp b/Engine/Source/Runtime/Experimental/Dataflow/Engine/Private/Dataflow/DataflowDebugDraw.cpp
--- a/Engine/Source/Runtime/Experimental/Dataflow/Engine/Private/Dataflow/DataflowDebugDraw.cpp
+++ b/Engine/Source/Runtime/Experimental/Dataflow/Engine/Private/Dataflow/DataflowDebugDraw.cpp
@@ -109,6 +109,11 @@ void FDataflowDebugDraw::ResetAllState()
 	ColorWithTranslucency = Color;
 }
 
+void FDataflowDebugDraw::ReservePoints(int32 NumAdditionalPoints)
+{
+	DebugRenderSceneProxy->ReservePoints(NumAdditionalPoints);
+}
+
 void FDataflowDebugDraw::DrawPoint(const FVector& Position)
 {
 	FDataflowDebugRenderSceneProxy::FDebugPoint NewPoint
diff --git a/Engine/Source/Runtime/Experimental/Dataflow/Engine/Private/Dataflow/DataflowDebugDrawComponent.cpp b/Engine/Source/Runtime/Experimental/Dataflow/Engine/Private/Dataflow/DataflowDebugDrawComponent.cpp
--- a/Engine/Source/Runtime/Experimental/Dataflow/Engine/Private/Dataflow/DataflowDebugDrawComponent.cpp
+++ b/Engine/Source/Runtime/Experimental/Dataflow/Engine/Private/Dataflow/DataflowDebugDrawComponent.cpp
@@ -124,6 +124,11 @@ FPrimitiveViewRelevance FDataflowDebugRenderSceneProxy::GetViewRelevance(const F
 	return Result;
 }
 
+void FDataflowDebugRenderSceneProxy::ReservePoints(int32 NumAdditionalPoints)
+{
+	Points.Reserve(Points.Num() + NumAdditionalPoints);
+}
+
 void FDataflowDebugRenderSceneProxy::AddPoint(const FDebugPoint& Point)
 {
 	Points.Add(Point);
diff --git a/Engine/Source/Runtime/Experimental/Dataflow/Engine/Public/Dataflow/DataflowDebugDraw.h b/Engine/Source/Runtime/Experimental/Dataflow/Engine/Public/Dataflow/DataflowDebugDraw.h
--- a/Engine/Source/Runtime/Experimental/Dataflow/Engine/Public/Dataflow/DataflowDebugDraw.h
+++ b/Engine/Source/Runtime/Experimental/Dataflow/Engine/Public/Dataflow/DataflowDebugDraw.h
@@ -33,6 +33,7 @@ private:
 	virtual void SetWorldPriority() override;
 	virtual void ResetAllState() override;
 
+	virtual void ReservePoints(int32 NumAdditionalPoints) override;
 	virtual void DrawPoint(const FVector& Position) override;
 	virtual void DrawLine(const FVector& Start, const FVector& End) const override;
 	virtual void DrawMesh(const IDebugDrawMesh& Mesh) const override;
diff --git a/Engine/Source/Runtime/Experimental/Dataflow/Engine/Public/Dataflow/DataflowDebugDrawComponent.h b/Engine/Source/Runtime/Experimental/Dataflow/Engine/Public/Dataflow/DataflowDebugDrawComponent.h
--- a/Engine/Source/Runtime/Experimental/Dataflow/Engine/Public/Dataflow/DataflowDebugDrawComponent.h
+++ b/Engine/Source/Runtime/Experimental/Dataflow/Engine/Public/Dataflow/DataflowDebugDrawComponent.h
@@ -32,6 +32,7 @@ public:
 		FLinearColor Color = FLinearColor::White;
 	};
 
+	void ReservePoints(int32 NumAdditionalPoints);
 	void AddPoint(const FDebugPoint& Point);
 
 private:
