Dataflow : evaluation button improvement
- Disable button in automatic mode
- When automatic mode is selected , icon is always green
- Save evaluation mode in the dataflow editor options when changed and reload when opening the editor

diff --git a/Engine/Plugins/Experimental/Dataflow/Source/DataflowEditor/Private/Dataflow/DataflowEditorOptions.h b/Engine/Plugins/Experimental/Dataflow/Source/DataflowEditor/Private/Dataflow/DataflowEditorOptions.h
--- a/Engine/Plugins/Experimental/Dataflow/Source/DataflowEditor/Private/Dataflow/DataflowEditorOptions.h
+++ b/Engine/Plugins/Experimental/Dataflow/Source/DataflowEditor/Private/Dataflow/DataflowEditorOptions.h
@@ -19,6 +19,15 @@ enum class EDataflowConstructionViewportMousePanButton : uint8
 	RightOrMiddle = 2,
 };
 
+UENUM(BlueprintType)
+enum class EDataflowEditorEvaluationMode : uint8
+{
+	/** Dataflow graph will evaluate automatically when values are changed */
+	Automatic = 0,
+
+	/** Dataflow graph will not eveluate until the user presses the evaluate button in the editor */
+	Manual = 1,
+};
 
 /** A settings class used to store and retreive user settings related to the Dataflow Editor */
 UCLASS(HideCategories = Object, Config = EditorPerProjectUserSettings, Meta = (DisplayName = "Dataflow Editor"), MinimalAPI)
@@ -54,6 +63,9 @@ class UDataflowEditorOptions : public UDeveloperSettings
 	UPROPERTY(Config, EditAnywhere, Category = UI)
 	EDataflowConstructionViewportMousePanButton ConstructionViewportMousePanButton;
 
+	UPROPERTY(Config, EditAnywhere, Category = UI)
+	EDataflowEditorEvaluationMode EditorEvaluationMode = EDataflowEditorEvaluationMode::Automatic;
+
 	// UDeveloperSettings overrides
 	virtual FName GetCategoryName() const override;
 
diff --git a/Engine/Plugins/Experimental/Dataflow/Source/DataflowEditor/Private/Dataflow/DataflowEditorToolkit.cpp b/Engine/Plugins/Experimental/Dataflow/Source/DataflowEditor/Private/Dataflow/DataflowEditorToolkit.cpp
--- a/Engine/Plugins/Experimental/Dataflow/Source/DataflowEditor/Private/Dataflow/DataflowEditorToolkit.cpp
+++ b/Engine/Plugins/Experimental/Dataflow/Source/DataflowEditor/Private/Dataflow/DataflowEditorToolkit.cpp
@@ -27,6 +27,7 @@
 #include "Dataflow/DataflowEditorUtil.h"
 #include "Dataflow/DataflowEditorSubGraphTabSummoner.h"
 #include "Dataflow/DataflowConstructionViewport.h"
+#include "Dataflow/DataflowEditorOptions.h"
 #include "Dataflow/DataflowEdNode.h"
 #include "Dataflow/DataflowGraphEditor.h"
 #include "Dataflow/DataflowNodeFactory.h"
@@ -105,6 +106,7 @@ FDataflowEditorToolkit::FDataflowEditorToolkit(UAssetEditor* InOwningAssetEditor
 	, DataflowEditor(Cast< UDataflowEditor>(InOwningAssetEditor))
 	, bDataflowEnableGraphEval(true)
 {
+	EvaluationMode = EDataflowEditorEvaluationMode::Automatic;
 
 	// When saving, only prompt to checkout and save assets that are actually modified
 	bCheckDirtyOnAssetSave = true;
@@ -596,7 +598,7 @@ void FDataflowEditorToolkit::PostInitAssetEditor()
 	UToolMenu* const AssetToolbar = UToolMenus::Get()->ExtendMenu(ToolBarName);
 	FToolMenuSection& Section = AssetToolbar->FindOrAddSection("DataflowTools");
 
-	AddEValuationWidget(Section);
+	AddEvaluationWidget(Section);
 
 	for (const TPair<FName, TSharedPtr<const FUICommandInfo>>& NodeAndAddCommand : DataflowMode->NodeTypeToAddNodeCommandMap)
 	{
@@ -611,9 +613,13 @@ void FDataflowEditorToolkit::PostInitAssetEditor()
 	UAssetViewerSettings::Get()->OnAssetViewerSettingsChanged().Broadcast(NAME_None);
 }
 
-void FDataflowEditorToolkit::SetEvaluateGraphMode(FDataflowEditorToolkit::EEvaluationMode Mode)
+void FDataflowEditorToolkit::SetEvaluateGraphMode(EDataflowEditorEvaluationMode Mode)
 {
 	EvaluationMode = Mode;
+	if (UDataflowEditorOptions* const Options = UDataflowEditorOptions::StaticClass()->GetDefaultObject<UDataflowEditorOptions>())
+	{
+		Options->EditorEvaluationMode = EvaluationMode;
+	}
 }
 
 void FDataflowEditorToolkit::ClearGraphCache()
@@ -659,7 +665,7 @@ FSlateIcon FDataflowEditorToolkit::GetEvaluationStatusImage() const
 	}
 
 	FName OverlayIcon;
-	if (IsGraphDirty())
+	if (EvaluationMode == EDataflowEditorEvaluationMode::Manual && IsGraphDirty())
 	{
 		OverlayIcon = CompileStatusUnknown;
 	}
@@ -681,20 +687,21 @@ FSlateIcon FDataflowEditorToolkit::GetEvaluationStatusImage() const
 	return FSlateIcon(FAppStyle::GetAppStyleSetName(), CompileStatusBackground, NAME_None, OverlayIcon);
 }
 
-void FDataflowEditorToolkit::AddEValuationWidget(FToolMenuSection& InSection)
+void FDataflowEditorToolkit::AddEvaluationWidget(FToolMenuSection& InSection)
 {
 	ToolkitCommands->MapAction(FDataflowEditorCommands::Get().EvaluateGraph,
-		FExecuteAction::CreateSP(this, &FDataflowEditorToolkit::EvaluateGraph)
+		FExecuteAction::CreateSP(this, &FDataflowEditorToolkit::EvaluateGraph),
+		FCanExecuteAction::CreateSP(this, &FDataflowEditorToolkit::IsEvaluateButtonEnabled)
 	);
 	ToolkitCommands->MapAction(FDataflowEditorCommands::Get().EvaluateGraphAutomatic,
-		FExecuteAction::CreateSP(this, &FDataflowEditorToolkit::SetEvaluateGraphMode, EEvaluationMode::Automatic),
+		FExecuteAction::CreateSP(this, &FDataflowEditorToolkit::SetEvaluateGraphMode, EDataflowEditorEvaluationMode::Automatic),
 		FCanExecuteAction(),
-		FIsActionChecked::CreateLambda([this]() { return EvaluationMode == EEvaluationMode::Automatic; })
+		FIsActionChecked::CreateLambda([this]() { return EvaluationMode == EDataflowEditorEvaluationMode::Automatic; })
 	);
 	ToolkitCommands->MapAction(FDataflowEditorCommands::Get().EvaluateGraphManual,
-		FExecuteAction::CreateSP(this, &FDataflowEditorToolkit::SetEvaluateGraphMode, EEvaluationMode::Manual),
+		FExecuteAction::CreateSP(this, &FDataflowEditorToolkit::SetEvaluateGraphMode, EDataflowEditorEvaluationMode::Manual),
 		FCanExecuteAction(),
-		FIsActionChecked::CreateLambda([this]() { return EvaluationMode == EEvaluationMode::Manual; })
+		FIsActionChecked::CreateLambda([this]() { return EvaluationMode == EDataflowEditorEvaluationMode::Manual; })
 	);
 	ToolkitCommands->MapAction(FDataflowEditorCommands::Get().ClearGraphCache,
 		FExecuteAction::CreateSP(this, &FDataflowEditorToolkit::ClearGraphCache),
@@ -719,6 +726,12 @@ void FDataflowEditorToolkit::AddEValuationWidget(FToolMenuSection& InSection)
 	);
 	EvaluationOptions.StyleNameOverride = "SlimToolBar";
 	InSection.AddEntry(EvaluationOptions);
+
+	// load options
+	if (UDataflowEditorOptions* const Options = UDataflowEditorOptions::StaticClass()->GetDefaultObject<UDataflowEditorOptions>())
+	{
+		EvaluationMode = Options->EditorEvaluationMode;
+	}
 }
 
 TSharedRef<SWidget> FDataflowEditorToolkit::GenerateEvaluationOptionsMenu()
@@ -733,6 +746,12 @@ TSharedRef<SWidget> FDataflowEditorToolkit::GenerateEvaluationOptionsMenu()
 	return MenuBuilder.MakeWidget();
 }
 
+bool FDataflowEditorToolkit::IsEvaluateButtonEnabled() const
+{
+	// we disable the button in automatic mode as this should not be needed
+	return (EvaluationMode == EDataflowEditorEvaluationMode::Manual);
+}
+
 void FDataflowEditorToolkit::InitializeEdMode(UBaseCharacterFXEditorMode* EdMode)
 {
 	UDataflowEditorMode* DataflowMode = Cast<UDataflowEditorMode>(EdMode);
@@ -1488,7 +1507,7 @@ void FDataflowEditorToolkit::EvaluateGraph()
 
 void FDataflowEditorToolkit::OnNodeInvalidated(UDataflow& DataflowAsset, FDataflowNode& Node)
 {
-	if (EvaluationMode == EEvaluationMode::Automatic)
+	if (EvaluationMode == EDataflowEditorEvaluationMode::Automatic)
 	{
 		if (const TObjectPtr<UDataflowBaseContent>& EditorContent = GetEditorContent())
 		{
diff --git a/Engine/Plugins/Experimental/Dataflow/Source/DataflowEditor/Public/Dataflow/DataflowEditorToolkit.h b/Engine/Plugins/Experimental/Dataflow/Source/DataflowEditor/Public/Dataflow/DataflowEditorToolkit.h
--- a/Engine/Plugins/Experimental/Dataflow/Source/DataflowEditor/Public/Dataflow/DataflowEditorToolkit.h
+++ b/Engine/Plugins/Experimental/Dataflow/Source/DataflowEditor/Public/Dataflow/DataflowEditorToolkit.h
@@ -35,6 +35,8 @@ class FDataflowSimulationSceneProfileIndexStorage;
 class FDataflowOutputLog;
 struct FDataflowPath;
 
+enum class EDataflowEditorEvaluationMode : uint8;
+
 namespace UE::Dataflow
 {
 	class FDataflowNodeDetailExtensionHandler;
@@ -45,12 +47,6 @@ class DATAFLOWEDITOR_API FDataflowEditorToolkit final : public FBaseCharacterFXE
 {
 	using FBaseCharacterFXEditorToolkit::ObjectScene;
 
-	enum class EEvaluationMode: uint8
-	{
-		Automatic,
-		Manual,
-	};
-
 public:
 
 	explicit FDataflowEditorToolkit(UAssetEditor* InOwningAssetEditor);
@@ -234,12 +230,13 @@ private:
 	TSharedPtr<SDataflowMembersWidget> CreateDataflowMembersWidget();
 	TSharedRef<SGraphEditor> CreateSubGraphEditorWidget(TSharedRef<FTabInfo> InTabInfo, UDataflowSubGraph* InGraph);
 
-	void AddEValuationWidget(FToolMenuSection& Section);
+	void AddEvaluationWidget(FToolMenuSection& Section);
 	TSharedRef<SWidget> GenerateEvaluationOptionsMenu();
 	FSlateIcon GetEvaluationStatusImage() const;
 	bool IsGraphDirty() const;
+	bool IsEvaluateButtonEnabled() const;
 
-	void SetEvaluateGraphMode(EEvaluationMode Mode);
+	void SetEvaluateGraphMode(EDataflowEditorEvaluationMode Mode);
 	void ClearGraphCache();
 	bool CanClearGraphCache() const;
 
@@ -318,7 +315,7 @@ private:
 	/** Cached value of the p.Dataflow.EnableGraphEval cvar, to avoid calling FindConsoleVariable too often */
 	bool bDataflowEnableGraphEval;
 
-	EEvaluationMode EvaluationMode = EEvaluationMode::Automatic;
+	EDataflowEditorEvaluationMode EvaluationMode;
 
 	/** Delegate for updating the cached value of p.Dataflow.EnableGraphEval */
 	FDelegateHandle GraphEvalCVarChangedDelegateHandle;
