ChaosFlesh:
-Added frame index text box to show current frame index and jump to arbitrary frame
-Added subframe rate to clamp time to nearest fraction (default: nearest integer frame)

diff --git a/Engine/Plugins/Experimental/Dataflow/Source/DataflowEditor/Private/Dataflow/DataflowSimulationPanel.cpp b/Engine/Plugins/Experimental/Dataflow/Source/DataflowEditor/Private/Dataflow/DataflowSimulationPanel.cpp
--- a/Engine/Plugins/Experimental/Dataflow/Source/DataflowEditor/Private/Dataflow/DataflowSimulationPanel.cpp
+++ b/Engine/Plugins/Experimental/Dataflow/Source/DataflowEditor/Private/Dataflow/DataflowSimulationPanel.cpp
@@ -11,6 +11,8 @@
 #include "Dataflow/DataflowContent.h"
 #include "Widgets/Input/SButton.h"
 #include "Widgets/Images/SImage.h"
+#include "Widgets/Input/SEditableTextBox.h"
+#include "Widgets/Layout/SBox.h"
 
 #define LOCTEXT_NAMESPACE "DataflowSimulationScrubPanel"
 
@@ -29,12 +31,27 @@ void SDataflowSimulationPanel::Construct( const SDataflowSimulationPanel::FArgum
 	}
 	const FTransportControlWidget NewWidget(FOnMakeTransportWidget::CreateSP(this, &SDataflowSimulationPanel::OnCreatePreviewPlaybackModeWidget));
 	TransportControlWidgets.Add(NewWidget);
-
+	FrameIndexWidget = SNew(SEditableTextBox)
+		.OnTextCommitted(this, &SDataflowSimulationPanel::SetFrameIndex)
+		.IsEnabled(true);
+	FrameIndexWidget->SetText(FText::AsNumber(0));
 	this->ChildSlot
 	[
 		SNew(SHorizontalBox)
 		.AddMetaData<FTagMetaData>(TEXT("DataflowSimulationScrub.Scrub"))
 		+SHorizontalBox::Slot()
+		.HAlign(HAlign_Fill)
+		.VAlign(VAlign_Center)
+		.AutoWidth()
+		.Padding(0,0,8,0)
+		[	
+			SNew(SBox)
+			.WidthOverride(60.f)
+			[
+				FrameIndexWidget.ToSharedRef()
+			]
+		]
+		+SHorizontalBox::Slot()
 		.HAlign(HAlign_Fill) 
 		.VAlign(VAlign_Center)
 		.FillWidth(1)
@@ -125,7 +142,7 @@ FReply SDataflowSimulationPanel::OnClick_Forward_Step()
 {
 	if (const TSharedPtr<FDataflowSimulationScene> PreviewScene = SimulationScene.Pin())
 	{
-		PreviewScene->SimulationTime += PreviewScene->GetDeltaTime();
+		UpdateSimulationTimeFromScrubValue(GetScrubValue() + PreviewScene->GetDeltaTime());
 	}
 
 	return FReply::Handled();
@@ -133,10 +150,7 @@ FReply SDataflowSimulationPanel::OnClick_Forward_Step()
 
 FReply SDataflowSimulationPanel::OnClick_Forward_End()
 {
-	if (const TSharedPtr<FDataflowSimulationScene> PreviewScene = SimulationScene.Pin())
-	{
-		PreviewScene->SimulationTime = PreviewScene->GetTimeRange()[1];
-	}
+	UpdateSimulationTimeFromScrubValue(GetSequenceLength());
 
 	return FReply::Handled();
 }
@@ -145,7 +159,7 @@ FReply SDataflowSimulationPanel::OnClick_Backward_Step()
 {
 	if (const TSharedPtr<FDataflowSimulationScene> PreviewScene = SimulationScene.Pin())
 	{
-		PreviewScene->SimulationTime -= PreviewScene->GetDeltaTime();
+		UpdateSimulationTimeFromScrubValue(GetScrubValue() - PreviewScene->GetDeltaTime());
 	}
 
 	return FReply::Handled();
@@ -155,7 +169,7 @@ FReply SDataflowSimulationPanel::OnClick_Backward_End()
 {
 	if (const TSharedPtr<FDataflowSimulationScene> PreviewScene = SimulationScene.Pin())
 	{
-		PreviewScene->SimulationTime = PreviewScene->GetTimeRange()[0];
+		UpdateSimulationTimeFromScrubValue(0.f);
 	}
 
 	return FReply::Handled();
@@ -222,28 +236,28 @@ void SDataflowSimulationPanel::OnTickPlayback(double InCurrentTime, float InDelt
 	{
 		const float SequenceLength = GetSequenceLength();
 		
-		const float SimulationTime = (PlaybackMode == EPlaybackMode::PlayingForward) ?
-			PreviewScene->SimulationTime + InDeltaTime : PreviewScene->SimulationTime - InDeltaTime;
+		const float ScrubValue = (PlaybackMode == EPlaybackMode::PlayingForward) ?
+			GetScrubValue() + InDeltaTime : GetScrubValue() - InDeltaTime;
 		
 		if(PreviewPlaybackMode == EDataflowPlaybackMode::Looping)
 		{
-			PreviewScene->SimulationTime = SimulationTime - SequenceLength * FMath::Floor((SimulationTime-PreviewScene->GetTimeRange()[0]) / SequenceLength);
+			UpdateSimulationTimeFromScrubValue(ScrubValue - SequenceLength * FMath::Floor(ScrubValue / SequenceLength));
 		}
 		else
 		{
 			if(PreviewPlaybackMode == EDataflowPlaybackMode::PingPong)
 			{
-				if((PlaybackMode == EPlaybackMode::PlayingForward) && (SimulationTime >= PreviewScene->GetTimeRange()[1]))
+				if((PlaybackMode == EPlaybackMode::PlayingForward) && (ScrubValue >= SequenceLength))
 				{
 					PlaybackMode = EPlaybackMode::PlayingReverse;
 				}
 
-				if((PlaybackMode == EPlaybackMode::PlayingReverse) && (SimulationTime <= PreviewScene->GetTimeRange()[0]))
+				if((PlaybackMode == EPlaybackMode::PlayingReverse) && (ScrubValue <= 0))
 				{
 					PlaybackMode = EPlaybackMode::PlayingForward;
 				}
 			}
-			PreviewScene->SimulationTime = FMath::Clamp(SimulationTime, PreviewScene->GetTimeRange()[0],  PreviewScene->GetTimeRange()[1]);
+			UpdateSimulationTimeFromScrubValue(FMath::Clamp(ScrubValue, 0.f, SequenceLength));
 		}
 	}
 }
@@ -253,11 +267,32 @@ EPlaybackMode::Type SDataflowSimulationPanel::GetPlaybackMode() const
 	return PlaybackMode;
 }
 
-void SDataflowSimulationPanel::OnValueChanged(float NewValue)
+void SDataflowSimulationPanel::UpdateSimulationTimeFromScrubValue(float ScrubValue)
 {
 	if (const TSharedPtr<FDataflowSimulationScene> PreviewScene = SimulationScene.Pin())
 	{
-		PreviewScene->SimulationTime = NewValue + PreviewScene->GetTimeRange()[0];
+		const int32 TotalFrameRate = PreviewScene->GetFrameRate() * PreviewScene->GetSubframeRate();
+		const float RoundedFrameTime = FMath::RoundToFloat(ScrubValue * TotalFrameRate) / TotalFrameRate;
+		PreviewScene->SimulationTime = RoundedFrameTime + PreviewScene->GetTimeRange()[0];
+		FrameIndexWidget->SetText(FText::AsNumber(RoundedFrameTime * PreviewScene->GetFrameRate()));
+	}
+}
+
+void SDataflowSimulationPanel::OnValueChanged(float NewValue)
+{
+	UpdateSimulationTimeFromScrubValue(NewValue);
+}
+
+void SDataflowSimulationPanel::SetFrameIndex(const FText& InNewText, ETextCommit::Type InCommitType)
+{
+	if (InNewText.IsNumeric())
+	{
+		float FrameIndex = 0.f;
+		LexFromString(FrameIndex, *InNewText.ToString());
+		if (const TSharedPtr<FDataflowSimulationScene> PreviewScene = SimulationScene.Pin())
+		{
+			UpdateSimulationTimeFromScrubValue(FrameIndex / float(PreviewScene->GetFrameRate()));
+		}
 	}
 }
 
diff --git a/Engine/Plugins/Experimental/Dataflow/Source/DataflowEditor/Public/Dataflow/DataflowSimulationPanel.h b/Engine/Plugins/Experimental/Dataflow/Source/DataflowEditor/Public/Dataflow/DataflowSimulationPanel.h
--- a/Engine/Plugins/Experimental/Dataflow/Source/DataflowEditor/Public/Dataflow/DataflowSimulationPanel.h
+++ b/Engine/Plugins/Experimental/Dataflow/Source/DataflowEditor/Public/Dataflow/DataflowSimulationPanel.h
@@ -9,6 +9,7 @@
 class SScrubControlPanel;
 class SButton;
 class UAnimSingleNodeInstance;
+class SEditableTextBox;
 
 /** Dataflow simulation panel to control an animation/simulation */
 class SDataflowSimulationPanel : public SCompoundWidget
@@ -47,7 +48,9 @@ private:
 
 	/** Main simulation delegates */
 	void OnTickPlayback(double InCurrentTime, float InDeltaTime);
+	void UpdateSimulationTimeFromScrubValue(float ScrubValue);
 	void OnValueChanged(float NewValue);
+	void SetFrameIndex(const FText& InNewText, ETextCommit::Type InCommitType);
 	void OnBeginSliderMovement();
 
 	/** Get the playback mode used in the widget */
@@ -74,6 +77,9 @@ private:
 	/** Playback mode button */
 	TSharedPtr<SButton> PreviewPlaybackModeButton;
 
+	/** Widget showing editable frame index*/
+	TSharedPtr<SEditableTextBox> FrameIndexWidget;
+
 	/** Preview playback mode (looping...)*/
 	EDataflowPlaybackMode PreviewPlaybackMode = EDataflowPlaybackMode::Looping;
 
diff --git a/Engine/Plugins/Experimental/Dataflow/Source/DataflowEditor/Public/Dataflow/DataflowSimulationScene.h b/Engine/Plugins/Experimental/Dataflow/Source/DataflowEditor/Public/Dataflow/DataflowSimulationScene.h
--- a/Engine/Plugins/Experimental/Dataflow/Source/DataflowEditor/Public/Dataflow/DataflowSimulationScene.h
+++ b/Engine/Plugins/Experimental/Dataflow/Source/DataflowEditor/Public/Dataflow/DataflowSimulationScene.h
@@ -146,6 +146,12 @@ public:
 	/** Get the number of frames */
 	const int32& GetNumFrames() const {return NumFrames;}
 
+	/** Get the frame rate */
+	const int32 GetFrameRate() const { return SceneDescription->CacheParams.FrameRate; }
+
+	/** Get the subframe rate */
+	const int32 GetSubframeRate() const { return SceneDescription->CacheParams.SubframeRate; }
+
 	/** Get delta time */
 	float GetDeltaTime() const { return DeltaTime; }
 
diff --git a/Engine/Source/Runtime/Experimental/Dataflow/Engine/Public/Dataflow/DataflowPreview.h b/Engine/Source/Runtime/Experimental/Dataflow/Engine/Public/Dataflow/DataflowPreview.h
--- a/Engine/Source/Runtime/Experimental/Dataflow/Engine/Public/Dataflow/DataflowPreview.h
+++ b/Engine/Source/Runtime/Experimental/Dataflow/Engine/Public/Dataflow/DataflowPreview.h
@@ -22,6 +22,12 @@ public :
 	UPROPERTY(EditAnywhere, Category="Caching")
 	int32 FrameRate = 30;
 
+	/** Number of subframes per frame for timeline clamping*/
+	UPROPERTY(EditAnywhere, meta = (EditCondition = "bCanEditSubframeRate", ClampMin = "1"), Category = "Caching")
+	int32 SubframeRate = 1;
+	UPROPERTY(EditAnywhere, meta = (InlineEditConditionToggle), Category = "Caching")
+	bool bCanEditSubframeRate = false;
+
 	/** Time range of the simulation*/
 	UPROPERTY(EditAnywhere, Category="Caching")
 	FVector2f TimeRange = FVector2f(0.0f, 5.0f);
