Dataflow : fix empty localization string

diff --git a/Engine/Plugins/Experimental/Dataflow/Source/DataflowEditor/Private/Dataflow/DataflowEditorToolkit.cpp b/Engine/Plugins/Experimental/Dataflow/Source/DataflowEditor/Private/Dataflow/DataflowEditorToolkit.cpp
--- a/Engine/Plugins/Experimental/Dataflow/Source/DataflowEditor/Private/Dataflow/DataflowEditorToolkit.cpp
+++ b/Engine/Plugins/Experimental/Dataflow/Source/DataflowEditor/Private/Dataflow/DataflowEditorToolkit.cpp
@@ -719,7 +719,7 @@ void FDataflowEditorToolkit::AddEvaluationWidget(FToolMenuSection& InSection)
 		"DataflowEvaluationOptions",
 		FUIAction(),
 		FOnGetContent::CreateSP(this, &FDataflowEditorToolkit::GenerateEvaluationOptionsMenu),
-		LOCTEXT("DataflowEvaluationOptions", ""),
+		LOCTEXT("DataflowEvaluationOptions", "Options"),
 		LOCTEXT("DataflowEvaluationOptions_ToolbarTooltip", "Options to customize how Dataflow evaluate"),
 		TAttribute<FSlateIcon>(),
 		true
