Dataflow scene outliner

diff --git a/Engine/Plugins/Experimental/Dataflow/Dataflow.uplugin b/Engine/Plugins/Experimental/Dataflow/Dataflow.uplugin
--- a/Engine/Plugins/Experimental/Dataflow/Dataflow.uplugin
+++ b/Engine/Plugins/Experimental/Dataflow/Dataflow.uplugin
@@ -62,6 +62,14 @@
 		{
 			"Name": "GeometryCache",
 			"Enabled": true
+		},
+		{
+			"Name": "EditorDataStorage",
+			"Enabled": true
+		},
+		{
+			"Name": "EditorDataStorageFeatures",
+			"Enabled": true
 		}
 	]
 }
diff --git a/Engine/Plugins/Experimental/Dataflow/Source/DataflowEditor/DataflowEditor.Build.cs b/Engine/Plugins/Experimental/Dataflow/Source/DataflowEditor/DataflowEditor.Build.cs
--- a/Engine/Plugins/Experimental/Dataflow/Source/DataflowEditor/DataflowEditor.Build.cs
+++ b/Engine/Plugins/Experimental/Dataflow/Source/DataflowEditor/DataflowEditor.Build.cs
@@ -4,8 +4,16 @@ namespace UnrealBuildTool.Rules
 {
 	public class DataflowEditor : ModuleRules
 	{
-        public DataflowEditor(ReadOnlyTargetRules Target) : base(Target)
+		public DataflowEditor(ReadOnlyTargetRules Target) : base(Target)
 		{
+
+			PrivateDependencyModuleNames.AddRange(
+				new string[]
+				{
+					"TedsOutliner",
+				}
+			);
+
 			PublicIncludePathModuleNames.AddRange(
 				new string[]
 				{
@@ -59,6 +67,7 @@ namespace UnrealBuildTool.Rules
 					"RenderCore",
 					"RHI",
 					"SceneOutliner",
+					"TedsOutliner",
 					"SharedSettingsWidgets",
 					"SkeletonEditor",
 					"Slate",
@@ -67,6 +76,7 @@ namespace UnrealBuildTool.Rules
 					"ToolMenus",
 					"ToolWidgets",
 					"TypedElementRuntime",
+					"TypedElementFramework",
 					"UnrealEd",
 					"WorkspaceMenuStructure",
 					"XmlParser",
diff --git a/Engine/Plugins/Experimental/Dataflow/Source/DataflowEditor/Private/Dataflow/DataflowConstructionScene.cpp b/Engine/Plugins/Experimental/Dataflow/Source/DataflowEditor/Private/Dataflow/DataflowConstructionScene.cpp
--- a/Engine/Plugins/Experimental/Dataflow/Source/DataflowEditor/Private/Dataflow/DataflowConstructionScene.cpp
+++ b/Engine/Plugins/Experimental/Dataflow/Source/DataflowEditor/Private/Dataflow/DataflowConstructionScene.cpp
@@ -21,8 +21,10 @@
 //
 
 FDataflowConstructionScene::FDataflowConstructionScene(FPreviewScene::ConstructionValues ConstructionValues, UDataflowEditor* InEditor)
-	: FDataflowPreviewSceneBase(ConstructionValues, InEditor)
-{}
+	: FDataflowPreviewSceneBase(ConstructionValues, InEditor, FName("Construction Components"))
+{
+	RootSceneActor->GetName();
+}
 
 FDataflowConstructionScene::~FDataflowConstructionScene()
 {
@@ -137,6 +139,8 @@ void FDataflowConstructionScene::UpdateDynamicMeshComponents()
 	// list of UPrimitiveComponents for rendering.
 	ResetDynamicMeshComponents();
 
+	AddSceneObject(RootSceneActor, true);
+
 	if (const TObjectPtr<UDataflowBaseContent>& EditorContent = GetEditorContent())
 	{
 		const TObjectPtr<UDataflow>& DataflowAsset = EditorContent->GetDataflowAsset();
@@ -268,9 +272,9 @@ void FDataflowConstructionScene::UpdateDynamicMeshComponents()
 					}
 				}
 			}
-
 		}
 	}
+	bPreviewSceneDirty = true;
 }
 
 void FDataflowConstructionScene::ResetDynamicMeshComponents()
@@ -286,10 +290,14 @@ void FDataflowConstructionScene::ResetDynamicMeshComponents()
 			SelectedComponents->Deselect(DynamicMeshComponent);
 			DynamicMeshComponent->PushSelectionToProxy();
 		}
+		RemoveSceneObject(DynamicMeshComponent);
 		RemoveComponent(DynamicMeshComponent);
 		DynamicMeshComponent->DestroyComponent();
 	}
 	DynamicMeshComponents.Reset();
+	bPreviewSceneDirty = true;
+
+	RemoveSceneObject(RootSceneActor);
 }
 
 TObjectPtr<UDynamicMeshComponent>& FDataflowConstructionScene::AddDynamicMeshComponent(FDataflowRenderKey InKey, const FString& MeshName, UE::Geometry::FDynamicMesh3&& DynamicMesh, const TArray<UMaterialInterface*>& MaterialSet)
@@ -343,6 +351,7 @@ TObjectPtr<UDynamicMeshComponent>& FDataflowConstructionScene::AddDynamicMeshCom
 
 	AddComponent(DynamicMeshComponent, DynamicMeshComponent->GetRelativeTransform());	
 	DynamicMeshComponents.Emplace(InKey, DynamicMeshComponent);
+	AddSceneObject(DynamicMeshComponent, true);
 	return DynamicMeshComponents[InKey];
 }
 
@@ -468,5 +477,6 @@ void FDataflowConstructionScene::UpdateConstructionScene()
 	}
 }
 
+
 #undef LOCTEXT_NAMESPACE
 
diff --git a/Engine/Plugins/Experimental/Dataflow/Source/DataflowEditor/Private/Dataflow/DataflowConstructionViewportClient.cpp b/Engine/Plugins/Experimental/Dataflow/Source/DataflowEditor/Private/Dataflow/DataflowConstructionViewportClient.cpp
--- a/Engine/Plugins/Experimental/Dataflow/Source/DataflowEditor/Private/Dataflow/DataflowConstructionViewportClient.cpp
+++ b/Engine/Plugins/Experimental/Dataflow/Source/DataflowEditor/Private/Dataflow/DataflowConstructionViewportClient.cpp
@@ -118,8 +118,9 @@ void FDataflowConstructionViewportClient::Tick(float DeltaSeconds)
 }
 
 USelection* FDataflowConstructionViewportClient::GetSelectedComponents() 
-{ 
-	return ModeTools->GetSelectedComponents(); 
+{
+	USelection* SceneSelection = FDataflowPreviewSceneBase::GetSelectedComponents(PreviewScene);
+	return SceneSelection ? SceneSelection : ModeTools->GetSelectedComponents(); 
 }
 
 bool FDataflowConstructionViewportClient::InputKey(const FInputKeyEventArgs& EventArgs)
@@ -246,7 +247,7 @@ void FDataflowConstructionViewportClient::OnViewportClicked(HHitProxy* HitProxy)
 	TArray<UPrimitiveComponent*> CurrentlySelectedComponents;
 	if (!IsInteractiveToolActive())
 	{
-		if (USelection* SelectedComponents = ModeTools->GetSelectedComponents())
+		if (USelection* SelectedComponents = GetSelectedComponents())
 		{
 			UpdateSelectedComponentInViewport(SelectedComponents);
 
@@ -264,7 +265,12 @@ void FDataflowConstructionViewportClient::OnViewportClicked(HHitProxy* HitProxy)
 			SelectedComponents->GetSelectedObjects<UPrimitiveComponent>(CurrentlySelectedComponents);
 		}
 	}
-	OnSelectionChangedMulticast.Broadcast(CurrentlySelectedComponents);
+	
+	// Get all the scene selected elements 
+	TArray<FDataflowBaseElement*> DataflowElements;
+	GetSelectedElements(HitProxy, DataflowElements);
+	
+	OnSelectionChangedMulticast.Broadcast(CurrentlySelectedComponents, DataflowElements);
 }
 
 void FDataflowConstructionViewportClient::SetConstructionViewMode(const UE::Dataflow::IDataflowConstructionViewMode* InViewMode)
@@ -326,7 +332,7 @@ void FDataflowConstructionViewportClient::SetConstructionViewMode(const UE::Data
 
 void FDataflowConstructionViewportClient::Draw(const FSceneView* View, FPrimitiveDrawInterface* PDI)
 {
-	FEditorViewportClient::Draw(View, PDI);
+	FDataflowEditorViewportClientBase::Draw(View, PDI);
 
 	using namespace UE::Dataflow;
 	for (const TPair<FName, TUniquePtr<IDataflowConstructionVisualization>>& Visualization : FDataflowConstructionVisualizationRegistry::GetInstance().GetVisualizations())
diff --git a/Engine/Plugins/Experimental/Dataflow/Source/DataflowEditor/Private/Dataflow/DataflowEditorMode.cpp b/Engine/Plugins/Experimental/Dataflow/Source/DataflowEditor/Private/Dataflow/DataflowEditorMode.cpp
--- a/Engine/Plugins/Experimental/Dataflow/Source/DataflowEditor/Private/Dataflow/DataflowEditorMode.cpp
+++ b/Engine/Plugins/Experimental/Dataflow/Source/DataflowEditor/Private/Dataflow/DataflowEditorMode.cpp
@@ -168,7 +168,7 @@ void UDataflowEditorMode::RegisterDataflowTool(TSharedPtr<FUICommandInfo> UIComm
 					}
 
 					bool bHadSingleSelectedMesh = false;
-					if (const USelection* const SelectedComponents = GetModeManager()->GetSelectedComponents())
+					if (const USelection* const SelectedComponents = FDataflowPreviewSceneBase::GetSelectedComponents(ConstructionScene))
 					{
 						bHadSingleSelectedMesh = (SelectedComponents->Num() == 1);
 					}
@@ -182,7 +182,7 @@ void UDataflowEditorMode::RegisterDataflowTool(TSharedPtr<FUICommandInfo> UIComm
 						const TArray<UDynamicMeshComponent*> DynamicMeshComponents = ConstructionScene->GetDynamicMeshComponents();
 						if (DynamicMeshComponents.Num() == 1)
 						{
-							if (USelection* const SelectedComponents = GetModeManager()->GetSelectedComponents())
+							if (USelection* const SelectedComponents = FDataflowPreviewSceneBase::GetSelectedComponents(ConstructionScene))
 							{
 								SelectedComponents->Select(DynamicMeshComponents[0]);
 							}
@@ -252,7 +252,7 @@ void UDataflowEditorMode::AddNode(FName NewNodeType)
 bool UDataflowEditorMode::CanAddNode(FName NewNodeType) const
 {
 	const UEdGraphNode* const CurrentlySelectedNode = GetSingleSelectedNodeWithOutputType(FManagedArrayCollection::StaticType());
-	if (CurrentlySelectedNode)
+	if (CurrentlySelectedNode && GetToolManager())
 	{
 		const bool bToolActive = GetToolManager()->HasActiveTool(EToolSide::Left);
 		return !bToolActive;
@@ -476,9 +476,9 @@ void UDataflowEditorMode::CreateToolTargets(const TArray<TObjectPtr<UObject>>& A
 
 bool UDataflowEditorMode::IsComponentSelected(const UPrimitiveComponent* InComponent)
 {
-	if (const FEditorModeTools* const ModeManager = GetModeManager())
+	if (ConstructionScene && ConstructionScene->GetDataflowModeManager())
 	{
-		if (const UTypedElementSelectionSet* const TypedElementSelectionSet = ModeManager->GetEditorSelectionSet())
+		if (const UTypedElementSelectionSet* const TypedElementSelectionSet = ConstructionScene->GetDataflowModeManager()->GetEditorSelectionSet())
 		{
 			if (const FTypedElementHandle ComponentElement = UEngineElementsLibrary::AcquireEditorComponentElementHandle(InComponent))
 			{
@@ -630,7 +630,7 @@ void UDataflowEditorMode::ModeTick(float DeltaTime)
 				const TArray<UDynamicMeshComponent*> DynamicMeshComponents = ConstructionScene->GetDynamicMeshComponents();
 				if (DynamicMeshComponents.Num() == 1)
 				{
-					if (USelection* const SelectedComponents = GetModeManager()->GetSelectedComponents())
+					if (USelection* const SelectedComponents = FDataflowPreviewSceneBase::GetSelectedComponents(ConstructionScene))
 					{
 						SelectedComponents->Select(DynamicMeshComponents[0]);
 					}
@@ -737,7 +737,7 @@ void UDataflowEditorMode::SetConstructionViewMode(const FName& NewViewModeName)
 
 	if (UInteractiveTool* const ActiveTool = ToolManager->GetActiveTool(EToolSide::Left))
 	{
-		if (const USelection* const SelectedComponents = GetModeManager()->GetSelectedComponents())
+		if (const USelection* const SelectedComponents = FDataflowPreviewSceneBase::GetSelectedComponents(ConstructionScene))
 		{
 			if (SelectedComponents->Num() == 1)			// TODO: Extend this to handle multiple selected components
 			{
@@ -811,7 +811,7 @@ void UDataflowEditorMode::SetConstructionViewMode(const FName& NewViewModeName)
 			const TArray<UDynamicMeshComponent*> DynamicMeshComponents = ConstructionScene->GetDynamicMeshComponents();
 			if (DynamicMeshComponents.Num() == 1)
 			{
-				if (USelection* const SelectedComponents = GetModeManager()->GetSelectedComponents())
+				if (USelection* const SelectedComponents = FDataflowPreviewSceneBase::GetSelectedComponents(ConstructionScene))
 				{
 					SelectedComponents->Select(DynamicMeshComponents[0]);
 				}
diff --git a/Engine/Plugins/Experimental/Dataflow/Source/DataflowEditor/Private/Dataflow/DataflowEditorPreviewSceneBase.cpp b/Engine/Plugins/Experimental/Dataflow/Source/DataflowEditor/Private/Dataflow/DataflowEditorPreviewSceneBase.cpp
--- a/Engine/Plugins/Experimental/Dataflow/Source/DataflowEditor/Private/Dataflow/DataflowEditorPreviewSceneBase.cpp
+++ b/Engine/Plugins/Experimental/Dataflow/Source/DataflowEditor/Private/Dataflow/DataflowEditorPreviewSceneBase.cpp
@@ -7,17 +7,47 @@
 #include "Elements/Framework/EngineElementsLibrary.h"
 #include "Selection.h"
 #include "AssetViewerSettings.h"
+#include "Dataflow/DataflowElement.h"
+#include "Elements/Columns/TypedElementHiearchyColumns.h"
+#include "Elements/Columns/TypedElementLabelColumns.h"
+#include "Elements/Common/EditorDataStorageFeatures.h"
+#include "Elements/Interfaces/TypedElementDataStorageCompatibilityInterface.h"
+#include "Elements/Columns/TypedElementMiscColumns.h"
+#include "Elements/Columns/TypedElementSelectionColumns.h"
+#include "Elements/Columns/TypedElementViewportColumns.h"
+#include "Elements/Columns/TypedElementVisibilityColumns.h"
 
 #define LOCTEXT_NAMESPACE "FDataflowPreviewSceneBase"
 
 
+namespace UE::Dataflow::Private
+{
+	static void AddObjectColumns(UE::Editor::DataStorage::ICoreProvider* DataStorage, const UE::Editor::DataStorage::RowHandle& Row, const bool bIsConstruction)
+	{
+		if(bIsConstruction)
+		{
+			DataStorage->AddColumn<FDataflowConstructionObjectTag>(Row);
+		}
+		else
+		{
+			DataStorage->AddColumn<FDataflowSimulationObjectTag>(Row);
+		}
+		DataStorage->AddColumn<FTypedElementSyncFromWorldTag>(Row);
+		DataStorage->AddColumn<FTypedElementLabelColumn>(Row);
+		DataStorage->AddColumn<FTypedElementLabelHashColumn>(Row);
+		DataStorage->AddColumn<FVisibleInEditorColumn>(Row);
+		DataStorage->AddColumn<FTableRowParentColumn>(Row);
+		DataStorage->AddColumn<FDataflowSceneTypeColumn>(Row);
+	}
+}
+
 bool bDataflowShowFloorDefault = true;
 FAutoConsoleVariableRef CVARDataflowShowFloorDefault(TEXT("p.Dataflow.Editor.ShowFloor"), bDataflowShowFloorDefault, TEXT("Show the floor in the dataflow editor[def:false]"));
 
 bool bDataflowShowEnvironmentDefault = true;
 FAutoConsoleVariableRef CVARDataflowShowEnvironmentDefault(TEXT("p.Dataflow.Editor.ShowEnvironment"), bDataflowShowEnvironmentDefault, TEXT("Show the environment in the dataflow editor[def:false]"));
 
-FDataflowPreviewSceneBase::FDataflowPreviewSceneBase(FPreviewScene::ConstructionValues ConstructionValues, UDataflowEditor* InEditor)
+FDataflowPreviewSceneBase::FDataflowPreviewSceneBase(FPreviewScene::ConstructionValues ConstructionValues, UDataflowEditor* InEditor, const FName& InActorName)
 	: FAdvancedPreviewScene(ConstructionValues)
 	, DataflowEditor(InEditor)
 {
@@ -25,7 +55,9 @@ FDataflowPreviewSceneBase::FDataflowPreviewSceneBase(FPreviewScene::Construction
 	DefaultSettings->OnAssetViewerSettingsChanged().Remove(RefreshDelegate);
 	RefreshDelegate = DefaultSettings->OnAssetViewerSettingsChanged().AddRaw(this, &FDataflowPreviewSceneBase::OnAssetViewerSettingsRefresh);
 
-	RootSceneActor = GetWorld()->SpawnActor<AActor>(AActor::StaticClass());
+	FActorSpawnParameters SpawnParameters;
+	SpawnParameters.Name = InActorName;
+	RootSceneActor = GetWorld()->SpawnActor<AActor>(AActor::StaticClass(), SpawnParameters);
 	
 	check(DataflowEditor);
 	
@@ -103,7 +135,17 @@ FBox FDataflowPreviewSceneBase::GetBoundingBox() const
 		}
 		else
 		{
-			SceneBounds += RootSceneActor->GetComponentsBoundingBox(true);
+			for(const IDataflowDebugDrawInterface::FDataflowElementsType::ElementType& SceneElement : GetSceneElements())
+			{
+				if(SceneElement.IsValid() && SceneElement->bIsSelected)
+				{
+					SceneBounds += SceneElement->BoundingBox;
+				}
+			}
+			if(!SceneBounds.IsValid)
+			{
+				SceneBounds = RootSceneActor->GetComponentsBoundingBox(true);
+			}
 		}
 	}
 	return SceneBounds;
@@ -142,5 +184,99 @@ void FDataflowPreviewSceneBase::SetCurrentProfileIndex(int32 NewProfileIndex)
 	CurrentProfileIndex = NewProfileIndex;
 }
 
+void FDataflowPreviewSceneBase::AddSceneObject(UObject* SceneObject, const bool bIsConstruction) const
+{
+	using namespace UE::Editor::DataStorage;
+	if (ICompatibilityProvider* Compatibility = GetMutableDataStorageFeature<ICompatibilityProvider>(CompatibilityFeatureName))
+	{
+		if (ICoreProvider* DataStorage = GetMutableDataStorageFeature<ICoreProvider>(StorageFeatureName))
+		{
+			RowHandle Row = Compatibility->FindRowWithCompatibleObject(SceneObject);
+			
+			if(Row == InvalidRowHandle)
+			{
+				Row = Compatibility->AddCompatibleObjectExplicit(SceneObject);
+			}
+			
+			UE::Dataflow::Private::AddObjectColumns(DataStorage, Row, bIsConstruction);
+			DataStorage->AddColumn<FDataflowSceneObjectTag>(Row, GetEditorContent()->GetDataflowOwner()->GetFName());
+		}
+	}
+}
+
+void FDataflowPreviewSceneBase::AddSceneStruct(void* SceneStruct, const TWeakObjectPtr<UScriptStruct> TypeInfo,  bool bIsConstruction) const
+{
+	using namespace UE::Editor::DataStorage;
+	if (ICompatibilityProvider* Compatibility = GetMutableDataStorageFeature<ICompatibilityProvider>(CompatibilityFeatureName))
+	{
+		if (ICoreProvider* DataStorage = GetMutableDataStorageFeature<ICoreProvider>(StorageFeatureName))
+		{
+			RowHandle Row = Compatibility->FindRowWithCompatibleObjectExplicit(SceneStruct);
+			
+			if(Row == InvalidRowHandle)
+			{
+				Row = Compatibility->AddCompatibleObjectExplicit(SceneStruct, TypeInfo);
+			}
+
+			UE::Dataflow::Private::AddObjectColumns(DataStorage, Row, bIsConstruction);
+			DataStorage->AddColumn<FDataflowSceneStructTag>(Row, GetEditorContent()->GetDataflowOwner()->GetFName());
+		}
+	}
+}
+
+void FDataflowPreviewSceneBase::RemoveSceneObject(UObject* SceneObject) const
+{
+	using namespace UE::Editor::DataStorage;
+
+	if (ICompatibilityProvider* Compatibility = GetMutableDataStorageFeature<ICompatibilityProvider>(CompatibilityFeatureName))
+	{
+		Compatibility->RemoveCompatibleObject(SceneObject);
+	}
+}
+
+void FDataflowPreviewSceneBase::RemoveSceneStruct(void* SceneStruct) const
+{
+	using namespace UE::Editor::DataStorage;
+
+	if (ICompatibilityProvider* Compatibility = GetMutableDataStorageFeature<ICompatibilityProvider>(CompatibilityFeatureName))
+	{
+		Compatibility->RemoveCompatibleObject(SceneStruct);
+	}
+}
+
+void FDataflowPreviewSceneBase::UnregisterSceneElements()
+{
+	for(IDataflowDebugDrawInterface::FDataflowElementsType::ElementType& SceneElement : ModifySceneElements())
+	{
+		if(SceneElement.IsValid())
+		{
+			RemoveSceneStruct(SceneElement.Get());
+		}
+	}
+}
+
+void FDataflowPreviewSceneBase::RegisterSceneElements(const bool bIsConstruction)
+{
+	for(IDataflowDebugDrawInterface::FDataflowElementsType::ElementType& SceneElement : ModifySceneElements())
+	{
+		if(SceneElement.IsValid())
+		{
+			AddSceneStruct(SceneElement.Get(), FDataflowBaseElement::StaticStruct(), bIsConstruction);
+		}
+	}
+}
+
+USelection* FDataflowPreviewSceneBase::GetSelectedComponents(FDataflowPreviewSceneBase* DataflowScene)
+{
+	if( DataflowScene)
+	{
+		if(TSharedPtr<FAssetEditorModeManager> ModeManager = DataflowScene->GetDataflowModeManager())
+		{
+			return ModeManager->GetSelectedComponents();
+		}
+	}
+	return nullptr; 
+}
+
 #undef LOCTEXT_NAMESPACE
 
diff --git a/Engine/Plugins/Experimental/Dataflow/Source/DataflowEditor/Private/Dataflow/DataflowEditorToolkit.cpp b/Engine/Plugins/Experimental/Dataflow/Source/DataflowEditor/Private/Dataflow/DataflowEditorToolkit.cpp
--- a/Engine/Plugins/Experimental/Dataflow/Source/DataflowEditor/Private/Dataflow/DataflowEditorToolkit.cpp
+++ b/Engine/Plugins/Experimental/Dataflow/Source/DataflowEditor/Private/Dataflow/DataflowEditorToolkit.cpp
@@ -42,6 +42,7 @@
 #include "Dataflow/DataflowSimulationVisualization.h"
 #include "Dataflow/DataflowSchema.h"
 #include "Dataflow/DataflowSkeletonView.h"
+#include "Dataflow/DataflowOutlinerView.h"
 #include "Dataflow/DataflowSimulationViewportClient.h"
 #include "Dataflow/DataflowSubGraphNodes.h"
 #include "Dataflow/DataflowMembersWidget.h"
@@ -58,6 +59,7 @@
 #include "GraphEditorActions.h"
 #include "IDetailCustomization.h"
 #include "ISkeletonTree.h"
+#include "ISceneOutliner.h"
 #include "IStructureDetailsView.h"
 #include "Kismet/GameplayStatics.h"
 #include "Misc/MessageDialog.h"
@@ -74,6 +76,7 @@
 #include "Widgets/Notifications/SNotificationList.h"
 #include "Dataflow/DataflowOutputLog.h"
 #include "Dataflow/DataflowPath.h"
+#include "Dataflow/DataflowSimulationNodes.h"
 
 #define LOCTEXT_NAMESPACE "DataflowEditorToolkit"
 
@@ -87,13 +90,15 @@ const FName FDataflowEditorToolkit::GraphCanvasTabId(TEXT("DataflowEditor_GraphC
 const FName FDataflowEditorToolkit::SubGraphCanvasTabId(TEXT("DataflowEditor_SubGraphTab"));
 const FName FDataflowEditorToolkit::NodeDetailsTabId(TEXT("DataflowEditor_NodeDetails"));
 const FName FDataflowEditorToolkit::PreviewSceneTabId(TEXT("DataflowEditor_PreviewScene"));
+const FName FDataflowEditorToolkit::OutlinerViewTabId(TEXT("DataflowEditor_SceneOutliner"));
 const FName FDataflowEditorToolkit::SkeletonViewTabId(TEXT("DataflowEditor_SkeletonView"));
 const FName FDataflowEditorToolkit::SelectionViewTabId_1(TEXT("DataflowEditor_SelectionView_1"));
 const FName FDataflowEditorToolkit::SelectionViewTabId_2(TEXT("DataflowEditor_SelectionView_2"));
 const FName FDataflowEditorToolkit::SelectionViewTabId_3(TEXT("DataflowEditor_SelectionView_3"));
 const FName FDataflowEditorToolkit::SelectionViewTabId_4(TEXT("DataflowEditor_SelectionView_4"));
 const FName FDataflowEditorToolkit::CollectionSpreadSheetTabId_1(TEXT("DataflowEditor_CollectionSpreadSheet_1"));
-const FName FDataflowEditorToolkit::CollectionSpreadSheetTabId_2(TEXT("DataflowEditor_CollectionSpreadSheet_2"));
+const FName FDataflowEditorToolkit::CollectionSpreadSheetTabId_2(TEXT("DataflowEditor_CollectionSpreadSheet_2"))
+ ;
 const FName FDataflowEditorToolkit::CollectionSpreadSheetTabId_3(TEXT("DataflowEditor_CollectionSpreadSheet_3"));
 const FName FDataflowEditorToolkit::CollectionSpreadSheetTabId_4(TEXT("DataflowEditor_CollectionSpreadSheet_4"));
 const FName FDataflowEditorToolkit::SimulationViewportTabId(TEXT("DataflowEditor_SimulationViewport"));
@@ -101,6 +106,118 @@ const FName FDataflowEditorToolkit::SimulationVisualizationTabId(TEXT("DataflowE
 const FName FDataflowEditorToolkit::MembersWidgetTabId(TEXT("DataflowEditor_MembersWidgetTab"));
 const FName FDataflowEditorToolkit::OutputLogTabId(TEXT("DataflowEditor_OutputLog"));
 
+namespace UE::Dataflow::Private
+{
+	/** Create a debug draw component to visualize per node informations */
+	TUniquePtr<IDataflowDebugDrawInterface> CreateDebugDrawComponent(
+		FDataflowPreviewSceneBase* DataflowScene,
+		TObjectPtr<class UDataflowDebugDrawComponent>& DebugDrawComponent)
+	{
+		if (DataflowScene && !DebugDrawComponent)
+		{
+			if (TObjectPtr<AActor> RootActor = DataflowScene->GetRootActor())
+			{
+				if (UDataflowDebugDrawComponent* Component = Cast<UDataflowDebugDrawComponent>(RootActor->FindComponentByClass<UDataflowDebugDrawComponent>()))
+				{
+					DebugDrawComponent = Cast<UDataflowDebugDrawComponent>(Component);
+				}
+				else
+				{
+					DebugDrawComponent = NewObject<UDataflowDebugDrawComponent>(RootActor, FName(LOCTEXT("DataflowDebugDrawComponent", "Dataflow Debug Draw Component").ToString(), RF_Transient));
+					DebugDrawComponent->RegisterComponentWithWorld(DataflowScene->GetWorld());
+				}
+			}
+		}
+
+		checkf(DebugDrawComponent, TEXT("Could not create or find a DebugDrawComponentComponent"));
+
+		FDataflowDebugRenderSceneProxy* SceneProxy = static_cast<FDataflowDebugRenderSceneProxy*>(DebugDrawComponent->GetSceneProxy());
+		checkf(SceneProxy, TEXT("Could not find a FDataflowDebugRenderSceneProxy on the DebugDrawComponent"));
+
+		return MakeUnique<FDataflowDebugDraw>(SceneProxy, DataflowScene->ModifySceneElements());
+	}
+
+	/** Update the debug draw component from the selected/pinned nodes matching the NodeType */
+	template<typename NodeType>
+	void UpdateDebugDrawComponent(const TObjectPtr<UDataflowBaseContent>& EditorContent,
+		const FString& RootName, FDataflowPreviewSceneBase* DataflowScene, const bool bIsConstruction, FString& DebugDrawOverlay,
+		const TSharedPtr<FEditorModeTools> EditorModeManager)
+	{
+		if(!DataflowScene || !EditorContent)
+		{
+			return;
+		}
+		DataflowScene->UnregisterSceneElements();
+		DataflowScene->ModifySceneElements().Reset();
+		
+		TObjectPtr<UDataflowDebugDrawComponent>& DebugDrawComponent = DataflowScene->ModifyDebugDrawComponent();
+		const TUniquePtr<IDataflowDebugDrawInterface> DebugDrawObject =
+				CreateDebugDrawComponent(DataflowScene, DebugDrawComponent);
+
+		TSharedPtr<FDataflowBaseElement> RootSceneElement = MakeShared<FDataflowBaseElement>(RootName, nullptr, FBox(ForceInitToZero), bIsConstruction);
+		DataflowScene->ModifySceneElements().Add(RootSceneElement);
+
+		FName CurrentViewMode;
+		if (UDataflowEditorMode* const DataflowMode = Cast<UDataflowEditorMode>(EditorModeManager->GetActiveScriptableMode(UDataflowEditorMode::EM_DataflowEditorModeId)))
+		{
+			if (const UE::Dataflow::IDataflowConstructionViewMode* const ConstructionViewMode = DataflowMode->GetConstructionViewMode())
+			{
+				CurrentViewMode = ConstructionViewMode->GetName();
+			}
+		}
+	
+		if (const TSharedPtr<UE::Dataflow::FEngineContext> DataflowContext = EditorContent->GetDataflowContext())
+		{
+			if (const UDataflowEdNode* const SelectedNode = EditorContent->GetSelectedNode())
+			{
+				if (const TSharedPtr<const FDataflowNode> SelectedDataflowNode = SelectedNode->GetDataflowNode())
+				{
+					if(SelectedDataflowNode->IsA(NodeType::StaticType()))
+					{
+						if (SelectedDataflowNode->CanDebugDrawViewMode(CurrentViewMode))
+						{
+							DebugDrawObject->ResetAllState();
+
+							FDataflowNode::FDebugDrawParameters DebugDrawParameters;
+							DebugDrawParameters.bNodeIsSelected = true;
+							DebugDrawParameters.bNodeIsPinned = false;
+							DebugDrawParameters.CurrentViewMode = CurrentViewMode;
+							SelectedDataflowNode->DebugDraw(*DataflowContext, *DebugDrawObject, DebugDrawParameters);
+						}
+					}
+				}
+			}
+
+			for (const UDataflowEdNode* const PinnedNode : EditorContent->GetDataflowAsset()->GetWireframeRenderTargets())
+			{
+				if (const TSharedPtr<const FDataflowNode> PinnedDataflowNode = PinnedNode->GetDataflowNode())
+				{
+					if(PinnedDataflowNode->IsA(NodeType::StaticType()))
+					{
+						if (PinnedDataflowNode->CanDebugDrawViewMode(CurrentViewMode))
+						{
+							DebugDrawObject->ResetAllState();
+
+							FDataflowNode::FDebugDrawParameters DebugDrawParameters;
+							DebugDrawParameters.bNodeIsSelected = false;
+							DebugDrawParameters.bNodeIsPinned = true;
+							DebugDrawParameters.CurrentViewMode = CurrentViewMode;
+							PinnedDataflowNode->DebugDraw(*DataflowContext, *DebugDrawObject, DebugDrawParameters);
+						}
+					}
+				}
+			}
+		}
+
+		DataflowScene->RegisterSceneElements(bIsConstruction);
+
+		DebugDrawComponent->UpdateBounds();
+		DebugDrawComponent->MarkRenderTransformDirty();
+
+		DebugDrawOverlay = DebugDrawObject->GetOverlayText();
+	}
+}
+
 FDataflowEditorToolkit::FDataflowEditorToolkit(UAssetEditor* InOwningAssetEditor)
 	: FBaseCharacterFXEditorToolkit(InOwningAssetEditor, FName("DataflowEditor"))
 	, DataflowEditor(Cast< UDataflowEditor>(InOwningAssetEditor))
@@ -341,9 +458,13 @@ FDataflowEditorToolkit::FDataflowEditorToolkit(UAssetEditor* InOwningAssetEditor
 
 FDataflowEditorToolkit::~FDataflowEditorToolkit()
 {
-	if (DebugDrawComponent)
+	if (ConstructionScene && ConstructionScene->ModifyDebugDrawComponent())
+	{
+		ConstructionScene->ModifyDebugDrawComponent()->UnregisterComponent();
+	}
+	if (SimulationScene && SimulationScene->ModifyDebugDrawComponent())
 	{
-		DebugDrawComponent->UnregisterComponent();
+		SimulationScene->ModifyDebugDrawComponent()->UnregisterComponent();
 	}
 
 	if (IConsoleVariable* const ConsoleVar = IConsoleManager::Get().FindConsoleVariable(TEXT("p.Dataflow.EnableGraphEval")))
@@ -424,18 +545,37 @@ void FDataflowEditorToolkit::UnregisterContextHandlers()
 
 void FDataflowEditorToolkit::CreateEditorModeManager()
 {
+	auto SetSelectionName = [](const USelection* SelectionObject)
+	{
+		if(SelectionObject)
+		{
+			if(UTypedElementSelectionSet* SelectionSet = SelectionObject->GetElementSelectionSet())
+			{
+				const FName SelectionSetName(FString::Printf(TEXT("DataflowSelectionSet%p"), SelectionSet));
+				SelectionSet->SetNameForTedsIntegration(SelectionSetName);
+			}
+		}
+	};
+	
 	// Setup the construction manager / scene
 	FBaseCharacterFXEditorToolkit::CreateEditorModeManager();
 	static_cast<FDataflowPreviewSceneBase*>(ObjectScene.Get())->GetDataflowModeManager()
 		= StaticCastSharedPtr<FAssetEditorModeManager>(EditorModeManager);
 
+	SetSelectionName(EditorModeManager->GetSelectedActors());
+	SetSelectionName(EditorModeManager->GetSelectedComponents());
+	SetSelectionName(EditorModeManager->GetSelectedObjects());
+
 	// Setup the simulation manager / scene
 	SimulationModeManager = MakeShared<FAssetEditorModeManager>();
 	StaticCastSharedPtr<FAssetEditorModeManager>(SimulationModeManager)->SetPreviewScene(
 		SimulationScene.Get());
-
 	static_cast<FDataflowPreviewSceneBase*>(SimulationScene.Get())->GetDataflowModeManager()
 		= StaticCastSharedPtr<FAssetEditorModeManager>(SimulationModeManager);
+	
+	SetSelectionName(SimulationModeManager->GetSelectedActors());
+	SetSelectionName(SimulationModeManager->GetSelectedComponents());
+	SetSelectionName(SimulationModeManager->GetSelectedObjects());
 }
 
 void FDataflowEditorToolkit::NotifyPreChange(FEditPropertyChain* PropertyAboutToChange)
@@ -590,6 +730,7 @@ void FDataflowEditorToolkit::PostInitAssetEditor()
 
 	FDataflowConstructionViewportClient* ConstructionViewportClient = static_cast<FDataflowConstructionViewportClient*>(ViewportClient.Get());
 	OnConstructionSelectionChangedDelegateHandle = ConstructionViewportClient->OnSelectionChangedMulticast.AddSP(this, &FDataflowEditorToolkit::OnConstructionViewSelectionChanged);
+	OnSimulationSelectionChangedDelegateHandle = SimulationViewportClient->OnSelectionChangedMulticast.AddSP(this, &FDataflowEditorToolkit::OnSimulationViewSelectionChanged);
 
 	// Populate editor toolbar
 
@@ -988,92 +1129,23 @@ void FDataflowEditorToolkit::CreateSimulationViewportClient()
 
 //~ End FBaseAssetToolkit overrides
 
-
-TUniquePtr<IDataflowDebugDrawInterface> FDataflowEditorToolkit::CreateDebugDrawObject()
+void FDataflowEditorToolkit::UpdateDebugDraw()
 {
-	if (ConstructionScene && !DebugDrawComponent)
+	if(TObjectPtr<UDataflowBaseContent>& EditorContent = GetEditorContent())
 	{
-		if (TObjectPtr<AActor> RootActor = ConstructionScene->GetRootActor())
+		if(EditorContent->GetDataflowAsset() && EditorContent->GetDataflowAsset()->Type == EDataflowType::Construction)
 		{
-			if (UDataflowDebugDrawComponent* Component = Cast<UDataflowDebugDrawComponent>(RootActor->FindComponentByClass<UDataflowDebugDrawComponent>()))
-			{
-				DebugDrawComponent = Cast<UDataflowDebugDrawComponent>(Component);
-			}
-			else
-			{
-				DebugDrawComponent = NewObject<UDataflowDebugDrawComponent>(RootActor, FName(LOCTEXT("DataflowDebugDrawComponent", "Dataflow Debug Draw Component").ToString(), RF_Transient));
-				DebugDrawComponent->RegisterComponentWithWorld(ConstructionScene->GetWorld());
-			}
+			static const FString RootElementsName = "Construction Elements";
+			UE::Dataflow::Private::UpdateDebugDrawComponent<FDataflowNode>(
+				EditorContent, RootElementsName, ConstructionScene, true, DebugDrawOverlayString, EditorModeManager);
 		}
-	}
-
-	checkf(DebugDrawComponent, TEXT("Could not create or find a DebugDrawComponentComponent"));
-
-	FDataflowDebugRenderSceneProxy* SceneProxy = static_cast<FDataflowDebugRenderSceneProxy*>(DebugDrawComponent->GetSceneProxy());
-	checkf(SceneProxy, TEXT("Could not find a FDataflowDebugRenderSceneProxy on the DebugDrawComponent"));
-
-	return MakeUnique<FDataflowDebugDraw>(SceneProxy);
-}
-
-void FDataflowEditorToolkit::UpdateDebugDraw()
-{
-	const TUniquePtr<IDataflowDebugDrawInterface> DebugDrawObject = CreateDebugDrawObject();
-
-	if (const TObjectPtr<UDataflowBaseContent>& EditorContent = GetEditorContent())
-	{
-		if (const TSharedPtr<UE::Dataflow::FEngineContext> DataflowContext = EditorContent->GetDataflowContext())
+		else
 		{
-			FName CurrentViewMode;
-			if (UDataflowEditorMode* const DataflowMode = Cast<UDataflowEditorMode>(EditorModeManager->GetActiveScriptableMode(UDataflowEditorMode::EM_DataflowEditorModeId)))
-			{
-				if (const UE::Dataflow::IDataflowConstructionViewMode* const ConstructionViewMode = DataflowMode->GetConstructionViewMode())
-				{
-					CurrentViewMode = ConstructionViewMode->GetName();
-				}
-			}
-
-			if (const UDataflowEdNode* const SelectedNode = EditorContent->GetSelectedNode())
-			{
-				if (const TSharedPtr<const FDataflowNode> SelectedDataflowNode = SelectedNode->GetDataflowNode())
-				{
-					if (SelectedDataflowNode->CanDebugDrawViewMode(CurrentViewMode))
-					{
-						DebugDrawObject->ResetAllState();
-
-						FDataflowNode::FDebugDrawParameters DebugDrawParameters;
-						DebugDrawParameters.bNodeIsSelected = true;
-						DebugDrawParameters.bNodeIsPinned = false;
-						DebugDrawParameters.CurrentViewMode = CurrentViewMode;
-
-						SelectedDataflowNode->DebugDraw(*DataflowContext, *DebugDrawObject, DebugDrawParameters);
-					}
-				}
-			}
-
-			for (const UDataflowEdNode* const PinnedNode : EditorContent->GetDataflowAsset()->GetWireframeRenderTargets())
-			{
-				if (const TSharedPtr<const FDataflowNode> PinnedDataflowNode = PinnedNode->GetDataflowNode())
-				{
-					if (PinnedDataflowNode->CanDebugDrawViewMode(CurrentViewMode))
-					{
-						DebugDrawObject->ResetAllState();
-
-						FDataflowNode::FDebugDrawParameters DebugDrawParameters;
-						DebugDrawParameters.bNodeIsSelected = false;
-						DebugDrawParameters.bNodeIsPinned = true;
-						DebugDrawParameters.CurrentViewMode = CurrentViewMode;
-						PinnedDataflowNode->DebugDraw(*DataflowContext, *DebugDrawObject, DebugDrawParameters);
-					}
-				}
-			}
+			static const FString RootElementsName = "Simulation Elements";
+			UE::Dataflow::Private::UpdateDebugDrawComponent<FDataflowSimulationNode>(
+				EditorContent, RootElementsName, SimulationScene.Get(), false, DebugDrawOverlayString, SimulationModeManager);
 		}
 	}
-
-	DebugDrawComponent->UpdateBounds();
-	DebugDrawComponent->MarkRenderTransformDirty();
-
-	// Build DebugDraw overlay string from the DebugDraw interface
-	DebugDrawOverlayString = DebugDrawObject->GetOverlayText();
 }
 
 namespace UE::Dataflow::Private
@@ -1549,11 +1621,19 @@ void FDataflowEditorToolkit::OnNodeDeleted(const TSet<UObject*>& NewSelection)
 	}
 }
 
-void FDataflowEditorToolkit::OnConstructionViewSelectionChanged(const TArray<UPrimitiveComponent*>& SelectedComponents)
+void FDataflowEditorToolkit::OnConstructionViewSelectionChanged(const TArray<UPrimitiveComponent*>& SelectedComponents, const TArray<FDataflowBaseElement*>& SelectedElements)
 {
 	for(IDataflowViewListener* Listener : ViewListeners)
 	{
-		Listener->OnConstructionViewSelectionChanged(SelectedComponents);
+		Listener->OnConstructionViewSelectionChanged(SelectedComponents, SelectedElements);
+	}
+}
+
+void FDataflowEditorToolkit::OnSimulationViewSelectionChanged(const TArray<UPrimitiveComponent*>& SelectedComponents, const TArray<FDataflowBaseElement*>& SelectedElements)
+{
+	for(IDataflowViewListener* Listener : ViewListeners)
+	{
+		Listener->OnSimulationViewSelectionChanged(SelectedComponents, SelectedElements);
 	}
 }
 
@@ -1766,6 +1846,12 @@ void FDataflowEditorToolkit::Tick(float DeltaTime)
 			DataflowEditor->UpdateTerminalContents(InitTimeStamp);
 		}
 	}
+	if(ConstructionScene->IsSceneDirty() || SimulationScene->IsSceneDirty())
+	{
+		bViewsNeedRefresh = true;
+		ConstructionScene->ResetDirtyFlag();
+		SimulationScene->ResetDirtyFlag();
+	}
 
 	RefreshViewsIfNeeded();
 }
@@ -1977,7 +2063,7 @@ TSharedRef<SDockTab> FDataflowEditorToolkit::SpawnTab_PreviewScene(const FSpawnT
 	check(Args.GetTabId() == PreviewSceneTabId);
 
 	return SNew(SDockTab)
-		.Label(LOCTEXT("DataflowEditor_PreviewScene_TabTitle", "PreviewScene"))
+		.Label(LOCTEXT("DataflowEditor_PreviewScene_TabTitle", "Preview Scene"))
 		[
 			AdvancedPreviewSettingsWidget->AsShared()
 		];
@@ -1988,7 +2074,7 @@ TSharedRef<SDockTab> FDataflowEditorToolkit::SpawnTab_GraphCanvas(const FSpawnTa
 	check(Args.GetTabId() == GraphCanvasTabId);
 
 	TSharedRef<SDockTab> SpawnedTab = SNew(SDockTab)
-		.Label(LOCTEXT("DataflowEditor_Dataflow_TabTitle", "Dataflow Graph"))
+		.Label(LOCTEXT("DataflowEditor_DataflowGraph_TabTitle", "Dataflow Graph"))
 		[
 			GraphEditor.ToSharedRef()
 		];
@@ -2009,7 +2095,7 @@ TSharedRef<SDockTab> FDataflowEditorToolkit::SpawnTab_SubGraphTab(const FSpawnTa
 		}
 	}
 	TSharedRef<SDockTab> SpawnedTab = SNew(SDockTab)
-		.Label(LOCTEXT("DataflowEditor_Subgraph_TabTitle", "Dataflow Graph"))
+		.Label(LOCTEXT("DataflowEditor_DataflowSubgraph_TabTitle", "Dataflow Subgraph"))
 		[
 			SubgraphEditor.ToSharedRef()
 		];
@@ -2049,12 +2135,30 @@ TSharedRef<SDockTab> FDataflowEditorToolkit::SpawnTab_SkeletonView(const FSpawnT
 
 	TSharedPtr<ISkeletonTree> SkeletonEditor = SkeletonEditorView->CreateEditor(SkeletonTreeArgs);
 	return SNew(SDockTab)
-		.Label(LOCTEXT("DataflowEditor_Outliner_TabTitle", "Outliner"))
+		.Label(LOCTEXT("DataflowEditor_SkeletonTree_TabTitle", "Skeleton Tree"))
 		[
 			SkeletonEditor.ToSharedRef()
 		];
 }
 
+TSharedRef<SDockTab> FDataflowEditorToolkit::SpawnTab_OutlinerView(const FSpawnTabArgs& Args)
+{
+	check(Args.GetTabId() == OutlinerViewTabId);
+	check(DataflowEditor);
+	check(DataflowEditor->GetEditorContent());
+
+	DataflowOutlinerView = MakeShared<FDataflowOutlinerView>(ConstructionScene, SimulationScene.Get(), DataflowEditor->GetEditorContent());
+	ViewListeners.Add(DataflowOutlinerView.Get());
+
+	TSharedPtr<ISceneOutliner> DataflowOutliner = DataflowOutlinerView->CreateWidget();
+	
+	return SNew(SDockTab)
+		.Label(LOCTEXT("DataflowEditor_SceneOutliner_TabTitle", "Scene Outliner"))
+		[
+			DataflowOutliner.ToSharedRef()
+		];
+}
+
 TSharedRef<SDockTab> FDataflowEditorToolkit::SpawnTab_SelectionView(const FSpawnTabArgs& Args)
 {
 	//	check(Args.GetTabId().TabType == SelectionViewTabId_1);
@@ -2217,7 +2321,7 @@ TSharedPtr<SWidget> FDataflowEditorToolkit::CreateSimulationVisualizationWidget(
 TSharedRef<SDockTab> FDataflowEditorToolkit::SpawnTab_SimulationVisualization(const FSpawnTabArgs& Args)
 {
 	TSharedRef<SDockTab> SimulationVisualizationTab = SNew(SDockTab)
-		.Label(LOCTEXT("SimulationVisualizationTitle", "Simulation Visualization"));
+		.Label(LOCTEXT("DataflowEditor_SimulationVisualisation_TabTitle", "Simulation Visualization"));
 
 	SimulationVisualizationWidget = CreateSimulationVisualizationWidget();
 	SimulationVisualizationTab->SetContent(SimulationVisualizationWidget.ToSharedRef());
@@ -2236,7 +2340,7 @@ TSharedRef<SDockTab> FDataflowEditorToolkit::SpawnTab_SimulationVisualization(co
 TSharedRef<SDockTab> FDataflowEditorToolkit::SpawnTab_MembersWidget(const FSpawnTabArgs& Args)
 {
 	TSharedRef<SDockTab> DataflowMembersTab = SNew(SDockTab)
-		.Label(LOCTEXT("MembersWidgetTitle", "Dataflow Members"));
+		.Label(LOCTEXT("DataflowEditor_DataflowMembers_TabTitle", "Dataflow Members"));
 
 	MembersWidget = CreateDataflowMembersWidget();
 	DataflowMembersTab->SetContent(MembersWidget.ToSharedRef());
@@ -2280,7 +2384,7 @@ void FDataflowEditorToolkit::RegisterTabSpawners(const TSharedRef<FTabManager>&
 	const TSharedRef<FWorkspaceItem> CollectionSpreadSheetWorkspaceMenuCategoryRef = EditorMenuCategory->AddGroup(LOCTEXT("WorkspaceMenu_CollectionSpreadSheet", "Collection SpreadSheet"), FSlateIcon(FAppStyle::GetAppStyleSetName(), "LevelEditor.Tabs.Outliner"));
 	
 	InTabManager->RegisterTabSpawner(ViewportTabID, FOnSpawnTab::CreateSP(this, &FDataflowEditorToolkit::SpawnTab_Viewport))
-		.SetDisplayName(LOCTEXT("DataflowViewportTab", "Construction Viewport"))
+		.SetDisplayName(LOCTEXT("ConstructionViewportTab", "Construction Viewport"))
 		.SetGroup(EditorMenuCategory.ToSharedRef())
 		.SetIcon(FSlateIcon(FAppStyle::GetAppStyleSetName(), "LevelEditor.Tabs.Viewports"));
 
@@ -2295,17 +2399,17 @@ void FDataflowEditorToolkit::RegisterTabSpawners(const TSharedRef<FTabManager>&
 		.SetIcon(FSlateIcon(FAppStyle::GetAppStyleSetName(), "LevelEditor.Tabs.Details"));
 
 	InTabManager->RegisterTabSpawner(PreviewSceneTabId, FOnSpawnTab::CreateSP(this, &FDataflowEditorToolkit::SpawnTab_PreviewScene))
-		.SetDisplayName(LOCTEXT("PreviewSceneTab", "PreviewScene"))
+		.SetDisplayName(LOCTEXT("PreviewSceneTab", "Preview Scene"))
 		.SetGroup(EditorMenuCategory.ToSharedRef())
-		.SetIcon(FSlateIcon(FAppStyle::GetAppStyleSetName(), "LevelEditor.Tabs.PreviewScene"));
+		.SetIcon(FSlateIcon(FAppStyle::GetAppStyleSetName(), "LevelEditor.Tabs.Details"));
 	
 	InTabManager->RegisterTabSpawner(GraphCanvasTabId, FOnSpawnTab::CreateSP(this, &FDataflowEditorToolkit::SpawnTab_GraphCanvas))
-		.SetDisplayName(LOCTEXT("DataflowTab", "Dataflow Graph"))
+		.SetDisplayName(LOCTEXT("DataflowGraphTab", "Dataflow Graph"))
 		.SetGroup(EditorMenuCategory.ToSharedRef())
 		.SetIcon(FSlateIcon(FAppStyle::GetAppStyleSetName(), "GraphEditor.EventGraph_16x"));
 
 	InTabManager->RegisterTabSpawner(SubGraphCanvasTabId, FOnSpawnTab::CreateSP(this, &FDataflowEditorToolkit::SpawnTab_SubGraphTab))
-		.SetDisplayName(LOCTEXT("SubgraphTab", "Dataflow Subgraph"))
+		.SetDisplayName(LOCTEXT("DataflowSubgraphTab", "Dataflow Subgraph"))
 		.SetGroup(EditorMenuCategory.ToSharedRef())
 		.SetIcon(FSlateIcon(FAppStyle::GetAppStyleSetName(), "GraphEditor.EventGraph_16x"))
 		.SetMenuType(ETabSpawnerMenuType::Hidden); // hide it from menus as those tabs can only be open by the toolkit itself 
@@ -2318,60 +2422,65 @@ void FDataflowEditorToolkit::RegisterTabSpawners(const TSharedRef<FTabManager>&
 	if (bDataflowEnableSkeletonView)
 	{
 		InTabManager->RegisterTabSpawner(SkeletonViewTabId, FOnSpawnTab::CreateSP(this, &FDataflowEditorToolkit::SpawnTab_SkeletonView))
-			.SetDisplayName(LOCTEXT("OutlinerTab", "Outliner"))
+			.SetDisplayName(LOCTEXT("SkeletonTreeTab", "Skeleton Tree"))
 			.SetGroup(EditorMenuCategory.ToSharedRef())
-			.SetIcon(FSlateIcon(FAppStyle::GetAppStyleSetName(), "LevelEditor.Tabs.SkeletonHierarchy"));
+			.SetIcon(FSlateIcon(FAppStyle::GetAppStyleSetName(), "LevelEditor.Tabs.Outliner"));
 	}
+	
+	InTabManager->RegisterTabSpawner(OutlinerViewTabId, FOnSpawnTab::CreateSP(this, &FDataflowEditorToolkit::SpawnTab_OutlinerView))
+			.SetDisplayName(LOCTEXT("SceneOutlinerTab", "Scene Outliner"))
+			.SetGroup(EditorMenuCategory.ToSharedRef())
+			.SetIcon(FSlateIcon(FAppStyle::GetAppStyleSetName(), "LevelEditor.Tabs.Outliner"));
 
 	InTabManager->RegisterTabSpawner(SelectionViewTabId_1, FOnSpawnTab::CreateSP(this, &FDataflowEditorToolkit::SpawnTab_SelectionView))
-		.SetDisplayName(LOCTEXT("DataflowSelectionViewTab1", "Selection View 1"))
+		.SetDisplayName(LOCTEXT("SelectionViewTab1", "Selection View 1"))
 		.SetGroup(SelectionViewWorkspaceMenuCategoryRef)
 		.SetIcon(FSlateIcon(FAppStyle::GetAppStyleSetName(), "LevelEditor.Tabs.Outliner"));
 
 	InTabManager->RegisterTabSpawner(SelectionViewTabId_2, FOnSpawnTab::CreateSP(this, &FDataflowEditorToolkit::SpawnTab_SelectionView))
-		.SetDisplayName(LOCTEXT("DataflowSelectionViewTab2", "Selection View 2"))
+		.SetDisplayName(LOCTEXT("SelectionViewTab2", "Selection View 2"))
 		.SetGroup(SelectionViewWorkspaceMenuCategoryRef)
 		.SetIcon(FSlateIcon(FAppStyle::GetAppStyleSetName(), "LevelEditor.Tabs.Outliner"));
 
 	InTabManager->RegisterTabSpawner(SelectionViewTabId_3, FOnSpawnTab::CreateSP(this, &FDataflowEditorToolkit::SpawnTab_SelectionView))
-		.SetDisplayName(LOCTEXT("DataflowSelectionViewTab3", "Selection View 3"))
+		.SetDisplayName(LOCTEXT("SelectionViewTab3", "Selection View 3"))
 		.SetGroup(SelectionViewWorkspaceMenuCategoryRef)
 		.SetIcon(FSlateIcon(FAppStyle::GetAppStyleSetName(), "LevelEditor.Tabs.Outliner"));
 
 	InTabManager->RegisterTabSpawner(SelectionViewTabId_4, FOnSpawnTab::CreateSP(this, &FDataflowEditorToolkit::SpawnTab_SelectionView))
-		.SetDisplayName(LOCTEXT("DataflowSelectionViewTab4", "Selection View 4"))
+		.SetDisplayName(LOCTEXT("SelectionViewTab4", "Selection View 4"))
 		.SetGroup(SelectionViewWorkspaceMenuCategoryRef)
 		.SetIcon(FSlateIcon(FAppStyle::GetAppStyleSetName(), "LevelEditor.Tabs.Outliner"));
 
 	InTabManager->RegisterTabSpawner(CollectionSpreadSheetTabId_1, FOnSpawnTab::CreateSP(this, &FDataflowEditorToolkit::SpawnTab_CollectionSpreadSheet))
-		.SetDisplayName(LOCTEXT("DataflowCollectionSpreadSheetTab1", "Collection SpreadSheet 1"))
+		.SetDisplayName(LOCTEXT("CollectionSpreadSheetTab1", "Collection SpreadSheet 1"))
 		.SetGroup(CollectionSpreadSheetWorkspaceMenuCategoryRef)
 		.SetIcon(FSlateIcon(FAppStyle::GetAppStyleSetName(), "LevelEditor.Tabs.Outliner"));
 
 	InTabManager->RegisterTabSpawner(CollectionSpreadSheetTabId_2, FOnSpawnTab::CreateSP(this, &FDataflowEditorToolkit::SpawnTab_CollectionSpreadSheet))
-		.SetDisplayName(LOCTEXT("DataflowCollectionSpreadSheetTab2", "Collection SpreadSheet 2"))
+		.SetDisplayName(LOCTEXT("CollectionSpreadSheetTab2", "Collection SpreadSheet 2"))
 		.SetGroup(CollectionSpreadSheetWorkspaceMenuCategoryRef)
 		.SetIcon(FSlateIcon(FAppStyle::GetAppStyleSetName(), "LevelEditor.Tabs.Outliner"));
 
 	InTabManager->RegisterTabSpawner(CollectionSpreadSheetTabId_3, FOnSpawnTab::CreateSP(this, &FDataflowEditorToolkit::SpawnTab_CollectionSpreadSheet))
-		.SetDisplayName(LOCTEXT("DataflowCollectionSpreadSheetTab3", "Collection SpreadSheet 3"))
+		.SetDisplayName(LOCTEXT("CollectionSpreadSheetTab3", "Collection SpreadSheet 3"))
 		.SetGroup(CollectionSpreadSheetWorkspaceMenuCategoryRef)
 		.SetIcon(FSlateIcon(FAppStyle::GetAppStyleSetName(), "LevelEditor.Tabs.Outliner"));
 
 	InTabManager->RegisterTabSpawner(CollectionSpreadSheetTabId_4, FOnSpawnTab::CreateSP(this, &FDataflowEditorToolkit::SpawnTab_CollectionSpreadSheet))
-		.SetDisplayName(LOCTEXT("DataflowCollectionSpreadSheetTab4", "Collection SpreadSheet 4"))
+		.SetDisplayName(LOCTEXT("CollectionSpreadSheetTab4", "Collection SpreadSheet 4"))
 		.SetGroup(CollectionSpreadSheetWorkspaceMenuCategoryRef)
 		.SetIcon(FSlateIcon(FAppStyle::GetAppStyleSetName(), "LevelEditor.Tabs.Outliner"));
 
 	InTabManager->RegisterTabSpawner(SimulationVisualizationTabId, FOnSpawnTab::CreateSP(this, &FDataflowEditorToolkit::SpawnTab_SimulationVisualization))
-		.SetDisplayName(LOCTEXT("SimulationVisualizationTabDisplayName", "Simulation Visualization"))
-		.SetGroup(AssetEditorTabsCategory.ToSharedRef());
-
+		.SetDisplayName(LOCTEXT("SimulationVisualizationTab", "Simulation Visualization"))
+		.SetGroup(AssetEditorTabsCategory.ToSharedRef())
+		.SetIcon(FSlateIcon(FAppStyle::GetAppStyleSetName(), "LevelEditor.Tabs.Outliner"));
 	
 	InTabManager->RegisterTabSpawner(MembersWidgetTabId, FOnSpawnTab::CreateSP(this, &FDataflowEditorToolkit::SpawnTab_MembersWidget))
-		.SetDisplayName(LOCTEXT("MembersWidgetTabDisplayName", "Dataflow Members"))
-		.SetGroup(AssetEditorTabsCategory.ToSharedRef());
-
+		.SetDisplayName(LOCTEXT("DataflowMembersTab", "Dataflow Members"))
+		.SetGroup(AssetEditorTabsCategory.ToSharedRef())
+		.SetIcon(FSlateIcon(FAppStyle::GetAppStyleSetName(), "LevelEditor.Tabs.Outliner"));
 
 	// Create a document manager to be able to spawn subgraph tabs
 	DocumentManager = MakeShareable(new FDocumentTracker(SubGraphCanvasTabId));
@@ -2386,7 +2495,7 @@ void FDataflowEditorToolkit::RegisterTabSpawners(const TSharedRef<FTabManager>&
 	DocumentManager->RegisterDocumentFactory(SubGraphTabSummoner);
 	
 	InTabManager->RegisterTabSpawner(OutputLogTabId, FOnSpawnTab::CreateSP(this, &FDataflowEditorToolkit::SpawnTab_OutputLog))
-		.SetDisplayName(LOCTEXT("OutputLogTab", "Dataflow Output Log"))
+		.SetDisplayName(LOCTEXT("OutputLogTab", "Output Log"))
 		.SetGroup(EditorMenuCategory.ToSharedRef())
 		.SetIcon(FSlateIcon(FAppStyle::GetAppStyleSetName(), "Log.TabIcon"));
 }
@@ -2403,6 +2512,7 @@ void FDataflowEditorToolkit::UnregisterTabSpawners(const TSharedRef<class FTabMa
 	InTabManager->UnregisterTabSpawner(SubGraphCanvasTabId);
 	InTabManager->UnregisterTabSpawner(NodeDetailsTabId);
 	InTabManager->UnregisterTabSpawner(SkeletonViewTabId);
+	InTabManager->UnregisterTabSpawner(OutlinerViewTabId);
 	InTabManager->UnregisterTabSpawner(SelectionViewTabId_1);
 	InTabManager->UnregisterTabSpawner(SelectionViewTabId_2);
 	InTabManager->UnregisterTabSpawner(SelectionViewTabId_3);
@@ -2450,7 +2560,7 @@ void FDataflowEditorToolkit::OnTabClosed(TSharedRef<SDockTab> Tab)
 	{
 		ViewListeners.Remove(DataflowCollectionSpreadSheet_4.Get());
 	}
-	else if (Tab->GetTabLabel().EqualTo(FText::FromString("Skeleton View")))
+	else if (Tab->GetTabLabel().EqualTo(FText::FromString("Skeleton Tree")))
 	{
 		ViewListeners.Remove(SkeletonEditorView.Get());
 	}
@@ -2586,7 +2696,14 @@ void FDataflowEditorToolkit::AddReferencedObjects(FReferenceCollector& Collector
 {
 	Collector.AddReferencedObjects(SelectedDataflowNodes);
 	Collector.AddReferencedObject(PrimarySelection);
-	Collector.AddReferencedObject(DebugDrawComponent);
+	if(ConstructionScene)
+	{
+		Collector.AddReferencedObject(ConstructionScene->ModifyDebugDrawComponent());
+	}
+	if(SimulationScene)
+	{
+		Collector.AddReferencedObject(SimulationScene->ModifyDebugDrawComponent());
+	}
 }
 
 const FString& FDataflowEditorToolkit::GetDebugDrawOverlayString() const
diff --git a/Engine/Plugins/Experimental/Dataflow/Source/DataflowEditor/Private/Dataflow/DataflowEditorViewportClientBase.cpp b/Engine/Plugins/Experimental/Dataflow/Source/DataflowEditor/Private/Dataflow/DataflowEditorViewportClientBase.cpp
--- a/Engine/Plugins/Experimental/Dataflow/Source/DataflowEditor/Private/Dataflow/DataflowEditorViewportClientBase.cpp
+++ b/Engine/Plugins/Experimental/Dataflow/Source/DataflowEditor/Private/Dataflow/DataflowEditorViewportClientBase.cpp
@@ -6,6 +6,8 @@
 #include "Tools/EdModeInteractiveToolsContext.h"
 #include "BaseBehaviors/ClickDragBehavior.h"
 #include "AssetViewerSettings.h"
+#include "Dataflow/DataflowElement.h"
+#include "Dataflow/DataflowElement.h"
 
 #define LOCTEXT_NAMESPACE "DataflowEditorViewportClientBase"
 
@@ -107,11 +109,62 @@ void FDataflowEditorViewportClientBase::RegisterDelegates()
 			SyncEngineShowFlags();
 		}
 	});
+
+	if(DataflowPreviewScene)
+	{
+		OnFocusRequestDelegateHandle = DataflowPreviewScene->OnFocusRequest().AddRaw(this, &FDataflowEditorViewportClientBase::HandleFocusRequest);
+	}
 }
 
 void FDataflowEditorViewportClientBase::DeregisterDelegates()
 {
 	UAssetViewerSettings::Get()->OnAssetViewerSettingsChanged().Remove(OnAssetViewerSettingsChangedDelegateHandle);
+
+	if(DataflowPreviewScene)
+	{
+		DataflowPreviewScene->OnFocusRequest().Remove(OnFocusRequestDelegateHandle);
+	}
+}
+
+void FDataflowEditorViewportClientBase::HandleFocusRequest(const FBox& BoundingBox)
+{
+	FocusViewportOnBox(BoundingBox);
+}
+
+void FDataflowEditorViewportClientBase::Draw(const FSceneView* View, FPrimitiveDrawInterface* PDI)
+{
+	Super::Draw(View, PDI);
+}
+
+void FDataflowEditorViewportClientBase::GetSelectedElements(HHitProxy* HitProxy, TArray<FDataflowBaseElement*>& SelectedElements) const
+{
+	TArray<FDataflowBaseElement*> DataflowElements;
+	if(DataflowPreviewScene)
+	{
+		for(IDataflowDebugDrawInterface::FDataflowElementsType::ElementType& DataflowElement : DataflowPreviewScene->ModifySceneElements())
+		{
+			if(DataflowElement.IsValid())
+			{
+				DataflowElement->bIsSelected = false;
+			}
+		}
+		if (HitProxy && HitProxy->IsA(HDataflowElementHitProxy::StaticGetType()))
+		{
+			if(const HDataflowElementHitProxy* DataflowElementProxy = static_cast<HDataflowElementHitProxy*>(HitProxy))
+			{
+				if(DataflowElementProxy->ElementIndex < DataflowPreviewScene->ModifySceneElements().Num())
+				{
+					IDataflowDebugDrawInterface::FDataflowElementsType::ElementType& DataflowElement = DataflowPreviewScene->ModifySceneElements()[DataflowElementProxy->ElementIndex];
+
+					if(DataflowElement.IsValid())
+					{
+						DataflowElement->bIsSelected = true;
+						SelectedElements.Add(DataflowElement.Get());
+					}
+				}
+			}
+		}
+	}
 }
 
 #undef LOCTEXT_NAMESPACE 
diff --git a/Engine/Plugins/Experimental/Dataflow/Source/DataflowEditor/Private/Dataflow/DataflowOutlinerFactory.cpp b/Engine/Plugins/Experimental/Dataflow/Source/DataflowEditor/Private/Dataflow/DataflowOutlinerFactory.cpp
new file mode 100644
--- /dev/null
+++ b/Engine/Plugins/Experimental/Dataflow/Source/DataflowEditor/Private/Dataflow/DataflowOutlinerFactory.cpp
@@ -0,0 +1,282 @@
+﻿// Copyright Epic Games, Inc. All Rights Reserved.
+
+#include "DataflowOutlinerFactory.h"
+#include "Dataflow/DataflowOutlinerFactory.h"
+#include "Dataflow/DataflowContent.h"
+#include "Dataflow/DataflowEditorPreviewSceneBase.h"
+#include "Dataflow/DataflowElement.h"
+#include "Elements/Columns/TypedElementCompatibilityColumns.h"
+#include "Elements/Columns/TypedElementHiearchyColumns.h"
+#include "Elements/Columns/TypedElementLabelColumns.h"
+#include "Elements/Columns/TypedElementMiscColumns.h"
+#include "Elements/Columns/TypedElementTypeInfoColumns.h"
+#include "Elements/Columns/TypedElementVisibilityColumns.h"
+#include "Elements/Framework/TypedElementQueryBuilder.h"
+
+namespace UE::Dataflow::Private
+{
+	/** Command to set visibility onto the actors */
+	struct FSetActorVisibilityCommand
+	{
+		void operator()()
+		{
+			if (TStrongObjectPtr<AActor> PinnedActor = SceneActor.Pin())
+			{
+				SceneActor->SetIsTemporarilyHiddenInEditor(!bIsVisible);
+			}
+			
+		}
+		TWeakObjectPtr<AActor> SceneActor;
+		bool bIsVisible;
+	};
+
+	/** Command to set visibility onto the components */
+	struct FSetComponentVisibilityCommand
+	{
+		void operator()()
+		{
+			if (TStrongObjectPtr<USceneComponent> PinnedComponent = SceneComponent.Pin())
+			{
+				PinnedComponent->SetVisibility(bIsVisible);
+			}
+		}
+		TWeakObjectPtr<USceneComponent> SceneComponent;
+		bool bIsVisible;
+	};
+
+	/** Command to set visibility onto the elements */
+	struct FSetElementVisibilityCommand
+	{
+		void operator()()
+		{
+			if (SceneElement)
+			{
+				SceneElement->bIsVisible = bIsVisible;
+			}
+		}
+		FDataflowBaseElement* SceneElement;
+		bool bIsVisible;
+	};
+}
+
+void UDataflowObjectFactory::RegisterTables(UE::Editor::DataStorage::ICoreProvider& DataStorage, UE::Editor::DataStorage::ICompatibilityProvider& DataStorageCompatibility)
+{
+	Super::RegisterTables(DataStorage, DataStorageCompatibility);
+}
+
+void UDataflowObjectFactory::RegisterHierarchyQueries(UE::Editor::DataStorage::ICoreProvider& DataStorage)
+{
+	using namespace UE::Editor::DataStorage::Queries;
+	using namespace UE::Editor::DataStorage;
+	
+	DataStorage.RegisterQuery(
+	Select(
+		TEXT("Sync dataflow object hierarchy to column"),
+		FProcessor(EQueryTickPhase::PrePhysics, DataStorage.GetQueryTickGroupName(EQueryTickGroups::SyncExternalToDataStorage))
+			.SetExecutionMode(EExecutionMode::GameThread),
+		[&DataStorage](IQueryContext& Context, RowHandle ChildHandle,
+			const FTypedElementUObjectColumn& RawObject, FTableRowParentColumn& ParentColumn)
+		{
+			UObject* ParentObject = nullptr;
+			if (const USceneComponent* SceneComponent = Cast<USceneComponent>(RawObject.Object))
+			{
+				if (USceneComponent* ParentComponent = SceneComponent->GetAttachParent())
+				{
+					ParentObject = ParentComponent;
+				}
+				else if(AActor* ParentActor = SceneComponent->GetOwner())
+				{
+					ParentObject = ParentActor;
+				}	
+			}
+			else if (const AActor* SceneActor = Cast<AActor>(RawObject.Object))
+			{
+				if(AActor* ParentActor = SceneActor->GetOwner())
+				{
+					ParentObject = ParentActor;
+				}
+			}
+			if(ParentObject)
+			{
+				FMapKeyView IdKey = FMapKeyView(ParentObject);
+				RowHandle ParentHandle = Context.LookupMappedRow(IdKey);
+
+				ParentColumn.Parent = ParentHandle;
+			}
+		})
+	.Where(TColumn<FDataflowSceneObjectTag>() && TColumn<FTypedElementSyncFromWorldTag>())
+	.Compile());
+
+	DataStorage.RegisterQuery(
+		Select(
+			TEXT("Sync dataflow struct hierarchy to column"),
+			FProcessor(EQueryTickPhase::PrePhysics, DataStorage.GetQueryTickGroupName(EQueryTickGroups::SyncExternalToDataStorage))
+				.SetExecutionMode(EExecutionMode::GameThread),
+			[&DataStorage](IQueryContext& Context, RowHandle ChildHandle,
+				const FTypedElementExternalObjectColumn& RawObject, const FTypedElementScriptStructTypeInfoColumn& TypeInfo,  FTableRowParentColumn& ParentColumn)
+			{
+				if (RawObject.Object && TypeInfo.TypeInfo->IsChildOf(FDataflowBaseElement::StaticStruct()))
+				{
+					FDataflowBaseElement* SceneObject = static_cast<FDataflowBaseElement*>(RawObject.Object);
+
+					if(SceneObject->ParentElement)
+					{
+						FMapKeyView IdKey = FMapKeyView(SceneObject->ParentElement);
+						RowHandle ParentHandle = Context.LookupMappedRow(IdKey);
+
+						ParentColumn.Parent = ParentHandle;
+					}
+				}
+			})
+		.Where(TColumn<FDataflowSceneStructTag>() && TColumn<FTypedElementSyncFromWorldTag>())
+		.Compile());
+}
+
+void UDataflowObjectFactory::RegisterLabelQueries(UE::Editor::DataStorage::ICoreProvider& DataStorage)
+{
+	using namespace UE::Editor::DataStorage::Queries;
+	using namespace UE::Editor::DataStorage;
+
+	DataStorage.RegisterQuery(
+	Select(
+		TEXT("Sync dataflow object label to column"),
+		FProcessor(EQueryTickPhase::PostPhysics, DataStorage.GetQueryTickGroupName(EQueryTickGroups::SyncExternalToDataStorage))
+			.SetExecutionMode(EExecutionMode::GameThread),
+		[](IQueryContext& Context, RowHandle RowHandle, 
+			const FTypedElementUObjectColumn& RawObject, FTypedElementLabelColumn& ObjectLabel, FTypedElementLabelHashColumn& LabelHash)
+		{
+			if (const UObject* SceneObject = RawObject.Object.Get())
+			{
+				const FString ObjectLabelName = SceneObject->GetName();
+				const uint64 ObjectLabelHash = CityHash64(reinterpret_cast<const char*>(*ObjectLabelName), ObjectLabelName.Len() * sizeof(**ObjectLabelName));
+				if (LabelHash.LabelHash != ObjectLabelHash)
+				{
+					LabelHash.LabelHash = ObjectLabelHash;
+					ObjectLabel.Label = ObjectLabelName;
+				}
+			}
+		})
+	.Where(TColumn<FDataflowSceneObjectTag>() && TColumn<FTypedElementSyncFromWorldTag>())
+	.Compile());
+
+	DataStorage.RegisterQuery(
+		Select(
+			TEXT("Sync dataflow struct label to column"),
+			FProcessor(EQueryTickPhase::PostPhysics, DataStorage.GetQueryTickGroupName(EQueryTickGroups::SyncExternalToDataStorage))
+				.SetExecutionMode(EExecutionMode::GameThread),
+			[](IQueryContext& Context, RowHandle RowHandle, 
+				const FTypedElementExternalObjectColumn& RawObject, const FTypedElementScriptStructTypeInfoColumn& TypeInfo, FTypedElementLabelColumn& ObjectLabel, FTypedElementLabelHashColumn& LabelHash)
+			{
+				if (RawObject.Object && TypeInfo.TypeInfo->IsChildOf(FDataflowBaseElement::StaticStruct()))
+				{
+					FDataflowBaseElement* SceneObject = static_cast<FDataflowBaseElement*>(RawObject.Object);
+					
+					const FString ObjectLabelName = SceneObject->ElementName;
+					const uint64 ObjectLabelHash = CityHash64(reinterpret_cast<const char*>(*ObjectLabelName), ObjectLabelName.Len() * sizeof(**ObjectLabelName));
+					if (LabelHash.LabelHash != ObjectLabelHash)
+					{
+						LabelHash.LabelHash = ObjectLabelHash;
+						ObjectLabel.Label = ObjectLabelName;
+					}
+				}
+			})
+		.Where(TColumn<FDataflowSceneStructTag>() && TColumn<FTypedElementSyncFromWorldTag>())
+		.Compile());
+}
+
+void UDataflowObjectFactory::RegisterVisibilityQueries(UE::Editor::DataStorage::ICoreProvider& DataStorage)
+{
+	using namespace UE::Editor::DataStorage::Queries;
+	using namespace UE::Editor::DataStorage;
+
+	DataStorage.RegisterQuery(
+		Select(
+			TEXT("Sync dataflow object visibility to Column"),
+			FProcessor(EQueryTickPhase::PrePhysics, DataStorage.GetQueryTickGroupName(EQueryTickGroups::SyncExternalToDataStorage)).SetExecutionMode(EExecutionMode::GameThread),
+			[](IQueryContext& Context, RowHandle RowHandle, const FTypedElementUObjectColumn& ObjectColumn, FVisibleInEditorColumn& VisibilityColumn)
+			{
+				if (const USceneComponent* SceneComponent = Cast<USceneComponent>(ObjectColumn.Object))
+				{
+					VisibilityColumn.bIsVisibleInEditor = SceneComponent->IsVisible();
+				}
+				else if (const AActor* SceneActor = Cast<AActor>(ObjectColumn.Object))
+				{
+					VisibilityColumn.bIsVisibleInEditor = !SceneActor->IsHiddenEd();
+				}
+			})
+		.Where(TColumn<FDataflowSceneObjectTag>() && TColumn<FTypedElementSyncFromWorldTag>())
+		.Compile()
+	);
+
+	DataStorage.RegisterQuery(
+		Select(
+			TEXT("Sync dataflow object visibility to Column"),
+			FProcessor(EQueryTickPhase::PrePhysics, DataStorage.GetQueryTickGroupName(EQueryTickGroups::SyncExternalToDataStorage)).SetExecutionMode(EExecutionMode::GameThread),
+			[](IQueryContext& Context, RowHandle RowHandle, const FTypedElementExternalObjectColumn& RawObject, const FTypedElementScriptStructTypeInfoColumn& TypeInfo,  FVisibleInEditorColumn& VisibilityColumn)
+			{
+				if (RawObject.Object && TypeInfo.TypeInfo->IsChildOf(FDataflowBaseElement::StaticStruct()))
+				{
+					FDataflowBaseElement* SceneObject = static_cast<FDataflowBaseElement*>(RawObject.Object);
+					VisibilityColumn.bIsVisibleInEditor = SceneObject->bIsVisible;
+				}
+			})
+		.Where(TColumn<FDataflowSceneStructTag>() && TColumn<FTypedElementSyncFromWorldTag>())
+		.Compile()
+	);
+
+	DataStorage.RegisterQuery(
+		Select(
+			TEXT("Sync visibility Column to dataflow object"),
+			FProcessor(EQueryTickPhase::PrePhysics, DataStorage.GetQueryTickGroupName(EQueryTickGroups::SyncDataStorageToExternal)).SetExecutionMode(EExecutionMode::GameThread),
+			[&](IQueryContext& Context, RowHandle ObjectHandle, FTypedElementUObjectColumn& ObjectColumn, const FVisibleInEditorColumn& VisibilityColumn)
+			{
+				if (USceneComponent* SceneComponent = Cast<USceneComponent>(ObjectColumn.Object))
+				{
+					Context.PushCommand(UE::Dataflow::Private::FSetComponentVisibilityCommand
+					{
+						.SceneComponent = SceneComponent,
+						.bIsVisible = VisibilityColumn.bIsVisibleInEditor
+					});
+				}
+				else if (AActor* SceneActor = Cast<AActor>(ObjectColumn.Object))
+				{
+					Context.PushCommand(UE::Dataflow::Private::FSetActorVisibilityCommand
+					{
+						.SceneActor = SceneActor,
+						.bIsVisible = VisibilityColumn.bIsVisibleInEditor
+					});
+				}
+			})
+		.Where(TColumn<FDataflowSceneObjectTag>() && TColumn<FTypedElementSyncBackToWorldTag>())
+		.Compile()
+	);
+
+	DataStorage.RegisterQuery(
+		Select(
+			TEXT("Sync visibility Column to dataflow object"),
+			FProcessor(EQueryTickPhase::PrePhysics, DataStorage.GetQueryTickGroupName(EQueryTickGroups::SyncDataStorageToExternal)).SetExecutionMode(EExecutionMode::GameThread),
+			[&](IQueryContext& Context, RowHandle ObjectHandle, const FTypedElementExternalObjectColumn& RawObject, const FTypedElementScriptStructTypeInfoColumn& TypeInfo,  const FVisibleInEditorColumn& VisibilityColumn)
+			{
+				if (RawObject.Object && TypeInfo.TypeInfo->IsChildOf(FDataflowBaseElement::StaticStruct()))
+				{
+					FDataflowBaseElement* SceneObject = static_cast<FDataflowBaseElement*>(RawObject.Object);
+					Context.PushCommand(UE::Dataflow::Private::FSetElementVisibilityCommand
+					{
+						.SceneElement = SceneObject,
+						.bIsVisible = VisibilityColumn.bIsVisibleInEditor
+					});
+				}
+			})
+		.Where(TColumn<FDataflowSceneStructTag>() && TColumn<FTypedElementSyncBackToWorldTag>())
+		.Compile()
+	);
+}
+
+void UDataflowObjectFactory::RegisterQueries(UE::Editor::DataStorage::ICoreProvider& DataStorage)
+{
+	Super::RegisterQueries(DataStorage);
+
+	RegisterLabelQueries(DataStorage);
+	RegisterHierarchyQueries(DataStorage);
+	RegisterVisibilityQueries(DataStorage);
+}
diff --git a/Engine/Plugins/Experimental/Dataflow/Source/DataflowEditor/Private/Dataflow/DataflowOutlinerFactory.h b/Engine/Plugins/Experimental/Dataflow/Source/DataflowEditor/Private/Dataflow/DataflowOutlinerFactory.h
new file mode 100644
--- /dev/null
+++ b/Engine/Plugins/Experimental/Dataflow/Source/DataflowEditor/Private/Dataflow/DataflowOutlinerFactory.h
@@ -0,0 +1,31 @@
+﻿// Copyright Epic Games, Inc. All Rights Reserved.
+#pragma once
+
+#include "Elements/Interfaces/TypedElementDataStorageFactory.h"
+
+#include "DataflowOutlinerFactory.generated.h"
+
+/** Dataflow outliner factory used to register TEDs queries */
+UCLASS()
+class UDataflowObjectFactory : public UEditorDataStorageFactory
+{
+	GENERATED_BODY()
+
+public:
+	virtual ~UDataflowObjectFactory() override = default;
+
+	//~ Begin UEditorDataStorageFactory interface
+	virtual void RegisterTables(UE::Editor::DataStorage::ICoreProvider& DataStorage, UE::Editor::DataStorage::ICompatibilityProvider&) override;
+	virtual void RegisterQueries(UE::Editor::DataStorage::ICoreProvider& DataStorage) override;
+	//~ End UEditorDataStorageFactory interface
+
+protected:
+	/** Register all the visibility queries */
+	void RegisterVisibilityQueries(UE::Editor::DataStorage::ICoreProvider& DataStorage);
+
+	/** Register all the label queries */
+	void RegisterLabelQueries(UE::Editor::DataStorage::ICoreProvider& DataStorage);
+
+	/** Register all the hierarchy queries */
+	void RegisterHierarchyQueries(UE::Editor::DataStorage::ICoreProvider& DataStorage);
+};
diff --git a/Engine/Plugins/Experimental/Dataflow/Source/DataflowEditor/Private/Dataflow/DataflowOutlinerMode.cpp b/Engine/Plugins/Experimental/Dataflow/Source/DataflowEditor/Private/Dataflow/DataflowOutlinerMode.cpp
new file mode 100644
--- /dev/null
+++ b/Engine/Plugins/Experimental/Dataflow/Source/DataflowEditor/Private/Dataflow/DataflowOutlinerMode.cpp
@@ -0,0 +1,330 @@
+// Copyright Epic Games, Inc. All Rights Reserved.
+
+#include "Dataflow/DataflowOutlinerMode.h"
+#include "Dataflow/DataflowEditorPreviewSceneBase.h"
+#include "AssetEditorModeManager.h"
+#include "EditorViewportCommands.h"
+#include "TedsOutlinerItem.h"
+#include "Selection.h"
+#include "Components/PrimitiveComponent.h"
+#include "Framework/Application/SlateApplication.h"
+#include "Dataflow/DataflowContent.h"
+#include "Dataflow/DataflowElement.h"
+#include "Elements/Columns/TypedElementCompatibilityColumns.h"
+#include "Elements/Columns/TypedElementVisibilityColumns.h"
+#include "Elements/Columns/TypedElementTypeInfoColumns.h"
+#include "Elements/Common/EditorDataStorageFeatures.h"
+
+struct FTypedElementScriptStructTypeInfoColumn;
+
+namespace UE::Dataflow::Private
+{
+	template<typename ObjectType>
+	static ObjectType* GetOutlinerItemObject(const TWeakPtr<ISceneOutlinerTreeItem>& WeakTreeItem)
+	{
+		using namespace UE::Editor::Outliner;
+		using namespace UE::Editor::DataStorage;
+
+		if (const TSharedPtr<ISceneOutlinerTreeItem> TreeItem = WeakTreeItem.Pin())
+		{
+			if (FTedsOutlinerTreeItem* TEDSItem = TreeItem->CastTo<FTedsOutlinerTreeItem>())
+			{
+				ICoreProvider* Storage = GetMutableDataStorageFeature<ICoreProvider>(StorageFeatureName);
+				if (TEDSItem->IsValid() && Storage)
+				{
+					const FTypedElementUObjectColumn* RawObjectColumn = Storage->GetColumn<FTypedElementUObjectColumn>(TEDSItem->GetRowHandle());
+					return RawObjectColumn ? Cast<ObjectType>(RawObjectColumn->Object) : nullptr;
+				}
+			}
+		}
+	
+		return nullptr;
+	}
+
+	template<typename ObjectType>
+	static ObjectType* GetOutlinerItemStruct(const TWeakPtr<ISceneOutlinerTreeItem>& WeakTreeItem)
+	{
+		using namespace UE::Editor::Outliner;
+		using namespace UE::Editor::DataStorage;
+		
+		if (const TSharedPtr<ISceneOutlinerTreeItem> TreeItem = WeakTreeItem.Pin())
+		{
+			if (FTedsOutlinerTreeItem* TEDSItem = TreeItem->CastTo<FTedsOutlinerTreeItem>())
+			{
+				ICoreProvider* Storage = GetMutableDataStorageFeature<ICoreProvider>(StorageFeatureName);
+				if (TEDSItem->IsValid() && Storage)
+				{
+					const FTypedElementExternalObjectColumn* RawObjectColumn = Storage->GetColumn<FTypedElementExternalObjectColumn>(TEDSItem->GetRowHandle());
+					const FTypedElementScriptStructTypeInfoColumn* TypeInfoColumn = Storage->GetColumn<FTypedElementScriptStructTypeInfoColumn>(TEDSItem->GetRowHandle());
+					if (RawObjectColumn && TypeInfoColumn)
+					{
+						if (RawObjectColumn->Object && TypeInfoColumn->TypeInfo == ObjectType::StaticStruct())
+						{
+							return static_cast<ObjectType*>(RawObjectColumn->Object);
+						}
+					}
+				}
+			}
+		}
+	
+		return nullptr;
+	}
+	
+	static void SetOutlinerItemVisibility(const TWeakPtr<ISceneOutlinerTreeItem>& WeakTreeItem, const bool bIsVisible)
+	{
+		using namespace UE::Editor::Outliner;
+		using namespace UE::Editor::DataStorage;
+
+		if (const TSharedPtr<ISceneOutlinerTreeItem> TreeItem = WeakTreeItem.Pin())
+		{
+			if (FTedsOutlinerTreeItem* TEDSItem = TreeItem->CastTo<FTedsOutlinerTreeItem>())
+			{
+				ICoreProvider* Storage = GetMutableDataStorageFeature<ICoreProvider>(StorageFeatureName);
+				if (TEDSItem->IsValid() && Storage)
+				{
+					if(FVisibleInEditorColumn* VisibilityColumn = Storage->GetColumn<FVisibleInEditorColumn>(TEDSItem->GetRowHandle()))
+					{
+						VisibilityColumn->bIsVisibleInEditor = bIsVisible;
+					}
+				}
+			}
+		}
+	}
+
+	/** Functor which can be used to get dataflow actors from a selection */
+	template<typename ObjectType>
+	struct FDataflowObjectSelector
+	{
+		FDataflowObjectSelector()
+		{}
+
+		bool operator()(const TWeakPtr<ISceneOutlinerTreeItem>& TreeItem, ObjectType*& SceneObject) const
+		{
+			if (ObjectType* ItemObject = GetOutlinerItemObject<ObjectType>(TreeItem))
+			{
+				SceneObject = ItemObject;
+				return true;
+			}
+			return false;
+		}
+	};
+	
+	/** Functor which can be used to get dataflow actors from a selection */
+	template<typename ObjectType>
+	struct FDataflowStructSelector
+	{
+		FDataflowStructSelector()
+		{}
+
+		bool operator()(const TWeakPtr<ISceneOutlinerTreeItem>& TreeItem, ObjectType*& SceneObject) const
+		{
+			if (ObjectType* ItemObject = GetOutlinerItemStruct<ObjectType>(TreeItem))
+			{
+				SceneObject = ItemObject;
+				return true;
+			}
+			return false;
+		}
+	};
+	
+	static void UpdateSceneSelection(const FSceneOutlinerItemSelection& Selection, FDataflowPreviewSceneBase* PreviewScene, USelection* SceneSelection)
+	{
+		// Get the selected components in TEDS
+		TArray<UPrimitiveComponent*> SelectedComponents = Selection.GetData<UPrimitiveComponent*>(UE::Dataflow::Private::FDataflowObjectSelector<UPrimitiveComponent>());
+
+		// Unselect all previous components
+		TArray<UPrimitiveComponent*> PreviousSelection;
+		SceneSelection->GetSelectedObjects<UPrimitiveComponent>(PreviousSelection);
+
+		SceneSelection->Modify();
+		SceneSelection->BeginBatchSelectOperation();
+		SceneSelection->DeselectAll();
+
+		// Transfer components selection from TEDS
+		for(UPrimitiveComponent* SelectedComponent : SelectedComponents)
+		{
+			if(SelectedComponent->GetWorld() == PreviewScene->GetWorld())
+			{
+				SceneSelection->Select(SelectedComponent);
+				SelectedComponent->PushSelectionToProxy();
+			}
+		}
+		SceneSelection->EndBatchSelectOperation();
+
+		// Update the previous proxies 
+		for(UPrimitiveComponent* PreviousComponent : PreviousSelection)
+		{
+			if(PreviousComponent->GetWorld() == PreviewScene->GetWorld())
+			{
+				PreviousComponent->PushSelectionToProxy();
+			}
+		}
+
+		// Get the selected elements in TEDS
+		TArray<FDataflowBaseElement*> SelectedElements = Selection.GetData<FDataflowBaseElement*>(UE::Dataflow::Private::FDataflowStructSelector<FDataflowBaseElement>());
+
+		// Unselect all previous elements
+		for(IDataflowDebugDrawInterface::FDataflowElementsType::ElementType& SelectedElement : PreviewScene->ModifySceneElements())
+		{
+			if(SelectedElement.IsValid())
+			{
+				SelectedElement->bIsSelected = false;
+			}
+		}
+
+		// Transfer elements selection from TEDS
+		for(FDataflowBaseElement* SelectedElement : SelectedElements)
+		{
+			if(SelectedElement)
+			{
+				SelectedElement->bIsSelected = true;
+			}
+		}
+	}
+}
+
+FDataflowOutlinerMode::FDataflowOutlinerMode(const UE::Editor::Outliner::FTedsOutlinerParams& InModeParams,
+		FDataflowPreviewSceneBase* InConstructionScene, FDataflowPreviewSceneBase* InSimulationScene)
+	: FTedsOutlinerMode(InModeParams)
+	, ConstructionScene(InConstructionScene)
+	, SimulationScene(InSimulationScene)
+{
+	if (SceneOutliner)
+	{
+		TAttribute<bool> ConditionalEnabledAttribute(true);
+		ConditionalEnabledAttribute.BindRaw(this, &FDataflowOutlinerMode::CanPopulate);
+
+		SceneOutliner->SetEnabled(ConditionalEnabledAttribute);
+	}
+}
+
+FDataflowOutlinerMode::~FDataflowOutlinerMode()
+{
+	if (SceneOutliner)
+	{
+		TAttribute<bool> EmptyConditionalEnabledAttribute;
+		SceneOutliner->SetEnabled(EmptyConditionalEnabledAttribute);
+	}
+}
+
+void FDataflowOutlinerMode::OnItemSelectionChanged(FSceneOutlinerTreeItemPtr Item, ESelectInfo::Type SelectionType, const FSceneOutlinerItemSelection& Selection)
+{
+	if ((SelectionType == ESelectInfo::Direct) || !ConstructionScene || !SimulationScene)
+	{
+		return;
+	}
+	if(TSharedPtr<FAssetEditorModeManager> ConstructionManager = ConstructionScene->GetDataflowModeManager())
+	{
+		UE::Dataflow::Private::UpdateSceneSelection(Selection, ConstructionScene, ConstructionManager->GetSelectedComponents());
+	}
+	if(TSharedPtr<FAssetEditorModeManager> SimulationManager = SimulationScene->GetDataflowModeManager())
+	{
+		UE::Dataflow::Private::UpdateSceneSelection(Selection, SimulationScene, SimulationManager->GetSelectedComponents());
+	}
+}
+
+void FDataflowOutlinerMode::OnItemAdded(FSceneOutlinerTreeItemPtr Item)
+{
+	SceneOutliner->SetItemExpansion(Item, false);
+	Item->Flags.bIsExpanded = false;
+}
+
+void FDataflowOutlinerMode::OnItemDoubleClick(FSceneOutlinerTreeItemPtr SelectedItem)
+{
+	if (!ConstructionScene || !SimulationScene)
+	{
+		return;
+	}
+
+	if(const UPrimitiveComponent* SelectedComponent = UE::Dataflow::Private::GetOutlinerItemObject<UPrimitiveComponent>(SelectedItem))
+	{
+		if(SelectedComponent->GetWorld() == ConstructionScene->GetWorld())
+		{
+			ConstructionScene->OnFocusRequest().Broadcast(SelectedComponent->Bounds.GetBox());
+		}
+		else if(SelectedComponent->GetWorld() == SimulationScene->GetWorld())
+		{
+			SimulationScene->OnFocusRequest().Broadcast(SelectedComponent->Bounds.GetBox());
+		}
+	}
+	else if(FDataflowBaseElement* SelectedElement = UE::Dataflow::Private::GetOutlinerItemStruct<FDataflowBaseElement>(SelectedItem))
+	{
+		if(SelectedElement->bIsConstruction)
+		{
+			ConstructionScene->OnFocusRequest().Broadcast(SelectedElement->BoundingBox);
+		}
+		else
+		{
+			SimulationScene->OnFocusRequest().Broadcast(SelectedElement->BoundingBox);
+		}
+	}
+}
+
+FReply FDataflowOutlinerMode::OnKeyDown(const FKeyEvent& InKeyEvent)
+{
+	if(SceneOutliner)
+	{
+		const FModifierKeysState ModifierKeys = FSlateApplication::Get().GetModifierKeys();
+		const FInputChord CheckChord( InKeyEvent.GetKey(), EModifierKey::FromBools(ModifierKeys.IsControlDown(), ModifierKeys.IsAltDown(), ModifierKeys.IsShiftDown(), ModifierKeys.IsCommandDown()) );
+
+		const FSceneOutlinerItemSelection& Selection = SceneOutliner->GetSelection();
+	
+		// Use the keyboard shortcut bound to 'Focus Viewport To Selection'
+		if (FEditorViewportCommands::Get().FocusViewportToSelection->HasActiveChord(CheckChord))
+		{
+			FBox ConstructionBox, SimulationBox;
+			for(TWeakPtr<ISceneOutlinerTreeItem>& SelectedItem: Selection.SelectedItems)
+			{
+				if(const UPrimitiveComponent* SelectedComponent = UE::Dataflow::Private::GetOutlinerItemObject<UPrimitiveComponent>(SelectedItem))
+				{
+					if(SelectedComponent->GetWorld() == ConstructionScene->GetWorld())
+					{
+						ConstructionBox += SelectedComponent->Bounds.GetBox();
+					}
+					else if(SelectedComponent->GetWorld() == SimulationScene->GetWorld())
+					{
+						SimulationBox += SelectedComponent->Bounds.GetBox();
+					}
+				}
+				else if(FDataflowBaseElement* SelectedElement = UE::Dataflow::Private::GetOutlinerItemStruct<FDataflowBaseElement>(SelectedItem))
+				{
+					if(SelectedElement->bIsConstruction)
+					{
+						ConstructionBox += SelectedElement->BoundingBox;
+					}
+					else
+					{
+						SimulationBox += SelectedElement->BoundingBox;
+					}
+				}
+			}
+			if(ConstructionBox.IsValid)
+			{
+				ConstructionScene->OnFocusRequest().Broadcast(ConstructionBox);
+			}
+			if(SimulationBox.IsValid)
+			{
+				SimulationScene->OnFocusRequest().Broadcast(SimulationBox);
+			}
+		}
+		else if (InKeyEvent.GetKey() == EKeys::H)
+		{
+			const bool bIsVisible = InKeyEvent.IsControlDown() ? true : false;
+			for(TWeakPtr<ISceneOutlinerTreeItem>& SelectedItem : Selection.SelectedItems)
+			{
+				if(UPrimitiveComponent* SelectedComponent = UE::Dataflow::Private::GetOutlinerItemObject<UPrimitiveComponent>(SelectedItem))
+				{
+					SelectedComponent->SetVisibility(bIsVisible);
+				}
+				else if(FDataflowBaseElement* SelectedElement = UE::Dataflow::Private::GetOutlinerItemStruct<FDataflowBaseElement>(SelectedItem))
+				{
+					SelectedElement->bIsVisible = bIsVisible;
+				}
+				UE::Dataflow::Private::SetOutlinerItemVisibility(SelectedItem, bIsVisible);
+			}
+		}
+	}
+	return UE::Editor::Outliner::FTedsOutlinerMode::OnKeyDown(InKeyEvent);
+}
+
+#undef LOCTEXT_NAMESPACE
diff --git a/Engine/Plugins/Experimental/Dataflow/Source/DataflowEditor/Private/Dataflow/DataflowOutlinerMode.h b/Engine/Plugins/Experimental/Dataflow/Source/DataflowEditor/Private/Dataflow/DataflowOutlinerMode.h
new file mode 100644
--- /dev/null
+++ b/Engine/Plugins/Experimental/Dataflow/Source/DataflowEditor/Private/Dataflow/DataflowOutlinerMode.h
@@ -0,0 +1,38 @@
+// Copyright Epic Games, Inc. All Rights Reserved.
+#pragma once
+
+#include "TedsOutlinerMode.h"
+
+class FDataflowPreviewSceneBase;
+
+/**
+ * Dataflow outliner mode used to represent the dataflow scenes (construction/simulation)
+ * It has a more limited view compared to the normal outliner, hiding features we don't support,
+ */
+class FDataflowOutlinerMode : public UE::Editor::Outliner::FTedsOutlinerMode
+{
+public:
+
+	FDataflowOutlinerMode(const UE::Editor::Outliner::FTedsOutlinerParams& InModeParams,
+		FDataflowPreviewSceneBase* InConstructionScene, FDataflowPreviewSceneBase* InSimulationScene);
+
+	virtual ~FDataflowOutlinerMode() override;
+
+	//~ Begin FTedsOutlinerMode interface
+	virtual void OnItemSelectionChanged(FSceneOutlinerTreeItemPtr Item, ESelectInfo::Type SelectionType, const FSceneOutlinerItemSelection& Selection) override;
+	virtual void OnItemDoubleClick(FSceneOutlinerTreeItemPtr Item) override;
+	virtual void OnItemAdded(FSceneOutlinerTreeItemPtr Item) override;
+	virtual bool CanInteract(const ISceneOutlinerTreeItem& Item) const override {return true;}
+	virtual bool CanPopulate() const override {return true;}
+	virtual bool CanSupportDragAndDrop() const override { return false; }
+	virtual ESelectionMode::Type GetSelectionMode() const override { return ESelectionMode::Multi; }
+	virtual FReply OnKeyDown(const FKeyEvent& InKeyEvent) override;
+	//~ End FTedsOutlinerMode interface
+
+private :
+	/** Construction dataflow scene */
+	FDataflowPreviewSceneBase* ConstructionScene;
+
+	/** Simulation dataflow scene */
+	FDataflowPreviewSceneBase* SimulationScene;
+};
diff --git a/Engine/Plugins/Experimental/Dataflow/Source/DataflowEditor/Private/Dataflow/DataflowOutlinerView.cpp b/Engine/Plugins/Experimental/Dataflow/Source/DataflowEditor/Private/Dataflow/DataflowOutlinerView.cpp
new file mode 100644
--- /dev/null
+++ b/Engine/Plugins/Experimental/Dataflow/Source/DataflowEditor/Private/Dataflow/DataflowOutlinerView.cpp
@@ -0,0 +1,158 @@
+// Copyright Epic Games, Inc. All Rights Reserved.
+
+#include "Dataflow/DataflowOutlinerView.h"
+#include "Dataflow/DataflowEditorPreviewSceneBase.h"
+#include "Dataflow/DataflowOutlinerMode.h"
+#include "TedsOutlinerImpl.h"
+#include "Compatibility/SceneOutlinerTedsBridge.h"
+#include "Dataflow/DataflowContent.h"
+#include "Elements/Columns/TypedElementTypeInfoColumns.h"
+#include "Elements/Columns/TypedElementVisibilityColumns.h"
+#include "Elements/Common/EditorDataStorageFeatures.h"
+#include "Elements/Common/TypedElementQueryDescription.h"
+#include "Elements/Framework/TypedElementQueryBuilder.h"
+
+FDataflowOutlinerView::FDataflowOutlinerView(FDataflowPreviewSceneBase* InConstructionScene, FDataflowPreviewSceneBase* InSimulationScene, TObjectPtr<UDataflowBaseContent> InContent)
+	: FDataflowNodeView(InContent)
+	, OutlinerWidget(nullptr)
+	, ConstructionScene(InConstructionScene)
+	, SimulationScene(InSimulationScene)
+{
+	check(InContent);
+}
+
+FDataflowOutlinerView::~FDataflowOutlinerView()
+{}
+
+TSharedPtr<ISceneOutliner> FDataflowOutlinerView::CreateWidget()
+{
+	using namespace UE::Editor::DataStorage::Queries;
+	using namespace UE::Editor::Outliner;
+	UE::Editor::DataStorage::FQueryDescription RowQueryDescription = 
+		Select()
+		.Where(TColumn<FDataflowSceneObjectTag>(GetEditorContent()->GetDataflowOwner().GetFName()) || TColumn<FDataflowSceneStructTag>(GetEditorContent()->GetDataflowOwner().GetFName()) )
+		.Compile();
+
+	UE::Editor::DataStorage::FQueryDescription ColumnQueryDescription =
+		Select()
+		.ReadOnly<FTypedElementClassTypeInfoColumn, FVisibleInEditorColumn>()
+		.Compile();
+
+	FSceneOutlinerInitializationOptions InitOptions;
+	InitOptions.bShowHeaderRow = true;
+	InitOptions.FilterBarOptions.bHasFilterBar = true;
+	InitOptions.bShowTransient = false;
+	InitOptions.OutlinerIdentifier = "DataflowOutliner";
+	
+	UE::Editor::Outliner::FTedsOutlinerParams Params(nullptr);
+	Params.QueryDescription = RowQueryDescription;
+	Params.bUseDefaultTedsFilters = false;
+	Params.bShowRowHandleColumn = false;
+
+	// Add outliner filter queries
+	UE::Editor::DataStorage::FQueryDescription SimulationFilterQuery =
+		Select()
+		.Where(TColumn<FDataflowSimulationObjectTag>())
+		.Compile();
+	UE::Editor::DataStorage::FQueryDescription ConstructionFilterQuery =
+		Select()
+		.Where(TColumn<FDataflowConstructionObjectTag>())
+		.Compile();
+	UE::Editor::DataStorage::FQueryDescription ElementsFilterQuery =
+		Select()
+		.Where(TColumn<FDataflowSceneStructTag>(GetEditorContent()->GetDataflowOwner().GetFName()))
+		.Compile();
+	UE::Editor::DataStorage::FQueryDescription ComponentsFilterQuery =
+		Select()
+		.Where(TColumn<FDataflowSceneObjectTag>(GetEditorContent()->GetDataflowOwner().GetFName()))
+		.Compile();
+
+	Params.FilterQueries.Emplace("Dataflow Construction", ConstructionFilterQuery);
+	Params.FilterQueries.Emplace("Dataflow Simulation", SimulationFilterQuery);
+	Params.FilterQueries.Emplace("Dataflow Elements", ElementsFilterQuery);
+	Params.FilterQueries.Emplace("Dataflow Components", ComponentsFilterQuery);
+	
+	// Empty selection set name is currently the level editor
+	Params.SelectionSetOverride = FName("DataflowSelection");
+
+	InitOptions.ModeFactory = FCreateSceneOutlinerMode::CreateLambda([this, &Params](SSceneOutliner* Outliner)
+	{
+		Params.SceneOutliner = Outliner;
+		return new FDataflowOutlinerMode(Params, ConstructionScene, SimulationScene);
+	});
+	
+	UE::Editor::DataStorage::QueryHandle InitialColumnQuery = GetMutableDataStorageFeature<ICoreProvider>(
+		UE::Editor::DataStorage::StorageFeatureName)->RegisterQuery(MoveTemp(ColumnQueryDescription));
+
+	InitOptions.ColumnMap.Add(FSceneOutlinerBuiltInColumnTypes::Label(), FSceneOutlinerColumnInfo(ESceneOutlinerColumnVisibility::Visible, 10));
+	OutlinerWidget = SNew(SSceneOutliner, InitOptions);
+	
+	FSceneOutlinerTedsQueryBinder::GetInstance().AssignQuery(InitialColumnQuery, OutlinerWidget, Params.CellWidgetPurpose, Params.HeaderWidgetPurpose);
+	FSceneOutlinerTedsQueryBinder::GetInstance().RegisterTreeItemIDDealiaser(OutlinerWidget, FTreeItemIDDealiaser::CreateLambda([](RowHandle RowHandle)
+	{
+		return FSceneOutlinerTreeItemID(RowHandle);
+	}));
+
+	return OutlinerWidget;
+}
+
+void FDataflowOutlinerView::SetSupportedOutputTypes()
+{
+	GetSupportedOutputTypes().Empty();
+	GetSupportedOutputTypes().Add("FManagedArrayCollection");
+}
+
+void FDataflowOutlinerView::RefreshView()
+{
+	UpdateViewData();
+}
+
+void FDataflowOutlinerView::UpdateViewData()
+{
+	if(OutlinerWidget)
+	{
+		OutlinerWidget->CollapseAll();
+		OutlinerWidget->FullRefresh();
+	}
+}
+
+void FDataflowOutlinerView::ConstructionViewSelectionChanged(const TArray<UPrimitiveComponent*>& SelectedComponents, const TArray<FDataflowBaseElement*>& SelectedElements)
+{
+	OutlinerWidget->ClearSelection();
+	
+	using namespace UE::Editor::DataStorage;
+	TArray<RowHandle> SelectedRowHandles;
+	if (const ICompatibilityProvider* Compatibility = GetDataStorageFeature<ICompatibilityProvider>(CompatibilityFeatureName))
+	{
+		// Transfer components selection to outliner
+		for(UPrimitiveComponent* SelectedComponent : SelectedComponents)
+		{
+			if (FSceneOutlinerTreeItemPtr SelectedTreeItem = OutlinerWidget->GetTreeItem(
+				Compatibility->FindRowWithCompatibleObject(SelectedComponent), true))
+			{
+				OutlinerWidget->AddToSelection(SelectedTreeItem);
+				OutlinerWidget->ScrollItemIntoView(SelectedTreeItem);
+			}
+		}
+		// Transfer elements selection to outliner
+		for(FDataflowBaseElement* SelectedElement : SelectedElements)
+		{
+			if (FSceneOutlinerTreeItemPtr SelectedTreeItem = OutlinerWidget->GetTreeItem(
+				Compatibility->FindRowWithCompatibleObject(SelectedElement), true))
+			{
+				OutlinerWidget->AddToSelection(SelectedTreeItem);
+				OutlinerWidget->ScrollItemIntoView(SelectedTreeItem);
+			}
+		}
+	}
+}
+
+void FDataflowOutlinerView::SimulationViewSelectionChanged(const TArray<UPrimitiveComponent*>& SelectedComponents, const TArray<FDataflowBaseElement*>& SelectedElements)
+{
+	ConstructionViewSelectionChanged(SelectedComponents, SelectedElements);
+}
+
+void FDataflowOutlinerView::AddReferencedObjects(FReferenceCollector& Collector)
+{
+	FDataflowNodeView::AddReferencedObjects(Collector);
+}
diff --git a/Engine/Plugins/Experimental/Dataflow/Source/DataflowEditor/Private/Dataflow/DataflowOutlinerView.h b/Engine/Plugins/Experimental/Dataflow/Source/DataflowEditor/Private/Dataflow/DataflowOutlinerView.h
new file mode 100644
--- /dev/null
+++ b/Engine/Plugins/Experimental/Dataflow/Source/DataflowEditor/Private/Dataflow/DataflowOutlinerView.h
@@ -0,0 +1,49 @@
+// Copyright Epic Games, Inc. All Rights Reserved.
+#pragma once
+
+#include "Dataflow/DataflowView.h"
+
+class SSceneOutliner;
+class ISceneOutliner;
+class FDataflowPreviewSceneBase;
+struct FDataflowBaseElement;
+
+/** Class to handle the dataflow outliner widget */
+class FDataflowOutlinerView : public FDataflowNodeView
+{
+public:
+	FDataflowOutlinerView(FDataflowPreviewSceneBase* InConstructionScene, FDataflowPreviewSceneBase* InSimulationScene, TObjectPtr<UDataflowBaseContent> InContent = nullptr);
+	~FDataflowOutlinerView();
+
+	/** Create the outliner widget */
+	TSharedPtr<ISceneOutliner> CreateWidget();
+
+	/** Set the supported output types */
+	virtual void SetSupportedOutputTypes() override;
+
+	/** Update the view if necessary */
+	virtual void UpdateViewData() override;
+
+	/** Refresh the view if necessary */
+	virtual void RefreshView() override;
+
+	/** Update the view based on changes in the construction view */
+	virtual void ConstructionViewSelectionChanged(const TArray<UPrimitiveComponent*>& SelectedComponents, const TArray<FDataflowBaseElement*>& SelectedElements) override;
+	
+	/** Update the view based on changes in the construction view */
+	virtual void SimulationViewSelectionChanged(const TArray<UPrimitiveComponent*>& SelectedComponents, const TArray<FDataflowBaseElement*>& SelectedElements) override;
+
+	/** Add GC managed objects*/
+	virtual void AddReferencedObjects(FReferenceCollector& Collector) override;
+
+private:
+	/** Outliner widget stored on the view */
+	TSharedPtr<SSceneOutliner> OutlinerWidget;
+
+	/** Construction scene the outliner could refer to */
+	FDataflowPreviewSceneBase* ConstructionScene;
+
+	/** Simulation scene the outliner could refer to */
+	FDataflowPreviewSceneBase* SimulationScene;
+};
+
diff --git a/Engine/Plugins/Experimental/Dataflow/Source/DataflowEditor/Private/Dataflow/DataflowSimulationScene.cpp b/Engine/Plugins/Experimental/Dataflow/Source/DataflowEditor/Private/Dataflow/DataflowSimulationScene.cpp
--- a/Engine/Plugins/Experimental/Dataflow/Source/DataflowEditor/Private/Dataflow/DataflowSimulationScene.cpp
+++ b/Engine/Plugins/Experimental/Dataflow/Source/DataflowEditor/Private/Dataflow/DataflowSimulationScene.cpp
@@ -22,6 +22,7 @@
 #include "Components/SkeletalMeshComponent.h"
 #include "Engine/SkeletalMesh.h"
 #include "GeometryCache.h"
+#include "Dataflow/DataflowElement.h"
 
 #if WITH_EDITOR
 #include "Misc/FileHelper.h"
@@ -33,13 +34,17 @@
 //
 
 FDataflowSimulationScene::FDataflowSimulationScene(FPreviewScene::ConstructionValues ConstructionValues, UDataflowEditor* InEditor)
-	: FDataflowPreviewSceneBase(ConstructionValues, InEditor)
+	: FDataflowPreviewSceneBase(ConstructionValues, InEditor, FName("RootActor"))
 {
 	SceneDescription = NewObject<UDataflowSimulationSceneDescription>();
 	SceneDescription->SetSimulationScene(this);
 
 	SimulationGenerator = MakeShared<UE::Dataflow::FDataflowSimulationGenerator>();
-	AChaosCacheManager* const CacheManager = GetWorld()->SpawnActor<AChaosCacheManager>();
+
+	FActorSpawnParameters SpawnParameters;
+	SpawnParameters.Name = FName("Simulation Components");
+	
+	AChaosCacheManager* const CacheManager = GetWorld()->SpawnActor<AChaosCacheManager>(SpawnParameters);
 	CacheManager->SetEditorIconVisibility(false);
 	RootSceneActor = CacheManager;
 
@@ -117,12 +122,21 @@ void FDataflowSimulationScene::ResetSimulationScene()
 	// Destroy the spawned root actor
 	if(PreviewActor && GetWorld())
 	{
+		PreviewActor->ForEachComponent<UActorComponent>(true, [this](UActorComponent* ActorComponent)
+		{
+			RemoveSceneObject(ActorComponent);
+		});
+		RemoveSceneObject(PreviewActor);
+		
 		GetWorld()->EditorDestroyActor(PreviewActor, true);
 			
 		// Since deletion can be delayed, rename to avoid future name collision
 		// Call UObject::Rename directly on actor to avoid AActor::Rename which unnecessarily sunregister and re-register components
 		PreviewActor->UObject::Rename(nullptr, GetTransientPackage(), REN_DontCreateRedirectors | REN_ForceNoResetLoaders);
 	}
+	RemoveSceneObject(RootSceneActor);
+
+	bPreviewSceneDirty = true;
 	
 	// Unbind the scene selection
 	UnbindSceneSelection();
@@ -215,10 +229,22 @@ void FDataflowSimulationScene::CreateSimulationScene()
 		SimulationGenerator->SetDeltaTime(DeltaTime);
 		// Setup all the skelmesh animations
 		UE::Dataflow::SetupSkeletonAnimation(PreviewActor, SceneDescription->bSkeletalMeshVisibility);
+
+		if(PreviewActor)
+		{
+			PreviewActor->ForEachComponent<UActorComponent>(true, [this](UActorComponent* ActorComponent)
+			{
+				AddSceneObject(ActorComponent, false);
+			});
+			AddSceneObject(PreviewActor, false);
+		}
+		AddSceneObject(RootSceneActor, false);
 		
 		GetWorld()->GetSubsystem<UDataflowSimulationManager>()->SetSimulationEnabled(false);
 	}
 
+	bPreviewSceneDirty = true;
+	
 	// update the selection binding since we are constantly editing the graph
 	BindSceneSelection();
 
diff --git a/Engine/Plugins/Experimental/Dataflow/Source/DataflowEditor/Private/Dataflow/DataflowSimulationViewportClient.cpp b/Engine/Plugins/Experimental/Dataflow/Source/DataflowEditor/Private/Dataflow/DataflowSimulationViewportClient.cpp
--- a/Engine/Plugins/Experimental/Dataflow/Source/DataflowEditor/Private/Dataflow/DataflowSimulationViewportClient.cpp
+++ b/Engine/Plugins/Experimental/Dataflow/Source/DataflowEditor/Private/Dataflow/DataflowSimulationViewportClient.cpp
@@ -98,19 +98,27 @@ void FDataflowSimulationViewportClient::OnViewportClicked(HHitProxy* HitProxy)
 			Component->PushSelectionToProxy();
 		}
 	}
-
 	SelectedComponents->EndBatchSelectOperation();
 
 	for (UPrimitiveComponent* const Component : PreviouslySelectedComponents)
 	{
 		Component->PushSelectionToProxy();
 	}
+
+	TArray<UPrimitiveComponent*> CurrentlySelectedComponents;
+	SelectedComponents->GetSelectedObjects<UPrimitiveComponent>(CurrentlySelectedComponents);
+
+	// Get all the scene selected elements 
+	TArray<FDataflowBaseElement*> DataflowElements;
+	GetSelectedElements(HitProxy, DataflowElements);
+	
+	OnSelectionChangedMulticast.Broadcast(CurrentlySelectedComponents, DataflowElements);
 }
 
 
 void FDataflowSimulationViewportClient::Draw(const FSceneView* View, FPrimitiveDrawInterface* PDI)
 {
-	FEditorViewportClient::Draw(View, PDI);
+	FDataflowEditorViewportClientBase::Draw(View, PDI);
 
 	using namespace UE::Dataflow;
 	const TMap<FName, TUniquePtr<IDataflowSimulationVisualization>>& Visualizations = FDataflowSimulationVisualizationRegistry::GetInstance().GetVisualizations();
diff --git a/Engine/Plugins/Experimental/Dataflow/Source/DataflowEditor/Private/Dataflow/DataflowSkeletonView.cpp b/Engine/Plugins/Experimental/Dataflow/Source/DataflowEditor/Private/Dataflow/DataflowSkeletonView.cpp
--- a/Engine/Plugins/Experimental/Dataflow/Source/DataflowEditor/Private/Dataflow/DataflowSkeletonView.cpp
+++ b/Engine/Plugins/Experimental/Dataflow/Source/DataflowEditor/Private/Dataflow/DataflowSkeletonView.cpp
@@ -104,7 +104,7 @@ void FDataflowSkeletonView::UpdateViewData()
 	UpdateSkeleton();
 }
 
-void FDataflowSkeletonView::ConstructionViewSelectionChanged(const TArray<UPrimitiveComponent*>& InSelectedComponents)
+void FDataflowSkeletonView::ConstructionViewSelectionChanged(const TArray<UPrimitiveComponent*>& InSelectedComponents, const TArray<FDataflowBaseElement*>& SelectedElements)
 {
 	if (ensure(SkeletonEditor))
 	{
diff --git a/Engine/Plugins/Experimental/Dataflow/Source/DataflowEditor/Private/Dataflow/DataflowView.cpp b/Engine/Plugins/Experimental/Dataflow/Source/DataflowEditor/Private/Dataflow/DataflowView.cpp
--- a/Engine/Plugins/Experimental/Dataflow/Source/DataflowEditor/Private/Dataflow/DataflowView.cpp
+++ b/Engine/Plugins/Experimental/Dataflow/Source/DataflowEditor/Private/Dataflow/DataflowView.cpp
@@ -49,13 +49,15 @@ bool FDataflowNodeView::SelectedNodeHaveSupportedOutputTypes(UDataflowEdNode* In
 	return false;
 }
 
-
-
-void FDataflowNodeView::OnConstructionViewSelectionChanged(const TArray<UPrimitiveComponent*>& InComponents)
+void FDataflowNodeView::OnConstructionViewSelectionChanged(const TArray<UPrimitiveComponent*>& SelectedComponents, const TArray<FDataflowBaseElement*>& SelectedElements)
 {
-	ConstructionViewSelectionChanged(InComponents);
+	ConstructionViewSelectionChanged(SelectedComponents, SelectedElements);
 }
 
+void FDataflowNodeView::OnSimulationViewSelectionChanged(const TArray<UPrimitiveComponent*>& SelectedComponents, const TArray<FDataflowBaseElement*>& SelectedElements)
+{
+	SimulationViewSelectionChanged(SelectedComponents, SelectedElements);
+}
 
 void FDataflowNodeView::OnSelectedNodeChanged(UDataflowEdNode* InNode)
 {
diff --git a/Engine/Plugins/Experimental/Dataflow/Source/DataflowEditor/Public/Dataflow/DataflowCollectionSpreadSheet.h b/Engine/Plugins/Experimental/Dataflow/Source/DataflowEditor/Public/Dataflow/DataflowCollectionSpreadSheet.h
--- a/Engine/Plugins/Experimental/Dataflow/Source/DataflowEditor/Public/Dataflow/DataflowCollectionSpreadSheet.h
+++ b/Engine/Plugins/Experimental/Dataflow/Source/DataflowEditor/Public/Dataflow/DataflowCollectionSpreadSheet.h
@@ -23,8 +23,9 @@ public:
 
 	virtual void SetSupportedOutputTypes() override;
 	virtual void UpdateViewData() override;
-	virtual void ConstructionViewSelectionChanged(const TArray<UPrimitiveComponent*>& InSelectedComponents) override {};
-
+	virtual void ConstructionViewSelectionChanged(const TArray<UPrimitiveComponent*>& SelectedComponents, const TArray<FDataflowBaseElement*>& SelectedElements) override {};
+	virtual void SimulationViewSelectionChanged(const TArray<UPrimitiveComponent*>& SelectedComponents, const TArray<FDataflowBaseElement*>& SelectedElements) override {};
+	
 	void SetCollectionSpreadSheet(TSharedPtr<SCollectionSpreadSheetWidget>& InCollectionSpreadSheet);
 
 private:
diff --git a/Engine/Plugins/Experimental/Dataflow/Source/DataflowEditor/Public/Dataflow/DataflowConstructionScene.h b/Engine/Plugins/Experimental/Dataflow/Source/DataflowEditor/Public/Dataflow/DataflowConstructionScene.h
--- a/Engine/Plugins/Experimental/Dataflow/Source/DataflowEditor/Public/Dataflow/DataflowConstructionScene.h
+++ b/Engine/Plugins/Experimental/Dataflow/Source/DataflowEditor/Public/Dataflow/DataflowConstructionScene.h
@@ -2,6 +2,7 @@
 
 #pragma once
 
+#include "Dataflow/DataflowContent.h"
 #include "Dataflow/DataflowEditorPreviewSceneBase.h"
 #include "Dataflow/DataflowNodeParameters.h"
 
diff --git a/Engine/Plugins/Experimental/Dataflow/Source/DataflowEditor/Public/Dataflow/DataflowConstructionViewportClient.h b/Engine/Plugins/Experimental/Dataflow/Source/DataflowEditor/Public/Dataflow/DataflowConstructionViewportClient.h
--- a/Engine/Plugins/Experimental/Dataflow/Source/DataflowEditor/Public/Dataflow/DataflowConstructionViewportClient.h
+++ b/Engine/Plugins/Experimental/Dataflow/Source/DataflowEditor/Public/Dataflow/DataflowConstructionViewportClient.h
@@ -53,7 +53,7 @@ public:
 	// FGCObject Interface
 	virtual FString GetReferencerName() const override { return TEXT("FDataflowConstructionViewportClient"); }
 
-	DECLARE_MULTICAST_DELEGATE_OneParam(FOnSelectionChangedMulticast, const TArray<UPrimitiveComponent*>&)
+	DECLARE_MULTICAST_DELEGATE_TwoParams(FOnSelectionChangedMulticast, const TArray<UPrimitiveComponent*>&, const TArray<FDataflowBaseElement*>&)
 	FOnSelectionChangedMulticast OnSelectionChangedMulticast;
 
 	FString GetOverlayString() const;
diff --git a/Engine/Plugins/Experimental/Dataflow/Source/DataflowEditor/Public/Dataflow/DataflowEditorPreviewSceneBase.h b/Engine/Plugins/Experimental/Dataflow/Source/DataflowEditor/Public/Dataflow/DataflowEditorPreviewSceneBase.h
--- a/Engine/Plugins/Experimental/Dataflow/Source/DataflowEditor/Public/Dataflow/DataflowEditorPreviewSceneBase.h
+++ b/Engine/Plugins/Experimental/Dataflow/Source/DataflowEditor/Public/Dataflow/DataflowEditorPreviewSceneBase.h
@@ -3,10 +3,74 @@
 #pragma once
 
 #include "AdvancedPreviewScene.h"
+#include "Dataflow/DataflowDebugDraw.h"
+#include "Dataflow/DataflowDebugDrawObject.h"
+#include "Elements/Common/TypedElementCommonTypes.h"
 
+#include "DataflowEditorPreviewSceneBase.generated.h"
+
+struct FDataflowBaseElement;
 class UDataflowBaseContent;
 class UDataflowEditor;
 class FAssetEditorModeManager;
+class USelection;
+
+/**
+ * TEDS tag added to any object that belongs to dataflow construction scene
+ */
+USTRUCT(meta = (DisplayName = "Dataflow construction object"))
+struct FDataflowConstructionObjectTag final : public FEditorDataStorageTag
+{
+	GENERATED_BODY()
+};
+
+/**
+ * TEDS tag added to any object that belongs to dataflow simulation scene
+ */
+USTRUCT(meta = (DisplayName = "Dataflow simulation object"))
+struct FDataflowSimulationObjectTag final : public FEditorDataStorageTag
+{
+	GENERATED_BODY()
+};
+
+/**
+ * TEDS tag added to any object that belongs to dataflow scene
+ */
+USTRUCT(meta = (DisplayName = "Dataflow scene object", EditorDataStorage_DynamicColumnTemplate))
+struct FDataflowSceneObjectTag final : public FEditorDataStorageTag
+{
+	GENERATED_BODY()
+
+	struct EditorDataStorage_DynamicColumnTemplate{};
+};
+
+/**
+ * TEDS tag added to any struct that belongs to dataflow scene
+ */
+USTRUCT(meta = (DisplayName = "Dataflow scene struct", EditorDataStorage_DynamicColumnTemplate))
+struct FDataflowSceneStructTag final : public FEditorDataStorageTag
+{
+	GENERATED_BODY()
+
+	struct EditorDataStorage_DynamicColumnTemplate{};
+};
+
+/**
+ * TEDS column added to any struct/scene to display the type (construction/simulation/evaluation)
+ */
+USTRUCT(meta = (DisplayName = "Dataflow scene type"))
+struct FDataflowSceneTypeColumn final : public FEditorDataStorageColumn
+{
+	GENERATED_BODY()
+
+public:
+
+	int32 SceneType = 0;
+};
+
+
+/** Dataflow focus request delegate */
+DECLARE_MULTICAST_DELEGATE_OneParam(FDataflowFocusRequestDelegate, const FBox&)
 
 /**
  * Dataflow preview scene base
@@ -17,7 +81,7 @@ class DATAFLOWEDITOR_API FDataflowPreviewSceneBase : public FAdvancedPreviewScen
 {
 public:
 
-	FDataflowPreviewSceneBase(FPreviewScene::ConstructionValues ConstructionValues, UDataflowEditor* Editor);
+	FDataflowPreviewSceneBase(FPreviewScene::ConstructionValues ConstructionValues, UDataflowEditor* Editor, const FName& InActorName);
 	virtual ~FDataflowPreviewSceneBase();
 	
 	// FGCObject interface
@@ -42,6 +106,18 @@ public:
 	/** Build the scene bounding box */
 	virtual FBox GetBoundingBox() const;
 
+	/** Get const scene elements */
+	const IDataflowDebugDrawInterface::FDataflowElementsType& GetSceneElements() const {return SceneElements;}
+
+	/** Get non const scene elements */
+	IDataflowDebugDrawInterface::FDataflowElementsType& ModifySceneElements() {return SceneElements;}
+
+	/** Get const debug draw component */
+	const TObjectPtr<class UDataflowDebugDrawComponent>& GetDebugDrawComponent() const {return DebugDrawComponent;}
+
+	/** Get non const debug draw component */
+	TObjectPtr<class UDataflowDebugDrawComponent>& ModifyDebugDrawComponent() {return DebugDrawComponent;}
+
 	/** Tick data flow scene */
 	virtual void TickDataflowScene(const float DeltaSeconds) {}
 
@@ -54,9 +130,38 @@ public:
 	/** Update the currently selected scene profile */
 	void SetCurrentProfileIndex(int32 NewProfileIndex);
 
+	/** Return true if the preview scene is dirty */
+	bool IsSceneDirty() const {return bPreviewSceneDirty;}
+
+	/** Reset the dirty flag */
+	void ResetDirtyFlag() {bPreviewSceneDirty = false;}
 
+	/** Event triggered when an object is focused in the scene (double-click in the scene outliner)*/
+	FDataflowFocusRequestDelegate& OnFocusRequest() { return FocusRequestDelegate; }
+
+	/** Register all the scene elements to TEDs */
+	void RegisterSceneElements(const bool bIsConstruction);
+
+	/** Unregister all the scene elements to TEDs */
+	void UnregisterSceneElements();
+
+	/** Return the scene selected components */
+	static USelection* GetSelectedComponents(FDataflowPreviewSceneBase* DataflowScene);
+	
 protected:
 
+	/** Store Scene object into the TEDS database */
+	void AddSceneObject(UObject* SceneObject, const bool bIsConstruction) const;
+
+	/** Store Scene struct into the TEDS database */
+	void AddSceneStruct(void* SceneStruct, TWeakObjectPtr<UScriptStruct> TypeInfo, const bool bIsConstruction) const;
+
+	/** Remove Scene object from the TEDS database */
+	void RemoveSceneObject(UObject* SceneObject) const;
+
+	/** Remove Scene struct from the TEDS database */
+	void RemoveSceneStruct(void* SceneStruct) const;
+
 	/** Respond to changes in the scene profile settings */
 	void OnAssetViewerSettingsRefresh(const FName& InPropertyName);
 
@@ -68,5 +173,17 @@ protected:
 
 	/** Mode Manager for selection */
 	TSharedPtr<FAssetEditorModeManager> DataflowModeManager;
+
+	/** Boolean to check if the preview scene is dirty or not */
+	bool bPreviewSceneDirty = false;
+
+	/** Delegate to focus the viewport */
+	FDataflowFocusRequestDelegate FocusRequestDelegate;
+
+	/** List of scene elements that could be used in the editor (outliner/viewport...) */
+	IDataflowDebugDrawInterface::FDataflowElementsType SceneElements;
+	
+	/** Persistent component used for debug drawing */
+	TObjectPtr<class UDataflowDebugDrawComponent> DebugDrawComponent;
 };
 
diff --git a/Engine/Plugins/Experimental/Dataflow/Source/DataflowEditor/Public/Dataflow/DataflowEditorToolkit.h b/Engine/Plugins/Experimental/Dataflow/Source/DataflowEditor/Public/Dataflow/DataflowEditorToolkit.h
--- a/Engine/Plugins/Experimental/Dataflow/Source/DataflowEditor/Public/Dataflow/DataflowEditorToolkit.h
+++ b/Engine/Plugins/Experimental/Dataflow/Source/DataflowEditor/Public/Dataflow/DataflowEditorToolkit.h
@@ -12,6 +12,7 @@
 #include "Dataflow/DataflowCollectionSpreadSheet.h"
 #include "Dataflow/DataflowConstructionViewport.h"
 #include "Dataflow/DataflowSimulationViewport.h"
+#include "Dataflow/DataflowContent.h"
 
 class FDocumentTracker;
 class FEditorViewportTabContent;
@@ -29,6 +30,7 @@ class FDataflowSimulationViewportClient;
 class UDataflowBaseContent;
 class FDataflowSimulationScene;
 class FDataflowSkeletonView;
+class FDataflowOutlinerView;
 class UDataflowEditor;
 class SDataflowMembersWidget;
 class FDataflowSimulationSceneProfileIndexStorage;
@@ -119,7 +121,8 @@ protected:
 	void OnNodeDeleted(const TSet<UObject*>& NewSelection);
 	void OnNodeSingleClicked(UObject* ClickedNode) const;
 	void OnAssetPropertyValueChanged(const FPropertyChangedEvent& PropertyChangedEvent);
-	void OnConstructionViewSelectionChanged(const TArray<UPrimitiveComponent*>& SelectedComponents);
+	void OnConstructionViewSelectionChanged(const TArray<UPrimitiveComponent*>& SelectedComponents, const TArray<FDataflowBaseElement*>& SelectedElements);
+	void OnSimulationViewSelectionChanged(const TArray<UPrimitiveComponent*>& SelectedComponents, const TArray<FDataflowBaseElement*>& SelectedElements);
 	// Callback to remove the closed one from the listener views
 	void OnTabClosed(TSharedRef<SDockTab> Tab);
 
@@ -144,6 +147,7 @@ private:
 	TSharedRef<SDockTab> SpawnTab_SubGraphTab(const FSpawnTabArgs& Args);
 	TSharedRef<SDockTab> SpawnTab_NodeDetails(const FSpawnTabArgs& Args);
 	TSharedRef<SDockTab> SpawnTab_SkeletonView(const FSpawnTabArgs& Args);
+	TSharedRef<SDockTab> SpawnTab_OutlinerView(const FSpawnTabArgs& Args);
 	TSharedRef<SDockTab> SpawnTab_SelectionView(const FSpawnTabArgs& Args);
 	TSharedRef<SDockTab> SpawnTab_CollectionSpreadSheet(const FSpawnTabArgs& Args);
 	TSharedRef<SDockTab> SpawnTab_AssetDetails(const FSpawnTabArgs& Args);
@@ -182,6 +186,7 @@ private:
 	static const FName SubGraphCanvasTabId;
 	static const FName NodeDetailsTabId;
 	static const FName SkeletonViewTabId;
+	static const FName OutlinerViewTabId;
 	static const FName SelectionViewTabId_1;
 	static const FName SelectionViewTabId_2;
 	static const FName SelectionViewTabId_3;
@@ -204,6 +209,7 @@ private:
 	TSharedPtr<SDataflowMembersWidget> MembersWidget;
 	TSharedPtr<UE::Dataflow::FDataflowNodeDetailExtensionHandler> NodeDetailsExtensionHandler;
 	TSharedPtr<FDataflowSkeletonView> SkeletonEditorView;
+	TSharedPtr<FDataflowOutlinerView> DataflowOutlinerView;
 	TSharedPtr<IDetailsView> AssetDetailsEditor;
 	TSharedPtr<FDataflowSelectionView> DataflowSelectionView_1;
 	TSharedPtr<FDataflowSelectionView> DataflowSelectionView_2;
@@ -248,9 +254,6 @@ private:
 	void RegisterContextHandlers();
 	void UnregisterContextHandlers();
 
-	/** Create a concrete DebugDrawInterface object */
-	TUniquePtr<IDataflowDebugDrawInterface> CreateDebugDrawObject();
-
 	/** Update the debug draw based on a change of currently selected or pinned nodes */
 	void UpdateDebugDraw();
 
@@ -267,6 +270,7 @@ private:
     FDelegateHandle OnFinishedChangingPropertiesDelegateHandle;
 	FDelegateHandle OnFinishedChangingAssetPropertiesDelegateHandle;
 	FDelegateHandle OnConstructionSelectionChangedDelegateHandle;
+	FDelegateHandle OnSimulationSelectionChangedDelegateHandle;
 	FDelegateHandle OnSimulationSceneChangedDelegateHandle;
 
 	// Delegates to communicate with Context
@@ -329,8 +333,5 @@ private:
 	FDateTime GraphEvaluationFinished;
 	bool bViewsNeedRefresh = false;
 
-	/** Persistent component used for debug drawing */
-	TObjectPtr<class UDataflowDebugDrawComponent> DebugDrawComponent;
-
 	FString DebugDrawOverlayString;
 };
diff --git a/Engine/Plugins/Experimental/Dataflow/Source/DataflowEditor/Public/Dataflow/DataflowEditorViewportClientBase.h b/Engine/Plugins/Experimental/Dataflow/Source/DataflowEditor/Public/Dataflow/DataflowEditorViewportClientBase.h
--- a/Engine/Plugins/Experimental/Dataflow/Source/DataflowEditor/Public/Dataflow/DataflowEditorViewportClientBase.h
+++ b/Engine/Plugins/Experimental/Dataflow/Source/DataflowEditor/Public/Dataflow/DataflowEditorViewportClientBase.h
@@ -4,6 +4,7 @@
 #include "EditorViewportClient.h"
 #include "InputBehaviorSet.h"
 
+struct FDataflowBaseElement;
 class FEditorModeTools;
 class FPreviewScene;
 class SEditorViewport;
@@ -26,6 +27,10 @@ protected:
 	// FGCObject
 	virtual void AddReferencedObjects(FReferenceCollector& Collector) override;
 
+	//~ Begin FEditorViewportClient interface
+	virtual void Draw(const FSceneView* View, FPrimitiveDrawInterface* PDI) override;
+	//~ End FEditorViewportClient interface
+
 	/** Dataflow preview scene from the toolkit */
 	FDataflowPreviewSceneBase* DataflowPreviewScene = nullptr;
 
@@ -35,7 +40,14 @@ protected:
 	/** All behaviors available to the current viewport (subclasses can add to this set) */
 	TObjectPtr<UInputBehaviorSet> BehaviorSet;
 
+protected :
+
+	/** Get all the scene selected elements */
+	void GetSelectedElements(HHitProxy* HitProxy, TArray<FDataflowBaseElement*>& SelectedElements) const;
+
 private:
+	/** Handle the focus request onto the bounding box */
+	void HandleFocusRequest(const FBox& BoundingBox);
 
 	// IInputBehaviorSource
 	virtual const UInputBehaviorSet* GetInputBehaviors() const override;
@@ -47,4 +59,6 @@ private:
 	/** Handle for callback when preview scene profile settings change */
 	FDelegateHandle OnAssetViewerSettingsChangedDelegateHandle;
 
+	/** Handle for callback when preview scene focus request is triggered */
+	FDelegateHandle OnFocusRequestDelegateHandle;
 };
diff --git a/Engine/Plugins/Experimental/Dataflow/Source/DataflowEditor/Public/Dataflow/DataflowOutputLog.h b/Engine/Plugins/Experimental/Dataflow/Source/DataflowEditor/Public/Dataflow/DataflowOutputLog.h
--- a/Engine/Plugins/Experimental/Dataflow/Source/DataflowEditor/Public/Dataflow/DataflowOutputLog.h
+++ b/Engine/Plugins/Experimental/Dataflow/Source/DataflowEditor/Public/Dataflow/DataflowOutputLog.h
@@ -35,7 +35,8 @@ public:
 private:
 	virtual void SetSupportedOutputTypes() override {};
 	virtual void UpdateViewData() override {};
-	virtual void ConstructionViewSelectionChanged(const TArray<UPrimitiveComponent*>& InSelectedComponents) override {};
+	virtual void ConstructionViewSelectionChanged(const TArray<UPrimitiveComponent*>& SelectedComponents, const TArray<FDataflowBaseElement*>& SelectedElements) override {};
+	virtual void SimulationViewSelectionChanged(const TArray<UPrimitiveComponent*>& SelectedComponents, const TArray<FDataflowBaseElement*>& SelectedElements) override {};
 
 	void OnMessageTokenClicked(const TSharedRef<IMessageToken>& InMessageToken);
 
diff --git a/Engine/Plugins/Experimental/Dataflow/Source/DataflowEditor/Public/Dataflow/DataflowSelectionView.h b/Engine/Plugins/Experimental/Dataflow/Source/DataflowEditor/Public/Dataflow/DataflowSelectionView.h
--- a/Engine/Plugins/Experimental/Dataflow/Source/DataflowEditor/Public/Dataflow/DataflowSelectionView.h
+++ b/Engine/Plugins/Experimental/Dataflow/Source/DataflowEditor/Public/Dataflow/DataflowSelectionView.h
@@ -22,7 +22,8 @@ public:
 
 	virtual void SetSupportedOutputTypes() override;
 	virtual void UpdateViewData() override;
-	virtual void ConstructionViewSelectionChanged(const TArray<UPrimitiveComponent*>& InSelectedComponents) override {};
+	virtual void ConstructionViewSelectionChanged(const TArray<UPrimitiveComponent*>& SelectedComponents, const TArray<FDataflowBaseElement*>& SelectedElements) override {};
+	virtual void SimulationViewSelectionChanged(const TArray<UPrimitiveComponent*>& SelectedComponents, const TArray<FDataflowBaseElement*>& SelectedElements) override {};
 
 	void SetSelectionView(TSharedPtr<SSelectionViewWidget>& InSelectionView);
 
diff --git a/Engine/Plugins/Experimental/Dataflow/Source/DataflowEditor/Public/Dataflow/DataflowSimulationViewportClient.h b/Engine/Plugins/Experimental/Dataflow/Source/DataflowEditor/Public/Dataflow/DataflowSimulationViewportClient.h
--- a/Engine/Plugins/Experimental/Dataflow/Source/DataflowEditor/Public/Dataflow/DataflowSimulationViewportClient.h
+++ b/Engine/Plugins/Experimental/Dataflow/Source/DataflowEditor/Public/Dataflow/DataflowSimulationViewportClient.h
@@ -38,6 +38,9 @@ public:
 
 	// FGCObject Interface
 	virtual FString GetReferencerName() const override { return TEXT("FDataflowSimulationViewportClient"); }
+	
+	DECLARE_MULTICAST_DELEGATE_TwoParams(FOnSelectionChangedMulticast, const TArray<UPrimitiveComponent*>&, const TArray<FDataflowBaseElement*>&)
+	FOnSelectionChangedMulticast OnSelectionChangedMulticast;
 
 private:
 
diff --git a/Engine/Plugins/Experimental/Dataflow/Source/DataflowEditor/Public/Dataflow/DataflowSkeletonView.h b/Engine/Plugins/Experimental/Dataflow/Source/DataflowEditor/Public/Dataflow/DataflowSkeletonView.h
--- a/Engine/Plugins/Experimental/Dataflow/Source/DataflowEditor/Public/Dataflow/DataflowSkeletonView.h
+++ b/Engine/Plugins/Experimental/Dataflow/Source/DataflowEditor/Public/Dataflow/DataflowSkeletonView.h
@@ -39,7 +39,10 @@ public:
 	virtual void UpdateViewData() override;
 
 	/** Update the view based on changes in the construction view */
-	virtual void ConstructionViewSelectionChanged(const TArray<UPrimitiveComponent*>& InSelectedComponents) override;
+	virtual void ConstructionViewSelectionChanged(const TArray<UPrimitiveComponent*>& SelectedComponents, const TArray<FDataflowBaseElement*>& SelectedElements) override;
+	
+	/** Update the view based on changes in the simulation view */
+	virtual void SimulationViewSelectionChanged(const TArray<UPrimitiveComponent*>& SelectedComponents, const TArray<FDataflowBaseElement*>& SelectedElements) override {}
 
 	/** Add GC managed objects*/
 	virtual void AddReferencedObjects(FReferenceCollector& Collector) override;
diff --git a/Engine/Plugins/Experimental/Dataflow/Source/DataflowEditor/Public/Dataflow/DataflowView.h b/Engine/Plugins/Experimental/Dataflow/Source/DataflowEditor/Public/Dataflow/DataflowView.h
--- a/Engine/Plugins/Experimental/Dataflow/Source/DataflowEditor/Public/Dataflow/DataflowView.h
+++ b/Engine/Plugins/Experimental/Dataflow/Source/DataflowEditor/Public/Dataflow/DataflowView.h
@@ -8,6 +8,7 @@
 #include "Dataflow/SelectionViewWidget.h"
 #include "Dataflow/DataflowSelection.h"
 
+struct FDataflowBaseElement;
 class UPrimitiveComponent;
 class UDataflowEditor;
 class UDataflowBaseContent;
@@ -19,7 +20,8 @@ class UDataflowBaseContent;
 class IDataflowViewListener
 {
 public:
-	virtual void OnConstructionViewSelectionChanged(const TArray<UPrimitiveComponent*>& InSelectedComponents) = 0;
+	virtual void OnConstructionViewSelectionChanged(const TArray<UPrimitiveComponent*>& SelectedComponents, const TArray<FDataflowBaseElement*>& SelectedElements) = 0;
+	virtual void OnSimulationViewSelectionChanged(const TArray<UPrimitiveComponent*>& SelectedComponents, const TArray<FDataflowBaseElement*>& SelectedElements) = 0;
 	virtual void OnSelectedNodeChanged(UDataflowEdNode* InNode) = 0;  // nullptr is valid
 	virtual void OnNodeInvalidated(FDataflowNode* InvalidatedNode) = 0;
 	virtual void RefreshView() = 0;
@@ -48,8 +50,8 @@ public:
 	*/
 	virtual void UpdateViewData() = 0;
 	virtual void SetSupportedOutputTypes() = 0;
-	virtual void ConstructionViewSelectionChanged(const TArray<UPrimitiveComponent*>& InComponent) = 0;
-
+	virtual void ConstructionViewSelectionChanged(const TArray<UPrimitiveComponent*>& SelectedComponents, const TArray<FDataflowBaseElement*>& SelectedElements) = 0;
+	virtual void SimulationViewSelectionChanged(const TArray<UPrimitiveComponent*>& SelectedComponents, const TArray<FDataflowBaseElement*>& SelectedElements) = 0;
 	/**
 	* Callback for PinnedDown change
 	*/
@@ -63,7 +65,8 @@ public:
 	/**
 	* Virtual function overrides from IDataflowViewListener base class
 	*/
-	virtual void OnConstructionViewSelectionChanged(const TArray<UPrimitiveComponent*>& InNode) override;
+	virtual void OnConstructionViewSelectionChanged(const TArray<UPrimitiveComponent*>& SelectedComponents, const TArray<FDataflowBaseElement*>& SelectedElements) override;
+	virtual void OnSimulationViewSelectionChanged(const TArray<UPrimitiveComponent*>& SelectedComponents, const TArray<FDataflowBaseElement*>& SelectedElements) override;
 	virtual void OnSelectedNodeChanged(UDataflowEdNode* InNode) override;  // nullptr is valid
 	virtual void OnNodeInvalidated(FDataflowNode* InvalidatedNode) override {}
 	virtual void RefreshView() override;
diff --git a/Engine/Plugins/Experimental/Dataflow/Source/DataflowNodes/Private/Dataflow/DataflowSkeletalMeshNodes.cpp b/Engine/Plugins/Experimental/Dataflow/Source/DataflowNodes/Private/Dataflow/DataflowSkeletalMeshNodes.cpp
--- a/Engine/Plugins/Experimental/Dataflow/Source/DataflowNodes/Private/Dataflow/DataflowSkeletalMeshNodes.cpp
+++ b/Engine/Plugins/Experimental/Dataflow/Source/DataflowNodes/Private/Dataflow/DataflowSkeletalMeshNodes.cpp
@@ -1,6 +1,7 @@
 // Copyright Epic Games, Inc. All Rights Reserved.
 
 #include "Dataflow/DataflowSkeletalMeshNodes.h"
+#include "Dataflow/DataflowContent.h"
 #include "Dataflow/DataflowCore.h"
 #include "Dataflow/DataflowEngine.h"
 #include "Dataflow/DataflowEngineTypes.h"
@@ -8,6 +9,13 @@
 #include "UObject/UnrealTypePrivate.h"
 #include "Animation/Skeleton.h"
 #include "Engine/SkeletalMesh.h"
+#include "SkeletalDebugRendering.h"
+#include "Dataflow/DataflowDebugDrawInterface.h"
+#include "Dataflow/DataflowElement.h"
+
+//#if WITH_EDITOR
+//#include "Dataflow/DataflowRenderingViewMode.h"
+//#endif
 
 #include UE_INLINE_GENERATED_CPP_BY_NAME(DataflowSkeletalMeshNodes)
 
@@ -26,6 +34,117 @@ namespace UE::Dataflow
 }
 
 ///////////////////////////////////////////////////////////////////////////////////////////////////////////////////////////////////
+///
+
+void FDataflowDebugDrawSkeletonObject::PopulateDataflowElements()
+{
+	ElementsOffset = DataflowElements.Num();
+	ElementsSize = ReferenceSkeleton.GetNum();
+	
+	DataflowElements.Reserve(ElementsOffset + ElementsSize);
+	
+	for(int32 BoneIndex = 0; BoneIndex < ElementsSize; ++BoneIndex)
+	{
+		int32 ParentIndex = ReferenceSkeleton.GetParentIndex(BoneIndex);
+		FBox BoundingBox;
+		BoundingBox += ReferenceSkeleton.GetBoneAbsoluteTransform(BoneIndex).GetLocation();
+		if(ParentIndex == INDEX_NONE)
+		{
+			ParentIndex = 0;
+		}
+		else
+		{
+			ParentIndex += ElementsOffset;
+			if(ParentIndex < DataflowElements.Num())
+			{
+				DataflowElements[ParentIndex]->BoundingBox += ReferenceSkeleton.GetBoneAbsoluteTransform(BoneIndex).GetLocation();
+			}
+		}
+		if (ParentIndex < DataflowElements.Num())
+		{
+			TSharedPtr<FDataflowProxyElement> ProxyElement = MakeShared<FDataflowProxyElement>(ReferenceSkeleton.GetBoneName(BoneIndex).ToString(),
+				DataflowElements[ParentIndex].Get(), BoundingBox, true);
+			ProxyElement->ElementProxy = MakeRefCount<HDataflowElementHitProxy>(ElementsOffset + BoneIndex, FName(ProxyElement->ElementName));
+			DataflowElements.Add(ProxyElement);
+		}
+	}
+}
+
+FBox FDataflowDebugDrawSkeletonObject::ComputeBoundingBox() const
+{
+	FBox BoundingBox(ForceInitToZero);
+
+	const uint32 ElementsEnd = ElementsSize + ElementsOffset;;
+	for(uint32 ElementIndex = ElementsOffset; ElementIndex < ElementsEnd; ++ElementIndex)
+	{
+		BoundingBox += DataflowElements[ElementIndex]->BoundingBox;
+	}
+
+	return BoundingBox;
+}
+
+void FDataflowDebugDrawSkeletonObject::DrawDataflowElements(FPrimitiveDrawInterface* PDI) const
+{
+	const uint32 NumBones = ReferenceSkeleton.GetNum();
+
+	TArray<FBoneIndexType> RequiredBones;
+	TArray<TRefCountPtr<HHitProxy>> HitProxies;
+	TArray<FLinearColor> BoneColors;
+	TArray<FTransform> WorldTransforms;
+	TArray<int32> SelectedBones;
+	
+	ReferenceSkeleton.GetBoneAbsoluteTransforms(WorldTransforms);
+	RequiredBones.Reserve(NumBones);
+
+	if(NumBones == ElementsSize)
+	{
+		for(int32 BoneIndex = 0; BoneIndex < ElementsSize; ++BoneIndex)
+		{
+			IDataflowDebugDrawInterface::FDataflowElementsType::ElementType& DataflowElement = DataflowElements[ElementsOffset+BoneIndex];
+			if(DataflowElement.IsValid())
+			{
+				if(DataflowElement->bIsVisible)
+				{
+					RequiredBones.Add(BoneIndex);
+				}
+				if(DataflowElement->bIsSelected)
+				{
+					SelectedBones.Add(BoneIndex);
+				}
+				// Populate the proxies with the ones stored on the elements. 
+				if(DataflowElement->IsA(FDataflowProxyElement::StaticType()))
+				{
+					if (FDataflowProxyElement* ProxyElement = StaticCast<FDataflowProxyElement*>(DataflowElement.Get()))
+					{
+						HitProxies.Add(ProxyElement->ElementProxy);
+					}
+				}
+			}
+		}
+	}
+	
+	FSkelDebugDrawConfig DrawConfig;
+	DrawConfig.bUseMultiColorAsDefaultColor = true;
+	DrawConfig.BoneDrawMode = EBoneDrawMode::Selected;
+	DrawConfig.BoneDrawSize = 1.f;
+	DrawConfig.bAddHitProxy = true;
+	DrawConfig.bForceDraw = true;
+	DrawConfig.DefaultBoneColor = FLinearColor(0.0f,0.0f,0.025f,1.0f);
+	DrawConfig.AffectedBoneColor = FLinearColor(1.0f,1.0f,1.0f,1.0f);
+	DrawConfig.SelectedBoneColor = FLinearColor(0.2f,1.0f,0.2f,1.0f);
+	DrawConfig.ParentOfSelectedBoneColor = FLinearColor(0.85f,0.45f,0.12f,1.0f);
+	
+	SkeletalDebugRendering::DrawBones(
+		PDI, FVector::Zero(),
+		RequiredBones,
+		ReferenceSkeleton,
+		WorldTransforms,
+		SelectedBones,
+		BoneColors,
+		HitProxies,
+		DrawConfig
+	);
+}
 
 void FGetSkeletalMeshDataflowNode::Evaluate(UE::Dataflow::FContext& Context, const FDataflowOutput* Out) const
 {
@@ -61,6 +180,26 @@ void FGetSkeletalMeshDataflowNode::SetAssetProperty(UObject* Asset)
 	}
 }
 
+#if WITH_EDITOR
+void FGetSkeletalMeshDataflowNode::DebugDraw(UE::Dataflow::FContext& Context,
+		IDataflowDebugDrawInterface& DataflowRenderingInterface,
+		const FDebugDrawParameters& DebugDrawParameters) const
+{
+	if(SkeletalMesh)
+	{
+		TRefCountPtr<IDataflowDebugDrawObject> SkeletonObject(MakeDebugDrawObject<FDataflowDebugDrawSkeletonObject>(
+			DataflowRenderingInterface.ModifyDataflowElements(), SkeletalMesh->GetRefSkeleton()));
+	
+		DataflowRenderingInterface.DrawObject(SkeletonObject);
+	}
+}
+
+bool FGetSkeletalMeshDataflowNode::CanDebugDrawViewMode(const FName& ViewModeName) const
+{
+	return true;// ViewModeName == UE::Dataflow::FDataflowConstruction3DViewMode::Name;
+}
+#endif
+
 ///////////////////////////////////////////////////////////////////////////////////////////////////////////////////////////////////
 
 void FGetSkeletonDataflowNode::Evaluate(UE::Dataflow::FContext& Context, const FDataflowOutput* Out) const
@@ -97,6 +236,28 @@ void FGetSkeletonDataflowNode::SetAssetProperty(UObject* Asset)
 	}
 }
 
+#if WITH_EDITOR
+
+void FGetSkeletonDataflowNode::DebugDraw(UE::Dataflow::FContext& Context,
+		IDataflowDebugDrawInterface& DataflowRenderingInterface,
+		const FDebugDrawParameters& DebugDrawParameters) const
+{
+	if(Skeleton)
+	{
+		TRefCountPtr<IDataflowDebugDrawObject> SkeletonObject(MakeDebugDrawObject<FDataflowDebugDrawSkeletonObject>(
+			DataflowRenderingInterface.ModifyDataflowElements(), Skeleton->GetReferenceSkeleton()));
+	
+		DataflowRenderingInterface.DrawObject(SkeletonObject);
+	}
+}
+
+bool FGetSkeletonDataflowNode::CanDebugDrawViewMode(const FName& ViewModeName) const
+{
+	return true;//ViewModeName == UE::Dataflow::FDataflowConstruction3DViewMode::Name;
+}
+
+#endif
+
 ///////////////////////////////////////////////////////////////////////////////////////////////////////////////////////////////////
 
 void FSkeletalMeshBoneDataflowNode::Evaluate(UE::Dataflow::FContext& Context, const FDataflowOutput* Out) const
diff --git a/Engine/Plugins/Experimental/Dataflow/Source/DataflowNodes/Public/Dataflow/DataflowSkeletalMeshNodes.h b/Engine/Plugins/Experimental/Dataflow/Source/DataflowNodes/Public/Dataflow/DataflowSkeletalMeshNodes.h
--- a/Engine/Plugins/Experimental/Dataflow/Source/DataflowNodes/Public/Dataflow/DataflowSkeletalMeshNodes.h
+++ b/Engine/Plugins/Experimental/Dataflow/Source/DataflowNodes/Public/Dataflow/DataflowSkeletalMeshNodes.h
@@ -4,6 +4,9 @@
 
 #include "CoreMinimal.h"
 #include "Dataflow/DataflowCore.h"
+#include "Dataflow/DataflowDebugDraw.h"
+#include "Dataflow/DataflowDebugDrawComponent.h"
+#include "Dataflow/DataflowDebugDrawObject.h"
 #include "Dataflow/DataflowEngine.h"
 
 #include "DataflowSkeletalMeshNodes.generated.h"
@@ -12,10 +15,31 @@ DEFINE_LOG_CATEGORY_STATIC(LogDataflowSkeletalMeshNodes, Log, All);
 
 class USkeletalMesh;
 
+/** Debug draw skeleton object to debug skeletal meshes */
+struct DATAFLOWNODES_API FDataflowDebugDrawSkeletonObject : public FDataflowDebugDrawBaseObject
+{
+	FDataflowDebugDrawSkeletonObject(IDataflowDebugDrawInterface::FDataflowElementsType& InDataflowElements, const FReferenceSkeleton& InReferenceSkeleton) :
+		FDataflowDebugDrawBaseObject(InDataflowElements), ReferenceSkeleton(InReferenceSkeleton)
+	{}
+	
+	/** Populate dataflow elements */
+	virtual void PopulateDataflowElements() override;
+
+	/** Debug draw dataflow element */
+	virtual void DrawDataflowElements(FPrimitiveDrawInterface* PDI) const override;
+
+	/** Compute the dataflow elements bounding box */
+	virtual FBox ComputeBoundingBox() const override;
+
+private :
+	/** Skeletal mesh to use to populate/draw scene elements */
+	const FReferenceSkeleton& ReferenceSkeleton;
+};
+
 USTRUCT()
 struct DATAFLOWNODES_API FGetSkeletalMeshDataflowNode: public FDataflowNode
 {
-	GENERATED_USTRUCT_BODY()
+	GENERATED_BODY()
 	DATAFLOW_NODE_DEFINE_INTERNAL(FGetSkeletalMeshDataflowNode, "SkeletalMesh", "General", "Skeletal Mesh")
 
 public:
@@ -37,12 +61,20 @@ public:
 private:
 	virtual bool SupportsAssetProperty(UObject* Asset) const override;
 	virtual void SetAssetProperty(UObject* Asset) override;
+	
+#if WITH_EDITOR
+	virtual bool CanDebugDraw() const override { return true; }
+	virtual bool CanDebugDrawViewMode(const FName& ViewModeName) const override;
+	virtual void DebugDraw(UE::Dataflow::FContext& Context,
+		IDataflowDebugDrawInterface& DataflowRenderingInterface,
+		const FDebugDrawParameters& DebugDrawParameters) const override;
+#endif
 };
 
 USTRUCT()
 struct DATAFLOWNODES_API FGetSkeletonDataflowNode : public FDataflowNode
 {
-	GENERATED_USTRUCT_BODY()
+	GENERATED_BODY()
 	DATAFLOW_NODE_DEFINE_INTERNAL(FGetSkeletonDataflowNode, "Skeleton", "General", "Skeletal Mesh")
 
 public:
@@ -64,13 +96,21 @@ public:
 private:
 	virtual bool SupportsAssetProperty(UObject* Asset) const override;
 	virtual void SetAssetProperty(UObject* Asset) override;
+
+#if WITH_EDITOR
+	virtual bool CanDebugDraw() const override { return true; }
+	virtual bool CanDebugDrawViewMode(const FName& ViewModeName) const override;
+	virtual void DebugDraw(UE::Dataflow::FContext& Context,
+		IDataflowDebugDrawInterface& DataflowRenderingInterface,
+		const FDebugDrawParameters& DebugDrawParameters) const override;
+#endif
 };
 
 
 USTRUCT()
 struct DATAFLOWNODES_API FSkeletalMeshBoneDataflowNode : public FDataflowNode
 {
-	GENERATED_USTRUCT_BODY()
+	GENERATED_BODY()
 	DATAFLOW_NODE_DEFINE_INTERNAL(FSkeletalMeshBoneDataflowNode, "SkeletalMeshBone", "General", "Skeletal Mesh")
 
 public:
diff --git a/Engine/Source/Runtime/Experimental/Dataflow/Core/Public/Dataflow/DataflowDebugDrawInterface.h b/Engine/Source/Runtime/Experimental/Dataflow/Core/Public/Dataflow/DataflowDebugDrawInterface.h
--- a/Engine/Source/Runtime/Experimental/Dataflow/Core/Public/Dataflow/DataflowDebugDrawInterface.h
+++ b/Engine/Source/Runtime/Experimental/Dataflow/Core/Public/Dataflow/DataflowDebugDrawInterface.h
@@ -6,10 +6,25 @@
 #include "Containers/Array.h"
 #include "DynamicMesh/DynamicMesh3.h"
 
+struct FDataflowBaseElement;
+
+/** Dataflow object debug draw interface */
+struct DATAFLOWCORE_API IDataflowDebugDrawObject : public FRefCountedObject
+{
+	virtual ~IDataflowDebugDrawObject() = default;
+
+	static FName StaticType() { return FName("IDataflowDebugDrawObject"); }
+
+	/** Check of the object type */
+	virtual bool IsA(FName InType) const  {return false;}
+};
+
 class IDataflowDebugDrawInterface
 {
 public:
 
+	using FDataflowElementsType = TArray<TSharedPtr<FDataflowBaseElement>>;
+
 	virtual ~IDataflowDebugDrawInterface() = default;
 
 	// State management
@@ -24,6 +39,7 @@ public:
 	DATAFLOWCORE_API virtual void ResetAllState() = 0;
 
 	DATAFLOWCORE_API virtual void ReservePoints(int32 NumAdditionalPoints) = 0;
+	DATAFLOWCORE_API virtual void DrawObject(const TRefCountPtr<IDataflowDebugDrawObject>& Object) = 0;
 	DATAFLOWCORE_API virtual void DrawPoint(const FVector& Position) = 0;
 	DATAFLOWCORE_API virtual void DrawLine(const FVector& Start, const FVector& End) const = 0;
 
@@ -50,6 +66,12 @@ public:
 	DATAFLOWCORE_API virtual void DrawOverlayText(const FString& InString) = 0;
 	DATAFLOWCORE_API virtual FString GetOverlayText() const = 0;
 
+	/** Dataflow elements non const accessor */
+	DATAFLOWCORE_API virtual FDataflowElementsType& ModifyDataflowElements() = 0;
+
+	/** Dataflow elements const accessor */
+	DATAFLOWCORE_API virtual const FDataflowElementsType& GetDataflowElements() const = 0;
+
 	// TODO: More draw functions
 
 };
diff --git a/Engine/Source/Runtime/Experimental/Dataflow/Core/Public/Dataflow/DataflowElement.h b/Engine/Source/Runtime/Experimental/Dataflow/Core/Public/Dataflow/DataflowElement.h
new file mode 100644
--- /dev/null
+++ b/Engine/Source/Runtime/Experimental/Dataflow/Core/Public/Dataflow/DataflowElement.h
@@ -0,0 +1,44 @@
+// Copyright Epic Games, Inc. All Rights Reserved.
+
+#pragma once
+
+#include "CoreMinimal.h"
+
+#include "DataflowElement.generated.h"
+
+/** Base dataflow scene element that could be used in outliner and for rendering */
+USTRUCT()
+struct FDataflowBaseElement
+{
+	GENERATED_BODY()
+
+	FDataflowBaseElement() {}
+
+	FDataflowBaseElement(const FString& InElementName, FDataflowBaseElement* InParentElement, const FBox& InBoundingBox, const bool bInIsConstruction) :
+		ElementName(InElementName), ParentElement(InParentElement), BoundingBox(InBoundingBox), bIsConstruction(bInIsConstruction) {}
+
+	virtual ~FDataflowBaseElement() {}
+
+	static FName StaticType() { return FName("FDataflowBaseElement"); }
+
+	/** Check of the element type */
+	virtual bool IsA(FName InType) const { return InType.ToString().Equals(StaticType().ToString());}
+	
+	/** Element name to be used to retrieve */
+	FString ElementName = TEXT("None");
+
+	/** Parent element to build the hierarchy if necessary */
+	FDataflowBaseElement* ParentElement = nullptr;
+
+	/** Bounding box to focus */
+	FBox BoundingBox = FBox();
+	
+	/** Construction flag */
+	bool bIsConstruction = true;
+	
+	/** Visible flag to enable/disable the rendering */
+	bool bIsVisible = true;
+
+	/** Selection flag */
+	bool bIsSelected = false;
+};
\ No newline at end of file
diff --git a/Engine/Source/Runtime/Experimental/Dataflow/Core/Public/Dataflow/DataflowNode.h b/Engine/Source/Runtime/Experimental/Dataflow/Core/Public/Dataflow/DataflowNode.h
--- a/Engine/Source/Runtime/Experimental/Dataflow/Core/Public/Dataflow/DataflowNode.h
+++ b/Engine/Source/Runtime/Experimental/Dataflow/Core/Public/Dataflow/DataflowNode.h
@@ -15,6 +15,7 @@
 #include "DataflowNode.generated.h"
 
 class UScriptStruct;
+struct FDataflowBaseElement;
 
 namespace UE::Dataflow {
 	struct FNodeParameters {
@@ -115,9 +116,14 @@ PRAGMA_ENABLE_DEPRECATION_WARNINGS
 #if WITH_EDITOR
 	struct FDebugDrawParameters
 	{
-		bool bNodeIsSelected;
-		bool bNodeIsPinned;
-		FName CurrentViewMode;
+		/** Boolean to check if the node is selected or not */
+		bool bNodeIsSelected = false;
+
+		/** Boolean to check if the node is pinned or not */
+		bool bNodeIsPinned = false;
+
+		/** Current view mode to use in the debug draw */
+		FName CurrentViewMode = NAME_None;
 	};
 
 	/** Override if DebugDraw has been implemented.*/
diff --git a/Engine/Source/Runtime/Experimental/Dataflow/Engine/Private/Dataflow/DataflowDebugDraw.cpp b/Engine/Source/Runtime/Experimental/Dataflow/Engine/Private/Dataflow/DataflowDebugDraw.cpp
--- a/Engine/Source/Runtime/Experimental/Dataflow/Engine/Private/Dataflow/DataflowDebugDraw.cpp
+++ b/Engine/Source/Runtime/Experimental/Dataflow/Engine/Private/Dataflow/DataflowDebugDraw.cpp
@@ -5,8 +5,8 @@
 #include "DynamicMesh/DynamicMesh3.h"
 #include "DynamicMeshBuilder.h"
 
-FDataflowDebugDraw::FDataflowDebugDraw(FDataflowDebugRenderSceneProxy* DebugRenderSceneProxy) :
-	DebugRenderSceneProxy(DebugRenderSceneProxy)
+FDataflowDebugDraw::FDataflowDebugDraw(FDataflowDebugRenderSceneProxy* InDebugRenderSceneProxy, IDataflowDebugDrawInterface::FDataflowElementsType& InDataflowElements) :
+	DebugRenderSceneProxy(InDebugRenderSceneProxy), DataflowElements(InDataflowElements)
 {
 	DebugRenderSceneProxy->DrawAlpha = 1;		// uint Alpha values below this are not rendered (default is 100 but we want to allow more translucency)
 	DebugRenderSceneProxy->ClearAll();
@@ -114,6 +114,11 @@ void FDataflowDebugDraw::ReservePoints(int32 NumAdditionalPoints)
 	DebugRenderSceneProxy->ReservePoints(NumAdditionalPoints);
 }
 
+void FDataflowDebugDraw::DrawObject(const TRefCountPtr<IDataflowDebugDrawObject>& Object)
+{
+	DebugRenderSceneProxy->AddObject(Object);
+}
+
 void FDataflowDebugDraw::DrawPoint(const FVector& Position)
 {
 	FDataflowDebugRenderSceneProxy::FDebugPoint NewPoint
diff --git a/Engine/Source/Runtime/Experimental/Dataflow/Engine/Private/Dataflow/DataflowDebugDrawComponent.cpp b/Engine/Source/Runtime/Experimental/Dataflow/Engine/Private/Dataflow/DataflowDebugDrawComponent.cpp
--- a/Engine/Source/Runtime/Experimental/Dataflow/Engine/Private/Dataflow/DataflowDebugDrawComponent.cpp
+++ b/Engine/Source/Runtime/Experimental/Dataflow/Engine/Private/Dataflow/DataflowDebugDrawComponent.cpp
@@ -3,6 +3,8 @@
 #include "Dataflow/DataflowDebugDrawComponent.h"
 #include "MeshElementCollector.h"
 #include "PrimitiveDrawInterface.h"
+#include "Dataflow/DataflowDebugDraw.h"
+#include "Dataflow/DataflowDebugDrawObject.h"
 
 FDebugRenderSceneProxy* UDataflowDebugDrawComponent::CreateDebugSceneProxy()
 {
@@ -81,6 +83,13 @@ FBoxSphereBounds UDataflowDebugDrawComponent::CalcBounds(const FTransform& Local
 		{
 			BoundsBuilder += Point.Position;
 		}
+		for (const TRefCountPtr<IDataflowDebugDrawObject>& Object : DebugSceneProxy->Objects)
+		{
+			if(Object.IsValid() && Object->IsA(FDataflowDebugDrawBaseObject::StaticType()))
+			{
+				BoundsBuilder += static_cast<FDataflowDebugDrawBaseObject*>(Object.GetReference())->ComputeBoundingBox();
+			}
+		}
 	}
 
 	FBoxSphereBounds ReturnBounds(BoundsBuilder);
@@ -111,6 +120,7 @@ void FDataflowDebugRenderSceneProxy::ClearAll()
 	Capsules.Reset();
 	CoordinateSystems.Reset();
 	Points.Reset();
+	Objects.Reset();
 }
 
 
@@ -134,15 +144,29 @@ void FDataflowDebugRenderSceneProxy::AddPoint(const FDebugPoint& Point)
 	Points.Add(Point);
 }
 
+void FDataflowDebugRenderSceneProxy::AddObject(const TRefCountPtr<IDataflowDebugDrawObject>& Object)
+{
+	Objects.Add(Object);
+}
+
 void FDataflowDebugRenderSceneProxy::GetDynamicMeshElementsForView(const FSceneView* View, const int32 ViewIndex, const FSceneViewFamily& ViewFamily, const uint32 VisibilityMap, FMeshElementCollector& Collector, FMaterialCache& DefaultMaterialCache, FMaterialCache& SolidMeshMaterialCache) const
 {
 	FDebugRenderSceneProxy::GetDynamicMeshElementsForView(View, ViewIndex, ViewFamily, VisibilityMap, Collector, DefaultMaterialCache, SolidMeshMaterialCache);
 
 	FPrimitiveDrawInterface* PDI = Collector.GetPDI(ViewIndex);
 
-	// Draw Lines
+	// Draw Points
 	for (const FDebugPoint& Point : Points)
 	{
 		PDI->DrawPoint(Point.Position, Point.Color, Point.Size, SDPG_World);
 	}
+
+	// Draw Objects
+	for (const TRefCountPtr<IDataflowDebugDrawObject>& Object : Objects)
+	{
+		if(Object.IsValid() && Object->IsA(FDataflowDebugDrawBaseObject::StaticType()))
+		{
+			static_cast<FDataflowDebugDrawBaseObject*>(Object.GetReference())->DrawDataflowElements(PDI);
+		}
+	}
 }
\ No newline at end of file
diff --git a/Engine/Source/Runtime/Experimental/Dataflow/Engine/Private/Dataflow/DataflowDebugDrawObject.cpp b/Engine/Source/Runtime/Experimental/Dataflow/Engine/Private/Dataflow/DataflowDebugDrawObject.cpp
new file mode 100644
--- /dev/null
+++ b/Engine/Source/Runtime/Experimental/Dataflow/Engine/Private/Dataflow/DataflowDebugDrawObject.cpp
@@ -0,0 +1,5 @@
+﻿// Copyright Epic Games, Inc. All Rights Reserved.
+
+#include "Dataflow/DataflowDebugDrawObject.h"
+
+IMPLEMENT_HIT_PROXY(HDataflowElementHitProxy, HHitProxy)
\ No newline at end of file
diff --git a/Engine/Source/Runtime/Experimental/Dataflow/Engine/Public/Dataflow/DataflowDebugDraw.h b/Engine/Source/Runtime/Experimental/Dataflow/Engine/Public/Dataflow/DataflowDebugDraw.h
--- a/Engine/Source/Runtime/Experimental/Dataflow/Engine/Public/Dataflow/DataflowDebugDraw.h
+++ b/Engine/Source/Runtime/Experimental/Dataflow/Engine/Public/Dataflow/DataflowDebugDraw.h
@@ -13,11 +13,13 @@
 
 #include "DataflowDebugDraw.generated.h"
 
+struct FDataflowBaseElement;
+
 class FDataflowDebugDraw : public IDataflowDebugDrawInterface
 {
 public:
 
-	DATAFLOWENGINE_API FDataflowDebugDraw(class FDataflowDebugRenderSceneProxy* DebugRenderSceneProxy);
+	DATAFLOWENGINE_API FDataflowDebugDraw(class FDataflowDebugRenderSceneProxy* DebugRenderSceneProxy, IDataflowDebugDrawInterface::FDataflowElementsType& DataflowElements);
 
 	virtual ~FDataflowDebugDraw() = default;
 
@@ -34,6 +36,7 @@ private:
 	virtual void ResetAllState() override;
 
 	virtual void ReservePoints(int32 NumAdditionalPoints) override;
+	virtual void DrawObject(const TRefCountPtr<IDataflowDebugDrawObject>& Object) override;
 	virtual void DrawPoint(const FVector& Position) override;
 	virtual void DrawLine(const FVector& Start, const FVector& End) const override;
 	virtual void DrawMesh(const IDebugDrawMesh& Mesh) const override;
@@ -45,8 +48,18 @@ private:
 	virtual void DrawOverlayText(const FString& InString) override;
 	virtual FString GetOverlayText() const override;
 
+	/** Dataflow elements non const accessor */
+	virtual IDataflowDebugDrawInterface::FDataflowElementsType& ModifyDataflowElements() { return DataflowElements; }
+
+	/** Dataflow elements const accessor */
+	virtual const IDataflowDebugDrawInterface::FDataflowElementsType& GetDataflowElements() const { return DataflowElements; }
+
+	/** Scene proxy coming from the dataflow debug draw component */
 	class FDataflowDebugRenderSceneProxy* DebugRenderSceneProxy = nullptr;
 
+	/** List of dataflow elements to filled by the debug draw */ 
+	IDataflowDebugDrawInterface::FDataflowElementsType& DataflowElements;
+
 	FLinearColor Color = FLinearColor::White;
 	double LineWidth = 1.0;
 	float PointSize = 5.0;
diff --git a/Engine/Source/Runtime/Experimental/Dataflow/Engine/Public/Dataflow/DataflowDebugDrawComponent.h b/Engine/Source/Runtime/Experimental/Dataflow/Engine/Public/Dataflow/DataflowDebugDrawComponent.h
--- a/Engine/Source/Runtime/Experimental/Dataflow/Engine/Public/Dataflow/DataflowDebugDrawComponent.h
+++ b/Engine/Source/Runtime/Experimental/Dataflow/Engine/Public/Dataflow/DataflowDebugDrawComponent.h
@@ -2,9 +2,12 @@
 
 #pragma once 
 
+#include "DataflowContent.h"
 #include "Debug/DebugDrawComponent.h"
 #include "DataflowDebugDrawComponent.generated.h"
 
+struct IDataflowDebugDrawObject;
+
 UCLASS(MinimalAPI)
 class UDataflowDebugDrawComponent : public UDebugDrawComponent
 {
@@ -15,7 +18,6 @@ private:
 	virtual FBoxSphereBounds CalcBounds(const FTransform& LocalToWorld) const override;
 };
 
-
 class FDataflowDebugRenderSceneProxy final : public FDebugRenderSceneProxy
 {
 public:
@@ -31,14 +33,24 @@ public:
 		float	Size = 0;
 		FLinearColor Color = FLinearColor::White;
 	};
+	
+	/** Add point to the scene proxy */
+	void AddPoint(const FDebugPoint& Point);
 
+	/** Add dataflow object to the scene proxy */
+	void AddObject(const TRefCountPtr<IDataflowDebugDrawObject>& Object);
+
+	/** Reserve the number of points to add */
 	void ReservePoints(int32 NumAdditionalPoints);
-	void AddPoint(const FDebugPoint& Point);
 
 private:
+		
 	friend class UDataflowDebugDrawComponent;
 	TArray<FDebugPoint> Points;
 
+	/** List of user provided dataflow objects to draw */
+	TArray<TRefCountPtr<IDataflowDebugDrawObject>> Objects;
+
 private:
 	virtual void GetDynamicMeshElementsForView(const FSceneView* View, const int32 ViewIndex, const FSceneViewFamily& ViewFamily, const uint32 VisibilityMap, FMeshElementCollector& Collector, FMaterialCache& DefaultMaterialCache, FMaterialCache& SolidMeshMaterialCache) const override;
 	virtual FPrimitiveViewRelevance GetViewRelevance(const FSceneView* View) const override;
diff --git a/Engine/Source/Runtime/Experimental/Dataflow/Engine/Public/Dataflow/DataflowDebugDrawObject.h b/Engine/Source/Runtime/Experimental/Dataflow/Engine/Public/Dataflow/DataflowDebugDrawObject.h
new file mode 100644
--- /dev/null
+++ b/Engine/Source/Runtime/Experimental/Dataflow/Engine/Public/Dataflow/DataflowDebugDrawObject.h
@@ -0,0 +1,106 @@
+// Copyright Epic Games, Inc. All Rights Reserved.
+
+#pragma once 
+
+#include "HitProxies.h"
+#include "Dataflow/DataflowDebugDrawInterface.h"
+#include "Dataflow/DataflowElement.h"
+#include "GenericPlatform/ICursor.h"
+
+#include "DataflowDebugDrawObject.generated.h"
+
+class FPrimitiveDrawInterface;
+
+/** Dataflow object debug draw parent class */
+struct DATAFLOWENGINE_API FDataflowDebugDrawBaseObject : public IDataflowDebugDrawObject
+{
+	using Super = IDataflowDebugDrawObject;
+	
+	static FName StaticType() { return FName("FDataflowDebugDrawBaseObject"); }
+	
+	FDataflowDebugDrawBaseObject(IDataflowDebugDrawInterface::FDataflowElementsType& InDataflowElements) :
+		IDataflowDebugDrawObject(), DataflowElements(InDataflowElements)
+	{}
+	
+	/** Populate dataflow elements */
+	virtual void PopulateDataflowElements() = 0;
+
+	/** Debug draw dataflow element */
+	virtual void DrawDataflowElements(FPrimitiveDrawInterface* PDI) const = 0;
+
+	/** Compute teh dataflow elements bounding box */
+	virtual FBox ComputeBoundingBox() const = 0;
+
+	/** Check of the object type */
+	virtual bool IsA(FName InType) const override
+	{
+		return InType.ToString().Equals(StaticType().ToString()) 
+			|| Super::IsA(InType); 
+	}
+protected :
+
+	/** List of dataflow elements the debug draw object is going to populate and render */
+	IDataflowDebugDrawInterface::FDataflowElementsType& DataflowElements;
+
+	/** Elements offset in the global array */
+	int32 ElementsOffset = 0;
+
+	/** Elements size in the global array */
+	int32 ElementsSize = 0;
+};
+
+template<typename ObjectType, typename... Args>
+FORCEINLINE TRefCountPtr<ObjectType> MakeDebugDrawObject(Args&& ... args)
+{
+	TRefCountPtr<ObjectType> DataflowObject = MakeRefCount<ObjectType>(std::forward<Args>(args)...);
+	DataflowObject->PopulateDataflowElements();
+
+	return DataflowObject;
+}
+
+/** Dataflow hit proxy for viewport selection */
+struct DATAFLOWENGINE_API HDataflowElementHitProxy : public HHitProxy
+{
+	DECLARE_HIT_PROXY()
+
+	HDataflowElementHitProxy(int32 InElementIndex, FName InElementName)
+		: HHitProxy(HPP_Foreground)
+		, ElementIndex(InElementIndex)
+		, ElementName(InElementName)
+	{}
+
+	//~ Begin HHitProxy interface
+	virtual EMouseCursor::Type GetMouseCursor() override { return EMouseCursor::Crosshairs; }
+	//~ End HHitProxy interface
+
+	/** Element index to retrieve the matching dataflow element */
+	int32 ElementIndex;
+
+	/** Element name to retrieve the matching dataflow element */
+	FName ElementName;
+};
+
+/** Proxy dataflow scene element that contains a ref to HDataflowElementHitProxy */
+USTRUCT()
+struct FDataflowProxyElement : public FDataflowBaseElement
+{
+	GENERATED_BODY()
+
+	FDataflowProxyElement() {}
+
+	FDataflowProxyElement(const FString& InElementName, FDataflowBaseElement* InParentElement, const FBox& InBoundingBox, const bool bInIsConstruction) :
+		FDataflowBaseElement(InElementName, InParentElement, InBoundingBox, bInIsConstruction ), ElementProxy(nullptr) {}
+
+	/** Element proxy used for selection */
+	TRefCountPtr<HHitProxy> ElementProxy = nullptr;
+
+	static FName StaticType() { return FName("FDatafloProxyElement"); }
+
+	/** Check of the element type */
+	virtual bool IsA(FName InType) const override
+	{
+		return InType.ToString().Equals(StaticType().ToString())
+			|| Super::IsA(InType);
+	}
+};
+
diff --git a/Engine/Source/Runtime/Experimental/Dataflow/Simulation/Public/Dataflow/DataflowSimulationContext.h b/Engine/Source/Runtime/Experimental/Dataflow/Simulation/Public/Dataflow/DataflowSimulationContext.h
--- a/Engine/Source/Runtime/Experimental/Dataflow/Simulation/Public/Dataflow/DataflowSimulationContext.h
+++ b/Engine/Source/Runtime/Experimental/Dataflow/Simulation/Public/Dataflow/DataflowSimulationContext.h
@@ -11,7 +11,7 @@ namespace UE::Dataflow
 {
 	/** Simulation context that will be used by all the simulation nodes*/
 	template<typename Base = UE::Dataflow::FContextSingle>
-	class DATAFLOWSIMULATION_API TSimulationContext : public UE::Dataflow::TEngineContext<Base>
+	class TSimulationContext : public UE::Dataflow::TEngineContext<Base>
 	{
 	public:
 		DATAFLOW_CONTEXT_INTERNAL(UE::Dataflow::TEngineContext<Base>, TSimulationContext);
@@ -36,25 +36,25 @@ namespace UE::Dataflow
 		void GetTypedProxies(TArray<ProxyType*>& FilteredProxies) const;
 
 		/** Filter the physics solvers matching the groups */
-		void GetSimulationProxies(const FString& ProxyType, const TArray<FString>& SimulationGroups, TArray<FDataflowSimulationProxy*>& FilteredProxies) const;
+		DATAFLOWSIMULATION_API void GetSimulationProxies(const FString& ProxyType, const TArray<FString>& SimulationGroups, TArray<FDataflowSimulationProxy*>& FilteredProxies) const;
 
 		/** Add simulation proxy to the context */
-		void AddSimulationProxy(const FString& ProxyType, FDataflowSimulationProxy* SimulationProxy);
+		DATAFLOWSIMULATION_API void AddSimulationProxy(const FString& ProxyType, FDataflowSimulationProxy* SimulationProxy);
 		
 		/** Remove simulation proxy from the context */
-    	void RemoveSimulationProxy(const FString& ProxyType, const FDataflowSimulationProxy* SimulationProxy);
+    	DATAFLOWSIMULATION_API void RemoveSimulationProxy(const FString& ProxyType, const FDataflowSimulationProxy* SimulationProxy);
 
 		/** Reset all the simulation proxies */
-		void ResetSimulationProxies();
+		DATAFLOWSIMULATION_API void ResetSimulationProxies();
 		
 		/** Return the number of physics solvers */
-		int32 NumSimulationProxies(const FString& ProxyType) const;
+		DATAFLOWSIMULATION_API int32 NumSimulationProxies(const FString& ProxyType) const;
 
 		/** Register all the proxy groups used in the proxy */
-		void RegisterProxyGroups();
+		DATAFLOWSIMULATION_API void RegisterProxyGroups();
 
 		/** Build group bits based on the string array */
-		void BuildGroupBits(const TArray<FString>& SimulationGroups, TBitArray<>& GroupBits) const;
+		DATAFLOWSIMULATION_API void BuildGroupBits(const TArray<FString>& SimulationGroups, TBitArray<>& GroupBits) const;
 
 		/** Push another level of iteration indices */
 		void PushIterationIndex()
