Dataflow : fix evaluation failing to re-evaluate some branches when the graph had parrallel branches
During invalidation,the terminal node would be reached before all the necessary branches are invalided ( because of the recursive nature of invalidation )
In order to fix this we are going back to a tick based model, but that only evaluate terminal nodes that have been explicitly called for evaluation

Also make sure update view code mark the construction view as dirty to avoid seeing previous results

diff --git a/Engine/Plugins/Experimental/Dataflow/Source/DataflowEditor/Private/Dataflow/DataflowEditorToolkit.cpp b/Engine/Plugins/Experimental/Dataflow/Source/DataflowEditor/Private/Dataflow/DataflowEditorToolkit.cpp
--- a/Engine/Plugins/Experimental/Dataflow/Source/DataflowEditor/Private/Dataflow/DataflowEditorToolkit.cpp
+++ b/Engine/Plugins/Experimental/Dataflow/Source/DataflowEditor/Private/Dataflow/DataflowEditorToolkit.cpp
@@ -1561,15 +1561,11 @@ void FDataflowEditorToolkit::EvaluateGraph()
 		{
 			if (TSharedPtr<UE::Dataflow::FGraph> DataflowGraph = DataflowAsset->GetDataflow())
 			{
-				for (TSharedPtr<FDataflowNode> Node : DataflowGraph->GetNodes())
+				for (const TSharedPtr<FDataflowNode>& DataflowNode : DataflowGraph->GetFilteredNodes(FDataflowTerminalNode::StaticType()))
 				{
-					if (Node)
+					if (const FDataflowTerminalNode* TerminalNode = DataflowNode ? DataflowNode->AsType<FDataflowTerminalNode>() : nullptr)
 					{
-						if (FDataflowTerminalNode* TerminalNode = Node->AsType<FDataflowTerminalNode>())
-						{
-							EvaluateTerminalNode(*TerminalNode);
-							bViewsNeedRefresh = true;
-						}
+						NodesToEvaluateOnTick.Add(TerminalNode->GetGuid());
 					}
 				}
 			}
@@ -1588,8 +1584,7 @@ void FDataflowEditorToolkit::OnNodeInvalidated(UDataflow& DataflowAsset, FDatafl
 				// evaluate only active terminals nodes 
 				if (FDataflowTerminalNode* TerminalNode = Node.AsType<FDataflowTerminalNode>())
 				{
-					EvaluateTerminalNode(*TerminalNode);
-					bViewsNeedRefresh = true;
+					NodesToEvaluateOnTick.Add(TerminalNode->GetGuid());
 				}
 			}
 		}
@@ -1844,6 +1839,29 @@ void FDataflowEditorToolkit::Tick(float DeltaTime)
 
 			// Update the list of dataflow terminal contents 
 			DataflowEditor->UpdateTerminalContents(InitTimeStamp);
+
+			// evaluate the node that are requested from NodesToEvaluateOnTick
+			if (NodesToEvaluateOnTick.Num() > 0)
+			{
+				if (UDataflow* DataflowAsset = EditorContent->GetDataflowAsset())
+				{
+					if (TSharedPtr<UE::Dataflow::FGraph> DataflowGraph = DataflowAsset->GetDataflow())
+					{
+						for (const TSharedPtr<FDataflowNode>& DataflowNode : DataflowGraph->GetFilteredNodes(FDataflowTerminalNode::StaticType()))
+						{
+							if (const FDataflowTerminalNode* TerminalNode = DataflowNode? DataflowNode->AsType<FDataflowTerminalNode>(): nullptr)
+							{
+								if (NodesToEvaluateOnTick.Contains(TerminalNode->GetGuid()))
+								{
+									EvaluateTerminalNode(*TerminalNode);
+									bViewsNeedRefresh = true;
+								}
+							}
+						}
+					}
+				}
+				NodesToEvaluateOnTick.Reset();
+			}
 		}
 	}
 	if(ConstructionScene->IsSceneDirty() || SimulationScene->IsSceneDirty())
@@ -1915,6 +1933,12 @@ void FDataflowEditorToolkit::RefreshViewsIfNeeded(bool bForce)
 		{
 			Listener->RefreshView();
 		}
+
+		if (const TObjectPtr<UDataflowBaseContent>& EditorContent = GetEditorContent())
+		{
+			EditorContent->SetConstructionDirty(true);
+		}
+
 		bViewsNeedRefresh = false;
 	}
 }
diff --git a/Engine/Plugins/Experimental/Dataflow/Source/DataflowEditor/Public/Dataflow/DataflowEditorToolkit.h b/Engine/Plugins/Experimental/Dataflow/Source/DataflowEditor/Public/Dataflow/DataflowEditorToolkit.h
--- a/Engine/Plugins/Experimental/Dataflow/Source/DataflowEditor/Public/Dataflow/DataflowEditorToolkit.h
+++ b/Engine/Plugins/Experimental/Dataflow/Source/DataflowEditor/Public/Dataflow/DataflowEditorToolkit.h
@@ -329,6 +329,8 @@ private:
 	/** Document tracker for dynamic tabs ( like subgraphs ) */
 	TSharedPtr<FDocumentTracker> DocumentManager;
 
+	TSet<FGuid> NodesToEvaluateOnTick;
+
 	FDateTime GraphEvaluationBegin;
 	FDateTime GraphEvaluationFinished;
 	bool bViewsNeedRefresh = false;
