Dataflow : fix struct and array variable override not refreshing the graph when some of their deeper propertties were changed
- better logic to extract the exac variable name from the stack of properties updated

diff --git a/Engine/Plugins/Experimental/Dataflow/Source/DataflowEditor/Private/Dataflow/DataflowMembersWidget.cpp b/Engine/Plugins/Experimental/Dataflow/Source/DataflowEditor/Private/Dataflow/DataflowMembersWidget.cpp
--- a/Engine/Plugins/Experimental/Dataflow/Source/DataflowEditor/Private/Dataflow/DataflowMembersWidget.cpp
+++ b/Engine/Plugins/Experimental/Dataflow/Source/DataflowEditor/Private/Dataflow/DataflowMembersWidget.cpp
@@ -489,18 +489,65 @@ IDataflowInstanceInterface* SDataflowMembersWidget::GetDataflowInstanceInterface
 	return nullptr;
 }
 
+namespace DataflowMembersWidget::Private
+{
+	FName ExtractVariableNameFromPropertyChangeEvent(const FPropertyChangedEvent& InPropertyChangedEvent)
+	{
+		// check if the property changed is deep inside a stackl of property ( like a member of an eleent of an array )
+		TMap<FString, int32> PropertyNameStack; // ordered from deeper to 
+		InPropertyChangedEvent.GetArrayIndicesPerObject(0, PropertyNameStack);
+
+		if (PropertyNameStack.Num() > 0)
+		{
+			TArray<FString> PropertyNames;
+			PropertyNameStack.GetKeys(PropertyNames);
+			// the 3 last ones should always be VariableOverrides / Variables / Value: 
+			//   - VariableOverrides because that's the name of the FDataflowInstance member
+			//	 - Variables because that the name of the property bag property in FDataflowVariableOverrides
+			//	 - Value becasue this is the container of properties in FInstancedPropertyBag
+			const bool bIsValid =
+				(PropertyNames.Num() > 4) &&
+				(PropertyNames[PropertyNames.Num() - 1] == FDataflowInstance::GetVariableOverridesPropertyName()) &&
+				(PropertyNames[PropertyNames.Num() - 2] == FDataflowVariableOverrides::GetVariablePropertyName()) &&
+				(PropertyNames[PropertyNames.Num() - 3] == "Value") // no programmaic access to the name 
+				;
+			if (bIsValid)
+			{
+				// variable name should be here :)
+				return FName(PropertyNames[PropertyNames.Num() - 4]);
+			}
+		}
+
+		return InPropertyChangedEvent.GetPropertyName();
+	}
+}
+
 void SDataflowMembersWidget::OverridesDetailsViewFinishedChangingProperties(const FPropertyChangedEvent& InPropertyChangedEvent)
 {
-	const FName VariableName = InPropertyChangedEvent.GetPropertyName();
+///InPropertyChangedEvent.Property->GetOwnerProperty
+
 	if (TStrongObjectPtr<UDataflow> DataflowAsset = DataflowAssetWeakPtr.Pin())
 	{
-		for (const TSharedPtr<FDataflowNode>& Node : DataflowAsset->Dataflow->GetNodes())
+		if (IDataflowInstanceInterface* InstanceInterface = GetDataflowInstanceInterface())
 		{
-			if (FGetDataflowVariableNode* VariableNode =  Node->AsType<FGetDataflowVariableNode>())
+			const FDataflowInstance& DataflowInstance = InstanceInterface->GetDataflowInstance();
+
+			TMap<FString, int32> PropertyNameStack; // ordered from deeper to 
+			InPropertyChangedEvent.GetArrayIndicesPerObject(0, PropertyNameStack);
+
+			const FName VariableName = DataflowMembersWidget::Private::ExtractVariableNameFromPropertyChangeEvent(InPropertyChangedEvent);
+			if (DataflowAsset->Variables.FindPropertyDescByName(VariableName))
 			{
-				if (VariableNode->GetVariableName() == VariableName)
+				// invalidate all the get nodes all the nodes that match the variable name 
+				for (const TSharedPtr<FDataflowNode>& Node : DataflowAsset->Dataflow->GetNodes())
 				{
-					VariableNode->Invalidate();
+					if (FGetDataflowVariableNode* VariableNode = Node->AsType<FGetDataflowVariableNode>())
+					{
+						if (VariableNode->GetVariableName() == VariableName)
+						{
+							VariableNode->Invalidate();
+						}
+					}
 				}
 			}
 		}
diff --git a/Engine/Source/Runtime/Experimental/Dataflow/Engine/Private/Dataflow/DataflowInstance.cpp b/Engine/Source/Runtime/Experimental/Dataflow/Engine/Private/Dataflow/DataflowInstance.cpp
--- a/Engine/Source/Runtime/Experimental/Dataflow/Engine/Private/Dataflow/DataflowInstance.cpp
+++ b/Engine/Source/Runtime/Experimental/Dataflow/Engine/Private/Dataflow/DataflowInstance.cpp
@@ -313,6 +313,11 @@ FName FDataflowInstance::GetDataflowAssetPropertyName()
 	return GET_MEMBER_NAME_CHECKED(FDataflowInstance, DataflowAsset);
 }
 
+FName FDataflowInstance::GetVariableOverridesPropertyName()
+{
+	return GET_MEMBER_NAME_CHECKED(FDataflowInstance, VariableOverrides);
+}
+
 void FDataflowInstance::SetDataflowAsset(UDataflow* InDataflowAsset)
 {
 	if (DataflowAsset != InDataflowAsset)
diff --git a/Engine/Source/Runtime/Experimental/Dataflow/Engine/Public/Dataflow/DataflowInstance.h b/Engine/Source/Runtime/Experimental/Dataflow/Engine/Public/Dataflow/DataflowInstance.h
--- a/Engine/Source/Runtime/Experimental/Dataflow/Engine/Public/Dataflow/DataflowInstance.h
+++ b/Engine/Source/Runtime/Experimental/Dataflow/Engine/Public/Dataflow/DataflowInstance.h
@@ -108,6 +108,9 @@ public:
 	/** Get the dataflow asset member property name */
 	DATAFLOWENGINE_API static FName GetDataflowAssetPropertyName();
 
+	/** Get the variable overrides member property name */
+	DATAFLOWENGINE_API static FName GetVariableOverridesPropertyName();
+
 	DATAFLOWENGINE_API const FDataflowVariableOverrides& GetVariableOverrides() const;
 
 	DATAFLOWENGINE_API FDataflowVariableOverrides& GetVariableOverrides();
