Dataflow : add performance data display conrol to the evaluation button option menu
- also remove the console variable previously controlling it

diff --git a/Engine/Plugins/Experimental/Dataflow/Source/DataflowEditor/Private/Dataflow/DataflowEditorCommands.cpp b/Engine/Plugins/Experimental/Dataflow/Source/DataflowEditor/Private/Dataflow/DataflowEditorCommands.cpp
--- a/Engine/Plugins/Experimental/Dataflow/Source/DataflowEditor/Private/Dataflow/DataflowEditorCommands.cpp
+++ b/Engine/Plugins/Experimental/Dataflow/Source/DataflowEditor/Private/Dataflow/DataflowEditorCommands.cpp
@@ -62,6 +62,7 @@ void FDataflowEditorCommandsImpl::RegisterCommands()
 	UI_COMMAND(EvaluateGraphAutomatic, "Automatic Graph Evaluation", "Set the evaluation mode of the graph to Automatic.", EUserInterfaceActionType::RadioButton, FInputChord());
 	UI_COMMAND(EvaluateGraphManual, "Manual Graph Evaluation", "Set the evaluation mode of the graph to Manual.", EUserInterfaceActionType::RadioButton, FInputChord());
 	UI_COMMAND(ClearGraphCache, "Clear Graph Cache", "Clear the graph Data cache.", EUserInterfaceActionType::Button, FInputChord());
+	UI_COMMAND(TogglePerfData, "Performance Data", "Toggle the evaluation performance data for each node.", EUserInterfaceActionType::ToggleButton, FInputChord());
 	UI_COMMAND(FreezeNodes, "FreezeNodes", "Freeze the evaluation of the selected nodes.", EUserInterfaceActionType::Button, FInputChord());
 	UI_COMMAND(UnfreezeNodes, "UnfreezeNodes", "Unfreeze the evaluation of the selected nodes.", EUserInterfaceActionType::Button, FInputChord());
 	UI_COMMAND(CreateComment, "CreateComment", "Create a Comment node.", EUserInterfaceActionType::None, FInputChord());
diff --git a/Engine/Plugins/Experimental/Dataflow/Source/DataflowEditor/Private/Dataflow/DataflowEditorToolkit.cpp b/Engine/Plugins/Experimental/Dataflow/Source/DataflowEditor/Private/Dataflow/DataflowEditorToolkit.cpp
--- a/Engine/Plugins/Experimental/Dataflow/Source/DataflowEditor/Private/Dataflow/DataflowEditorToolkit.cpp
+++ b/Engine/Plugins/Experimental/Dataflow/Source/DataflowEditor/Private/Dataflow/DataflowEditorToolkit.cpp
@@ -763,6 +763,29 @@ void FDataflowEditorToolkit::SetEvaluateGraphMode(EDataflowEditorEvaluationMode
 	}
 }
 
+void FDataflowEditorToolkit::TogglePerfData()
+{
+	if (TObjectPtr<UDataflowBaseContent>& EditorContent = GetEditorContent())
+	{
+		if (TSharedPtr<UE::Dataflow::FEngineContext> DataflowContext = EditorContent->GetDataflowContext())
+		{
+			DataflowContext->EnablePerfData(!DataflowContext->IsPerfDataEnabled());
+		}
+	}
+}
+
+bool FDataflowEditorToolkit::IsPerfDataEnabled() const
+{
+	if (const TObjectPtr<UDataflowBaseContent>& EditorContent = GetEditorContent())
+	{
+		if (TSharedPtr<UE::Dataflow::FEngineContext> DataflowContext = EditorContent->GetDataflowContext())
+		{
+			return DataflowContext->IsPerfDataEnabled();
+		}
+	}
+	return false;
+}
+
 void FDataflowEditorToolkit::ClearGraphCache()
 {
 	if (TObjectPtr<UDataflowBaseContent>& EditorContent = GetEditorContent())
@@ -848,6 +871,11 @@ void FDataflowEditorToolkit::AddEvaluationWidget(FToolMenuSection& InSection)
 		FExecuteAction::CreateSP(this, &FDataflowEditorToolkit::ClearGraphCache),
 		FCanExecuteAction::CreateSP(this, &FDataflowEditorToolkit::CanClearGraphCache)
 	);
+	ToolkitCommands->MapAction(FDataflowEditorCommands::Get().TogglePerfData,
+		FExecuteAction::CreateSP(this, &FDataflowEditorToolkit::TogglePerfData),
+		FCanExecuteAction(),
+		FIsActionChecked::CreateSP(this, &FDataflowEditorToolkit::IsPerfDataEnabled)
+	);
 
 	InSection.AddEntry(FToolMenuEntry::InitToolBarButton(
 		FDataflowEditorCommands::Get().EvaluateGraph,
@@ -882,6 +910,8 @@ TSharedRef<SWidget> FDataflowEditorToolkit::GenerateEvaluationOptionsMenu()
 	MenuBuilder.AddMenuEntry(FDataflowEditorCommands::Get().EvaluateGraphAutomatic);
 	MenuBuilder.AddMenuEntry(FDataflowEditorCommands::Get().EvaluateGraphManual);
 	MenuBuilder.AddSeparator();
+	MenuBuilder.AddMenuEntry(FDataflowEditorCommands::Get().TogglePerfData);
+	MenuBuilder.AddSeparator();
 	MenuBuilder.AddMenuEntry(FDataflowEditorCommands::Get().ClearGraphCache);
 	MenuBuilder.EndSection();
 	return MenuBuilder.MakeWidget();
diff --git a/Engine/Plugins/Experimental/Dataflow/Source/DataflowEditor/Private/Dataflow/DataflowSNode.cpp b/Engine/Plugins/Experimental/Dataflow/Source/DataflowEditor/Private/Dataflow/DataflowSNode.cpp
--- a/Engine/Plugins/Experimental/Dataflow/Source/DataflowEditor/Private/Dataflow/DataflowSNode.cpp
+++ b/Engine/Plugins/Experimental/Dataflow/Source/DataflowEditor/Private/Dataflow/DataflowSNode.cpp
@@ -30,9 +30,6 @@
 
 namespace UE::Dataflow::Private
 {
-	static bool bShowPerfData = true;
-	static FAutoConsoleVariableRef CVarChaosBoxCalcBoundsISPCEnabled(TEXT("p.Dataflow.ShowPerfData"), bShowPerfData, TEXT("Whether to display perf data on top of dataflow nodes"));
-
 	void AddOutputValueToString(const FDataflowOutput& Output, FContext& Context, FString& OutString)
 	{
 		using namespace UE::Dataflow::Private;
@@ -244,9 +241,9 @@ TArray<FOverlayWidgetInfo> SDataflowEdNode::GetOverlayWidgets(bool bSelected, co
 			Widgets.Add(FreezeImageInfo);
 		}
 
-		if (UE::Dataflow::Private::bShowPerfData)
+		if (TSharedPtr<UE::Dataflow::FContext> Context = DataflowInterface->GetDataflowContext())
 		{
-			if (TSharedPtr<UE::Dataflow::FContext> Context = DataflowInterface->GetDataflowContext())
+			if (Context->IsPerfDataEnabled())
 			{
 				if (TSharedPtr<const FDataflowNode> DataflowNode = DataflowGraphNode->GetDataflowNode())
 				{
diff --git a/Engine/Plugins/Experimental/Dataflow/Source/DataflowEditor/Public/Dataflow/DataflowEditorCommands.h b/Engine/Plugins/Experimental/Dataflow/Source/DataflowEditor/Public/Dataflow/DataflowEditorCommands.h
--- a/Engine/Plugins/Experimental/Dataflow/Source/DataflowEditor/Public/Dataflow/DataflowEditorCommands.h
+++ b/Engine/Plugins/Experimental/Dataflow/Source/DataflowEditor/Public/Dataflow/DataflowEditorCommands.h
@@ -60,6 +60,7 @@ public:
 	TSharedPtr<FUICommandInfo> EvaluateGraphAutomatic;
 	TSharedPtr<FUICommandInfo> EvaluateGraphManual;
 	TSharedPtr<FUICommandInfo> ClearGraphCache;
+	TSharedPtr<FUICommandInfo> TogglePerfData;
 	TSharedPtr<FUICommandInfo> FreezeNodes;
 	TSharedPtr<FUICommandInfo> UnfreezeNodes;
 	TSharedPtr<FUICommandInfo> CreateComment;
diff --git a/Engine/Plugins/Experimental/Dataflow/Source/DataflowEditor/Public/Dataflow/DataflowEditorToolkit.h b/Engine/Plugins/Experimental/Dataflow/Source/DataflowEditor/Public/Dataflow/DataflowEditorToolkit.h
--- a/Engine/Plugins/Experimental/Dataflow/Source/DataflowEditor/Public/Dataflow/DataflowEditorToolkit.h
+++ b/Engine/Plugins/Experimental/Dataflow/Source/DataflowEditor/Public/Dataflow/DataflowEditorToolkit.h
@@ -245,6 +245,8 @@ private:
 	void SetEvaluateGraphMode(EDataflowEditorEvaluationMode Mode);
 	void ClearGraphCache();
 	bool CanClearGraphCache() const;
+	void TogglePerfData();
+	bool IsPerfDataEnabled() const;
 
 	/** Create the simulation viewport client */
 	void CreateSimulationViewportClient();
diff --git a/Engine/Source/Runtime/Experimental/Dataflow/Core/Private/Dataflow/DataflowNodeParameters.cpp b/Engine/Source/Runtime/Experimental/Dataflow/Core/Private/Dataflow/DataflowNodeParameters.cpp
--- a/Engine/Source/Runtime/Experimental/Dataflow/Core/Private/Dataflow/DataflowNodeParameters.cpp
+++ b/Engine/Source/Runtime/Experimental/Dataflow/Core/Private/Dataflow/DataflowNodeParameters.cpp
@@ -22,7 +22,7 @@ namespace UE::Dataflow
 
 	void FContextPerfData::Accumulate(const FDataflowConnection* Connection, uint64 TotalTime, uint64 ExternalTime)
 	{
-		if (Connection)
+		if (Connection && bEnabled)
 		{
 			FData& Data = DataPerOutput.FindOrAdd(Connection->GetGuid());
 			Data.ExclusiveTimeMs = FMath::Max(Data.ExclusiveTimeMs, (float)FPlatformTime::ToMilliseconds64(TotalTime - ExternalTime));
@@ -203,6 +203,23 @@ namespace UE::Dataflow
 #endif
 	}
 
+	void FContext::EnablePerfData(bool bEnable)
+	{
+#if DATAFLOW_EDITOR_EVALUATION
+		PerfData.bEnabled = bEnable;
+#endif
+	}
+
+	bool FContext::IsPerfDataEnabled() const
+	{
+#if DATAFLOW_EDITOR_EVALUATION
+		return PerfData.bEnabled;
+#else
+		return false;
+#endif
+	}
+
+
 	void FContext::PushToCallstack(const FDataflowConnection* Connection)
 	{
 #if DATAFLOW_EDITOR_EVALUATION
diff --git a/Engine/Source/Runtime/Experimental/Dataflow/Core/Public/Dataflow/DataflowNodeParameters.h b/Engine/Source/Runtime/Experimental/Dataflow/Core/Public/Dataflow/DataflowNodeParameters.h
--- a/Engine/Source/Runtime/Experimental/Dataflow/Core/Public/Dataflow/DataflowNodeParameters.h
+++ b/Engine/Source/Runtime/Experimental/Dataflow/Core/Public/Dataflow/DataflowNodeParameters.h
@@ -51,6 +51,7 @@ namespace UE::Dataflow
 			FTimestamp LastTimestamp = FTimestamp(0); // timestamp of the owner node when data was last collected
 		};
 		TMap<FGuid, FData> DataPerOutput;
+		bool bEnabled = false;
 	};
 
 	//////////////////////////////////////////////////////////////////////////////////
@@ -310,6 +311,8 @@ namespace UE::Dataflow
 
 		DATAFLOWCORE_API FContextPerfData::FData GetPerfDataForNode(const FDataflowNode& Node) const;
 		DATAFLOWCORE_API void ResetPerfDataForNode(const FDataflowNode& Node);
+		DATAFLOWCORE_API void EnablePerfData(bool bEnable);
+		DATAFLOWCORE_API bool IsPerfDataEnabled() const;
 
 		DATAFLOWCORE_API bool IsCacheEntryAfterTimestamp(FContextCacheKey InKey, const FTimestamp InTimestamp) const;
 
