Dataflow editor : Fixing over refreshing of the consruction views
- consolidae some recen code change wih the refresh view function
- remove the construction view dirying via the content for now as this was solving an async eval issue that a current one

diff --git a/Engine/Plugins/Experimental/Dataflow/Source/DataflowEditor/Private/Dataflow/DataflowEditorToolkit.cpp b/Engine/Plugins/Experimental/Dataflow/Source/DataflowEditor/Private/Dataflow/DataflowEditorToolkit.cpp
--- a/Engine/Plugins/Experimental/Dataflow/Source/DataflowEditor/Private/Dataflow/DataflowEditorToolkit.cpp
+++ b/Engine/Plugins/Experimental/Dataflow/Source/DataflowEditor/Private/Dataflow/DataflowEditorToolkit.cpp
@@ -1894,13 +1894,6 @@ void FDataflowEditorToolkit::Tick(float DeltaTime)
 			}
 		}
 	}
-	if(ConstructionScene->IsSceneDirty() || SimulationScene->IsSceneDirty())
-	{
-		bViewsNeedRefresh = true;
-		ConstructionScene->ResetDirtyFlag();
-		SimulationScene->ResetDirtyFlag();
-	}
-
 	RefreshViewsIfNeeded();
 }
 
@@ -1957,6 +1950,11 @@ void FDataflowEditorToolkit::EvaluateNode(const FDataflowNode* Node, const FData
 
 void FDataflowEditorToolkit::RefreshViewsIfNeeded(bool bForce)
 {
+	if (ConstructionScene->IsSceneDirty() || SimulationScene->IsSceneDirty())
+	{
+		bViewsNeedRefresh = true;
+	}
+
 	if (bForce || bViewsNeedRefresh)
 	{
 		for (IDataflowViewListener* Listener : ViewListeners)
@@ -1964,13 +1962,16 @@ void FDataflowEditorToolkit::RefreshViewsIfNeeded(bool bForce)
 			Listener->RefreshView();
 		}
 
-		if (const TObjectPtr<UDataflowBaseContent>& EditorContent = GetEditorContent())
-		{
-			EditorContent->SetConstructionDirty(true);
-		}
-
 		bViewsNeedRefresh = false;
 	}
+
+	// reset dirty flags
+	if (ConstructionScene->IsSceneDirty() || SimulationScene->IsSceneDirty())
+	{
+		ConstructionScene->ResetDirtyFlag();
+		SimulationScene->ResetDirtyFlag();
+	}
+
 }
 
 TSharedRef<SDataflowGraphEditor> FDataflowEditorToolkit::CreateGraphEditorWidget(UEdGraph* GraphToEdit, TSharedPtr<IStructureDetailsView> InNodeDetailsEditor)
