Dataflow : Update the collection spreadsheet header to be slimmer so it takes less space in the default layout
- also reorganized the UI code to make it easier to re-arrange the UI elements if needed

diff --git a/Engine/Plugins/Experimental/Dataflow/Source/DataflowEditor/Private/Dataflow/CollectionSpreadSheetWidget.cpp b/Engine/Plugins/Experimental/Dataflow/Source/DataflowEditor/Private/Dataflow/CollectionSpreadSheetWidget.cpp
--- a/Engine/Plugins/Experimental/Dataflow/Source/DataflowEditor/Private/Dataflow/CollectionSpreadSheetWidget.cpp
+++ b/Engine/Plugins/Experimental/Dataflow/Source/DataflowEditor/Private/Dataflow/CollectionSpreadSheetWidget.cpp
@@ -387,6 +387,151 @@ void SCollectionSpreadSheetWidget::Construct(const FArguments& InArgs)
 {
 	constexpr float CScrollBarWidth = 12.f;
 
+	// Output: [ TransformSelection        |V|]
+	TSharedRef<SWidget> OutputSelectionWidget =
+		SNew(SHorizontalBox)
+		+ SHorizontalBox::Slot()
+		.AutoWidth()
+		.VAlign(EVerticalAlignment::VAlign_Center)
+		[
+			SNew(STextBlock)
+				.Text(FText::FromString("Output: "))
+		]
+		+ SHorizontalBox::Slot()
+		.FillWidth(1.0f)
+		[
+			SAssignNew(NodeOutputsComboBox, SComboBox<FName>)
+			.ToolTipText(FText::FromString("Select a node output to see the output's data"))
+			.OptionsSource(&NodeOutputs)
+			.OnGenerateWidget(SComboBox<FName>::FOnGenerateWidget::CreateLambda([](FName Output)->TSharedRef<SWidget>
+			{
+				return SNew(STextBlock)
+					.Text(FText::FromName(Output));
+			}))
+			.OnSelectionChanged(this, &SCollectionSpreadSheetWidget::NodeOutputsComboBoxSelectionChanged)
+			[
+				SAssignNew(NodeOutputsComboBoxLabel, STextBlock)
+				.Text(GetNoOutputText())
+			]
+		];
+
+	// Group: [ TransformGroup        |V|]
+	TSharedRef<SWidget> GroupSelectionWidget =
+		SNew(SHorizontalBox)
+		+ SHorizontalBox::Slot()
+		.AutoWidth()
+		.VAlign(EVerticalAlignment::VAlign_Center)
+		[
+			SNew(STextBlock)
+				.Text(FText::FromString("Group: "))
+		]
+		+ SHorizontalBox::Slot()
+		.FillWidth(1.0f)
+		[
+			SAssignNew(CollectionGroupsComboBox, SComboBox<FName>)
+			.ToolTipText(FText::FromString("Select a group see the corresponding data"))
+			.OptionsSource(&CollectionGroups)
+			.OnGenerateWidget(SComboBox<FName>::FOnGenerateWidget::CreateLambda([](FName Group)->TSharedRef<SWidget>
+			{
+				return SNew(STextBlock)
+					.Text(FText::FromName(Group));
+			}))
+			.OnSelectionChanged(this, &SCollectionSpreadSheetWidget::CollectionGroupsComboBoxSelectionChanged)
+			[
+				SAssignNew(CollectionGroupsComboBoxLabel, STextBlock)
+					.Text(GetNoGroupText())
+			]
+		];
+
+	// Pin button widget
+	TSharedRef<SWidget> PinButtonWidget =
+		SNew(SCheckBox)
+		.Style(FAppStyle::Get(), "ToggleButtonCheckbox")
+		.ToolTipText(FText::FromString("The button pins down the panel. When it pinned down it doesn't react to node selection change."))
+		.IsChecked_Lambda([this]()-> ECheckBoxState
+		{
+			if (bIsPinnedDown)
+			{
+				return ECheckBoxState::Checked;
+			}
+			return ECheckBoxState::Unchecked;
+		})
+		.OnCheckStateChanged_Lambda([this](ECheckBoxState State)
+		{
+			if (!NodeNameTextBlock->GetText().IsEmpty())
+			{
+				bIsPinnedDown = !bIsPinnedDown;
+				OnPinnedDownChangedDelegate.Broadcast(bIsPinnedDown);
+			}
+		})
+		.Padding(2.0f)
+		.HAlign(HAlign_Center)
+		[
+			SNew(SImage)
+				.ColorAndOpacity(FSlateColor::UseForeground())
+				.Image(this, &SCollectionSpreadSheetWidget::GetPinButtonImage)
+		];
+
+	// Lock button widget
+	TSharedRef<SWidget> LockButtonWidget =
+		SNew(SCheckBox)
+		.Style(FAppStyle::Get(), "ToggleButtonCheckbox")
+		.ToolTipText(FText::FromString("The button locks the refresh of the values in the panel."))
+		.IsChecked_Lambda([this]()-> ECheckBoxState
+		{
+			if (bIsRefreshLocked)
+			{
+				return ECheckBoxState::Checked;
+			}
+			return ECheckBoxState::Unchecked;
+		})
+		.OnCheckStateChanged_Lambda([this](ECheckBoxState State)
+		{
+			if (!NodeNameTextBlock->GetText().IsEmpty())
+			{
+				bIsRefreshLocked = !bIsRefreshLocked;
+				OnRefreshLockedChangedDelegate.Broadcast(bIsRefreshLocked);
+			}
+		})
+		.Padding(2.0f)
+		.HAlign(HAlign_Center)
+		[
+			SNew(SImage)
+				.ColorAndOpacity(FSlateColor::UseForeground())
+				.Image(this, &SCollectionSpreadSheetWidget::GetLockButtonImage)
+		];
+
+	TSharedRef<SWidget> PinAndLockButtonsWidget =
+		SNew(SHorizontalBox)
+		+SHorizontalBox::Slot()
+		.AutoWidth()
+		.Padding(1.0f, 0)
+		[
+			PinButtonWidget
+		]
+		+ SHorizontalBox::Slot()
+		.AutoWidth()
+		.Padding(1.0f, 0)
+		[
+			LockButtonWidget
+		];
+
+	TSharedRef<SWidget> NodeNameWidget =
+		SNew(SHorizontalBox)
+		+ SHorizontalBox::Slot()
+		.AutoWidth()
+		.VAlign(EVerticalAlignment::VAlign_Center)
+		[
+			SNew(STextBlock)
+			.Text(FText::FromString("Node: "))
+		]
+		+ SHorizontalBox::Slot()
+		.FillWidth(1.f)
+		.VAlign(EVerticalAlignment::VAlign_Center)
+		[
+			SAssignNew(NodeNameTextBlock, STextBlock)
+		];
+
 	SAssignNew(SpreadSheetHorizontalScrollBar, SScrollBar)
 	.Orientation(Orient_Horizontal)
 	.Thickness(FVector2D(CScrollBarWidth, CScrollBarWidth));
@@ -415,160 +560,43 @@ void SCollectionSpreadSheetWidget::Construct(const FArguments& InArgs)
 	[
 		SNew(SVerticalBox)
 
+		// Selection header
 		+ SVerticalBox::Slot()
 		.AutoHeight()
+		.Padding(4.0f)
 		[
-			// Node: [ MakeBox_0                 ] [O]
 			SNew(SHorizontalBox)
-
 			+ SHorizontalBox::Slot()
 			.AutoWidth()
-			.Padding(10.0f, 10.0f, 0.0f, 5.0f)
+			.MinWidth(150.0)
+			.Padding(4.0f, 0.0f)
 			[
-				SNew(STextBlock)
-				.Text(FText::FromString("Node: "))
+				OutputSelectionWidget
 			]
-
-			+ SHorizontalBox::Slot()
-			.FillWidth(1.f)
-			.Padding(10.0f, 10.0f, 10.0f, 5.0f)
-			[
-				SAssignNew(NodeNameTextBlock, STextBlock)
-			]
-
 			+ SHorizontalBox::Slot()
 			.AutoWidth()
-			.Padding(10.0f, 10.0f, 4.0f, 5.0f)
+			.MinWidth(150.0)
+			.Padding(4.0f, 0.0f)
 			[
-				SNew(SCheckBox)
-				.Style(FAppStyle::Get(), "ToggleButtonCheckbox")
-				.ToolTipText(FText::FromString("The button pins down the panel. When it pinned down it doesn't react to node selection change."))
-				.IsChecked_Lambda([this]()-> ECheckBoxState
-				{
-					if (bIsPinnedDown)
-					{
-						return ECheckBoxState::Checked;
-					}
-					return ECheckBoxState::Unchecked;
-				})
-				.OnCheckStateChanged_Lambda([this](ECheckBoxState State)
-				{
-					if (!NodeNameTextBlock->GetText().IsEmpty())
-					{
-						bIsPinnedDown = !bIsPinnedDown;
-						OnPinnedDownChangedDelegate.Broadcast(bIsPinnedDown);
-					}
-				})
-				.Padding(2.0f)
-				.HAlign(HAlign_Center)
-				[
-					SNew(SImage)
-					.ColorAndOpacity(FSlateColor::UseForeground())
-					.Image(this, &SCollectionSpreadSheetWidget::GetPinButtonImage)
-				]
-			]
-
-			+ SHorizontalBox::Slot()
-			.AutoWidth()
-			.Padding(4.0f, 10.0f, 10.0f, 5.0f)
-			[
-				SNew(SCheckBox)
-				.Style(FAppStyle::Get(), "ToggleButtonCheckbox")
-				.ToolTipText(FText::FromString("The button locks the refresh of the values in the panel."))
-				.IsChecked_Lambda([this]()-> ECheckBoxState
-					{
-						if (bIsRefreshLocked)
-						{
-							return ECheckBoxState::Checked;
-						}
-						return ECheckBoxState::Unchecked;
-					})
-				.OnCheckStateChanged_Lambda([this](ECheckBoxState State)
-				{
-					if (!NodeNameTextBlock->GetText().IsEmpty())
-					{
-						bIsRefreshLocked = !bIsRefreshLocked;
-						OnRefreshLockedChangedDelegate.Broadcast(bIsRefreshLocked);
-					}
-				})
-				.Padding(2.0f)
-				.HAlign(HAlign_Center)
-				[
-					SNew(SImage)
-					.ColorAndOpacity(FSlateColor::UseForeground())
-					.Image(this, &SCollectionSpreadSheetWidget::GetLockButtonImage)
-				]
-			]
-		]
-
-		+ SVerticalBox::Slot()
-		.AutoHeight()
-		[
-			// Output: [ TransformSelection        |V|]
-			SNew(SHorizontalBox)
-
-			+ SHorizontalBox::Slot()
-			.AutoWidth()
-			.Padding(10.0f, 0.0f, 0.0f, 10.0f)
-			[
-				SNew(STextBlock)
-				.Text(FText::FromString("Output: "))
+				GroupSelectionWidget
 			]
-
 			+ SHorizontalBox::Slot()
 			.FillWidth(1.0f)
-			.Padding(10.0f, 0.0f, 10.0f, 10.0f)
+			.Padding(4.0f, 0.0f)
+			.HAlign(EHorizontalAlignment::HAlign_Center)
 			[
-				SAssignNew(NodeOutputsComboBox, SComboBox<FName>)
-				.ToolTipText(FText::FromString("Select a node output to see the output's data"))
-				.OptionsSource(&NodeOutputs)
-				.OnGenerateWidget(SComboBox<FName>::FOnGenerateWidget::CreateLambda([](FName Output)->TSharedRef<SWidget>
-				{
-					return SNew(STextBlock)
-						.Text(FText::FromName(Output));
-				}))
-				.OnSelectionChanged(this, &SCollectionSpreadSheetWidget::NodeOutputsComboBoxSelectionChanged)
-				[
-					SAssignNew(NodeOutputsComboBoxLabel, STextBlock)
-					.Text(GetNoOutputText())
-				]
+				NodeNameWidget
 			]
-		]
-
-		+ SVerticalBox::Slot()
-		.AutoHeight()
-		[
-			// Group: [ TransformGroup        |V|]
-			SNew(SHorizontalBox)
-
 			+ SHorizontalBox::Slot()
 			.AutoWidth()
-			.Padding(10.0f, 0.0f, 0.0f, 2.0f)
+			.Padding(4.0f, 0.0f)
+			.HAlign(EHorizontalAlignment::HAlign_Right)
 			[
-				SNew(STextBlock)
-				.Text(FText::FromString("Group: "))
-			]
-
-			+ SHorizontalBox::Slot()
-			.FillWidth(1.0f)
-			.Padding(10.0f, 0.0f, 10.0f, 2.0f)
-			[
-				SAssignNew(CollectionGroupsComboBox, SComboBox<FName>)
-				.ToolTipText(FText::FromString("Select a group see the corresponding data"))
-				.OptionsSource(&CollectionGroups)
-				.OnGenerateWidget(SComboBox<FName>::FOnGenerateWidget::CreateLambda([](FName Group)->TSharedRef<SWidget>
-				{
-					return SNew(STextBlock)
-						.Text(FText::FromName(Group));
-				}))
-				.OnSelectionChanged(this, &SCollectionSpreadSheetWidget::CollectionGroupsComboBoxSelectionChanged)
-				[
-					SAssignNew(CollectionGroupsComboBoxLabel, STextBlock)
-						.Text(GetNoGroupText())
-				]
+				PinAndLockButtonsWidget
 			]
 		]
 
+		// Data section
 		+SVerticalBox::Slot()
 		[
 			SNew(SHorizontalBox)
