Dataflow:
- Added VertexID, TriangleID, Vertex Normal, Triangle Normal display

diff --git a/Engine/Plugins/ChaosClothAssetEditor/Source/ChaosClothAssetEditor/Private/ChaosClothAsset/ClothDataflowConstructionVisualization.cpp b/Engine/Plugins/ChaosClothAssetEditor/Source/ChaosClothAssetEditor/Private/ChaosClothAsset/ClothDataflowConstructionVisualization.cpp
--- a/Engine/Plugins/ChaosClothAssetEditor/Source/ChaosClothAssetEditor/Private/ChaosClothAsset/ClothDataflowConstructionVisualization.cpp
+++ b/Engine/Plugins/ChaosClothAssetEditor/Source/ChaosClothAssetEditor/Private/ChaosClothAsset/ClothDataflowConstructionVisualization.cpp
@@ -264,7 +264,7 @@ namespace UE::Chaos::ClothAsset
 		MenuBuilder.EndSection();
 	}
 
-	void FClothDataflowConstructionVisualization::Draw(const FDataflowConstructionScene* ConstructionScene, FPrimitiveDrawInterface* PDI)
+	void FClothDataflowConstructionVisualization::Draw(const FDataflowConstructionScene* ConstructionScene, FPrimitiveDrawInterface* PDI, const FSceneView* View)
 	{
 		if (!ConstructionScene)
 		{
diff --git a/Engine/Plugins/ChaosClothAssetEditor/Source/ChaosClothAssetEditor/Private/ChaosClothAsset/ClothDataflowConstructionVisualization.h b/Engine/Plugins/ChaosClothAssetEditor/Source/ChaosClothAssetEditor/Private/ChaosClothAsset/ClothDataflowConstructionVisualization.h
--- a/Engine/Plugins/ChaosClothAssetEditor/Source/ChaosClothAssetEditor/Private/ChaosClothAsset/ClothDataflowConstructionVisualization.h
+++ b/Engine/Plugins/ChaosClothAssetEditor/Source/ChaosClothAssetEditor/Private/ChaosClothAsset/ClothDataflowConstructionVisualization.h
@@ -17,7 +17,7 @@ namespace UE::Chaos::ClothAsset
 		// IDataflowConstructionVisualization
 		virtual FName GetName() const override;
 		virtual void ExtendViewportShowMenu(const TSharedPtr<FDataflowConstructionViewportClient>& ViewportClient, FMenuBuilder& MenuBuilder) override;
-		virtual void Draw(const FDataflowConstructionScene* ConstructionScene, FPrimitiveDrawInterface* PDI) override;
+		virtual void Draw(const FDataflowConstructionScene* ConstructionScene, FPrimitiveDrawInterface* PDI, const FSceneView* View = nullptr) override;
 
 		bool bSeamVisualizationEnabled = false;
 		bool bCollapseSeams = false;
diff --git a/Engine/Plugins/Experimental/Dataflow/Source/DataflowEditor/DataflowEditor.Build.cs b/Engine/Plugins/Experimental/Dataflow/Source/DataflowEditor/DataflowEditor.Build.cs
--- a/Engine/Plugins/Experimental/Dataflow/Source/DataflowEditor/DataflowEditor.Build.cs
+++ b/Engine/Plugins/Experimental/Dataflow/Source/DataflowEditor/DataflowEditor.Build.cs
@@ -85,7 +85,8 @@ namespace UnrealBuildTool.Rules
 					"AnimGraph", 
 					"ChaosCaching", // UAnimSingleNodeInstance
 					"StructUtilsEditor",
-					"MessageLog"
+					"MessageLog",
+					"AppFramework"
 				}
 			);
 		}
diff --git a/Engine/Plugins/Experimental/Dataflow/Source/DataflowEditor/Private/Dataflow/DataflowConstructionScene.cpp b/Engine/Plugins/Experimental/Dataflow/Source/DataflowEditor/Private/Dataflow/DataflowConstructionScene.cpp
--- a/Engine/Plugins/Experimental/Dataflow/Source/DataflowEditor/Private/Dataflow/DataflowConstructionScene.cpp
+++ b/Engine/Plugins/Experimental/Dataflow/Source/DataflowEditor/Private/Dataflow/DataflowConstructionScene.cpp
@@ -130,12 +130,34 @@ void FDataflowConstructionScene::TickDataflowScene(const float DeltaSeconds)
 	}
 }
 
+void FDataflowConstructionScene::FDebugMesh::Reset()
+{
+	VertexMap.Reset();
+	FaceMap.Reset();
+}
+
+void FDataflowConstructionScene::FDebugMesh::Build(const TArray<TObjectPtr<UDynamicMeshComponent>>& InDynamicMeshComponents)
+{
+	ResultMesh.Clear();
+	ResultMesh.EnableAttributes();
+
+	for (const UDynamicMeshComponent* const DynamicMeshComponent : InDynamicMeshComponents)
+	{
+		if (const UE::Geometry::FDynamicMesh3* const Mesh = DynamicMeshComponent->GetMesh())
+		{
+			ResultMesh.AppendWithOffsets(*Mesh);
+		}
+	}
+
+	Spatial = UE::Geometry::FDynamicMeshAABBTree3(&ResultMesh, true);
+}
+
 void FDataflowConstructionScene::UpdateDynamicMeshComponents()
 {
 	using namespace UE::Geometry;//FDynamicMesh3
 
 	// The preview scene for the construction view will be
-	// cleared and rebuilt from scratch. This will genrate a 
+	// cleared and rebuilt from scratch. This will generate a 
 	// list of UPrimitiveComponents for rendering.
 	ResetDynamicMeshComponents();
 
@@ -158,6 +180,16 @@ void FDataflowConstructionScene::UpdateDynamicMeshComponents()
 					UE::Dataflow::RenderNodeOutput(Facade, *Target, *EditorContent);
 
 					const int32 NumGeometry = Facade.NumGeometry();
+
+					DebugMesh.Reset();
+
+					int32 VertexIndex = 0, FaceIndex = 0;
+
+					const TManagedArray<int32>& VertexStart = Facade.GetVertexStart();
+					const TManagedArray<int32>& VertexCount = Facade.GetVertexCount();
+					const TManagedArray<int32>& FaceStart = Facade.GetIndicesStart();
+					const TManagedArray<int32>& FaceCount = Facade.GetIndicesCount();
+
 					for (int32 MeshIndex = 0; MeshIndex < NumGeometry; ++MeshIndex)
 					{
 						FDynamicMesh3 DynamicMesh;
@@ -184,12 +216,23 @@ void FDataflowConstructionScene::UpdateDynamicMeshComponents()
 							}
 
 							AddDynamicMeshComponent({Target, MeshIndex}, MeshName, MoveTemp(DynamicMesh), Materials);
+
+							for (int32 Idx = 0; Idx < VertexCount[MeshIndex]; ++Idx)
+							{
+								DebugMesh.VertexMap.Add(VertexIndex, VertexStart[MeshIndex] + Idx);
+								VertexIndex++;
+							}
+
+							for (int32 Idx = 0; Idx < FaceCount[MeshIndex]; ++Idx)
+							{
+								DebugMesh.FaceMap.Add(FaceIndex, FaceStart[MeshIndex] + Idx);
+								FaceIndex++;
+							}
 						}
 					}
 				}
 			}
 
-
 			// If we have a single mesh component in the scene, select it
 			if (DynamicMeshComponents.Num() == 1 && DataflowModeManager)
 			{
@@ -221,6 +264,7 @@ void FDataflowConstructionScene::UpdateDynamicMeshComponents()
 					UE::Dataflow::RenderNodeOutput(Facade, *Target, *EditorContent);
 
 					const int32 NumGeometry = Facade.NumGeometry();
+
 					for (int32 MeshIndex = 0; MeshIndex < NumGeometry; ++MeshIndex)
 					{
 						FDataflowRenderKey WireframeDynamicMeshKey{ Target, MeshIndex };
@@ -273,7 +317,11 @@ void FDataflowConstructionScene::UpdateDynamicMeshComponents()
 				}
 			}
 		}
+	
+		// Build a single mesh out of all the components
+		DebugMesh.Build(GetDynamicMeshComponents());
 	}
+
 	bPreviewSceneDirty = true;
 }
 
@@ -302,7 +350,7 @@ void FDataflowConstructionScene::ResetDynamicMeshComponents()
 
 TObjectPtr<UDynamicMeshComponent>& FDataflowConstructionScene::AddDynamicMeshComponent(FDataflowRenderKey InKey, const FString& MeshName, UE::Geometry::FDynamicMesh3&& DynamicMesh, const TArray<UMaterialInterface*>& MaterialSet)
 {
-	// Dont use the MakeUniqueObjectName for the component, we need to keep the name aligned with the collection so selection will work in 
+	// Don't use the MakeUniqueObjectName for the component, we need to keep the name aligned with the collection so selection will work in 
 	// other editors. 
 	// const FName UniqueObjectName = MakeUniqueObjectName(RootSceneActor, UDataflowEditorCollectionComponent::StaticClass(), FName(MeshName));
 	TObjectPtr<UDataflowEditorCollectionComponent> DynamicMeshComponent = NewObject<UDataflowEditorCollectionComponent>(RootSceneActor, FName(MeshName));
diff --git a/Engine/Plugins/Experimental/Dataflow/Source/DataflowEditor/Private/Dataflow/DataflowConstructionViewportClient.cpp b/Engine/Plugins/Experimental/Dataflow/Source/DataflowEditor/Private/Dataflow/DataflowConstructionViewportClient.cpp
--- a/Engine/Plugins/Experimental/Dataflow/Source/DataflowEditor/Private/Dataflow/DataflowConstructionViewportClient.cpp
+++ b/Engine/Plugins/Experimental/Dataflow/Source/DataflowEditor/Private/Dataflow/DataflowConstructionViewportClient.cpp
@@ -337,7 +337,7 @@ void FDataflowConstructionViewportClient::Draw(const FSceneView* View, FPrimitiv
 	using namespace UE::Dataflow;
 	for (const TPair<FName, TUniquePtr<IDataflowConstructionVisualization>>& Visualization : FDataflowConstructionVisualizationRegistry::GetInstance().GetVisualizations())
 	{
-		Visualization.Value->Draw(static_cast<FDataflowConstructionScene*>(PreviewScene), PDI);
+		Visualization.Value->Draw(static_cast<FDataflowConstructionScene*>(PreviewScene), PDI, View);
 	}
 }
 
diff --git a/Engine/Plugins/Experimental/Dataflow/Source/DataflowEditor/Private/Dataflow/MeshConstructionVisualization.cpp b/Engine/Plugins/Experimental/Dataflow/Source/DataflowEditor/Private/Dataflow/MeshConstructionVisualization.cpp
--- a/Engine/Plugins/Experimental/Dataflow/Source/DataflowEditor/Private/Dataflow/MeshConstructionVisualization.cpp
+++ b/Engine/Plugins/Experimental/Dataflow/Source/DataflowEditor/Private/Dataflow/MeshConstructionVisualization.cpp
@@ -5,21 +5,285 @@
 #include "Dataflow/DataflowConstructionViewportClient.h"
 #include "Dataflow/DataflowConstructionScene.h"
 #include "Framework/MultiBox/MultiBoxBuilder.h"
+#include "Engine/Engine.h"  // for GEngine->GetSmallFont()
+#include "CanvasItem.h"
+#include "CanvasTypes.h"
+#include "SceneView.h"
+#include "BoxTypes.h"
+#include "DynamicMesh/DynamicMeshAABBTree3.h"
+#include "DynamicMeshEditor.h"
+#include "Widgets/Input/SNumericEntryBox.h"
+#include "Widgets/Colors/SColorBlock.h"
+#include "Widgets/Colors/SColorPicker.h"
+#include "Widgets/Layout/SSpacer.h"
 
 #define LOCTEXT_NAMESPACE "MeshConstructionVisualization"
 
 namespace UE::Dataflow
 {
+	const FMargin LabelWidgetsMargin(15.0f, 0.0f, 3.0f, 0.0f);
+	const FMargin WidgetsMargin(2.0f, 0.0f, 3.0f, 0.0f);
+	const FVector2D SpacerSize = FVector2D(1, 4);
+
+	TSharedRef<SWidget> CreateNumericEntryWidget(
+		const TSharedRef<SWidget>& InNumericBoxWidget,
+		const FText& InLabel)
+	{
+		return SNew(SHorizontalBox)
+			+ SHorizontalBox::Slot()
+			.VAlign(VAlign_Center)
+			.HAlign(HAlign_Left)
+			.Padding(LabelWidgetsMargin)
+			.AutoWidth()
+			[
+				SNew(STextBlock)
+					.Text(InLabel)
+			]
+			+ SHorizontalBox::Slot()
+			.VAlign(VAlign_Center)
+			.HAlign(HAlign_Right)
+			.Padding(FMargin(6.0f, 0))
+			.FillContentWidth(1.0)
+			[
+				SNew(SBox)
+					.Padding(WidgetsMargin)
+					.MinDesiredWidth(80.0f)
+					[
+						InNumericBoxWidget
+					]
+			];
+	}
+
+	TSharedRef<SWidget> CreateColorEntryWidget(
+		const TSharedRef<SWidget>& InColorWidget,
+		const FText& InLabel)
+	{
+		return SNew(SHorizontalBox)
+			+ SHorizontalBox::Slot()
+			.VAlign(VAlign_Center)
+			.HAlign(HAlign_Left)
+			.Padding(LabelWidgetsMargin)
+			.AutoWidth()
+			[
+				SNew(STextBlock)
+					.Text(InLabel)
+			]
+			+ SHorizontalBox::Slot()
+			.VAlign(VAlign_Center)
+			.HAlign(HAlign_Right)
+			.Padding(FMargin(6.0f, 0))
+			.FillContentWidth(1.0)
+			[
+				SNew(SBox)
+					.Padding(WidgetsMargin)
+					.MinDesiredWidth(80.0f)
+					[
+						InColorWidget
+					]
+			];
+	}
+
 	FName FMeshConstructionVisualization::Name = FName("MeshConstructionVisualization");
 
 	void FMeshConstructionVisualization::ExtendViewportShowMenu(const TSharedPtr<FDataflowConstructionViewportClient>& ViewportClient, FMenuBuilder& MenuBuilder)
 	{
 		MenuBuilder.BeginSection(TEXT("MeshVisualization"), LOCTEXT("MeshVisualizationSectionName", "Mesh"));
 		{
-			const FUIAction MeshNormalsToggleAction(
+			TSharedRef<SWidget> VertexIDColorEntryWidget =
+				SNew(SColorBlock)
+				.ToolTipText(LOCTEXT("", ""))
+				.Color_Lambda([this]()
+					{
+						return GetVertexIDColor();
+					})
+				.OnMouseButtonDown_Lambda([this](const FGeometry&, const FPointerEvent&)
+					{
+						FColorPickerArgs PickerArgs;
+						PickerArgs.bUseAlpha = false;
+						FLinearColor Color = GetVertexNormalColor();
+						PickerArgs.InitialColor = Color;
+						PickerArgs.OnColorCommitted = FOnLinearColorValueChanged::CreateLambda([this](FLinearColor Color)
+							{
+								SetVertexIDColor(Color);
+							});
+						OpenColorPicker(PickerArgs);
+						return FReply::Handled();
+					});
+
+			TSharedRef<SWidget> FaceIDColorEntryWidget =
+				SNew(SColorBlock)
+				.ToolTipText(LOCTEXT("", ""))
+				.Color_Lambda([this]()
+					{
+						return GetFaceIDColor();
+					})
+				.OnMouseButtonDown_Lambda([this](const FGeometry&, const FPointerEvent&)
+					{
+						FColorPickerArgs PickerArgs;
+						PickerArgs.bUseAlpha = false;
+						FLinearColor Color = GetVertexNormalColor();
+						PickerArgs.InitialColor = Color;
+						PickerArgs.OnColorCommitted = FOnLinearColorValueChanged::CreateLambda([this](FLinearColor Color)
+							{
+								SetFaceIDColor(Color);
+							});
+						OpenColorPicker(PickerArgs);
+						return FReply::Handled();
+					});
+
+			TSharedRef<SWidget> VertexNormalLengthNumericEntryWidget =
+				SNew(SNumericEntryBox<float>)
+				.ToolTipText(LOCTEXT("", ""))
+				.MinValue(1.f)
+				.MaxValue(10.f)
+				.MaxSliderValue(10.f)
+				.AllowSpin(true)
+				.MaxFractionalDigits(1)
+				.Font(FAppStyle::GetFontStyle(TEXT("MenuItem.Font")))
+				.OnValueChanged_Lambda([this](float InValue)
+					{
+						SetVertexNormalLength(InValue);
+					})
+				.Value_Lambda([this]()
+					{
+						return GetVertexNormalLength();
+					});
+
+			TSharedRef<SWidget> VertexNormalThicknessNumericEntryWidget =
+				SNew(SNumericEntryBox<float>)
+				.ToolTipText(LOCTEXT("", ""))
+				.MinValue(0.2f)
+				.MaxValue(1.f)
+				.MaxSliderValue(1.f)
+				.AllowSpin(true)
+				.MaxFractionalDigits(2)
+				.Font(FAppStyle::GetFontStyle(TEXT("MenuItem.Font")))
+				.OnValueChanged_Lambda([this](float InValue)
+					{
+						SetVertexNormalThickness(InValue);
+					})
+				.Value_Lambda([this]()
+					{
+						return GetVertexNormalThickness();
+					});
+
+			TSharedRef<SWidget> VertexNormalColorEntryWidget =
+				SNew(SColorBlock)
+				.ToolTipText(LOCTEXT("", ""))
+				.Color_Lambda([this]()
+					{
+						return GetVertexNormalColor();
+					})
+				.OnMouseButtonDown_Lambda([this](const FGeometry&, const FPointerEvent&)
+					{
+						FColorPickerArgs PickerArgs;
+						PickerArgs.bUseAlpha = false;
+						FLinearColor Color = GetVertexNormalColor();
+						PickerArgs.InitialColor = Color;
+						PickerArgs.OnColorCommitted = FOnLinearColorValueChanged::CreateLambda([this](FLinearColor Color)
+							{
+								SetVertexNormalColor(Color);
+							});
+						OpenColorPicker(PickerArgs);
+						return FReply::Handled();
+					});
+
+			TSharedRef<SWidget> FaceNormalLengthNumericEntryWidget =
+				SNew(SNumericEntryBox<float>)
+				.ToolTipText(LOCTEXT("", ""))
+				.MinValue(1.f)
+				.MaxValue(10.f)
+				.MaxSliderValue(10.f)
+				.AllowSpin(true)
+				.MaxFractionalDigits(1)
+				.Font(FAppStyle::GetFontStyle(TEXT("MenuItem.Font")))
+				.OnValueChanged_Lambda([this](float InValue)
+					{
+						SetFaceNormalLength(InValue);
+					})
+				.Value_Lambda([this]()
+					{
+						return GetFaceNormalLength();
+					});
+
+			TSharedRef<SWidget> FaceNormalThicknessNumericEntryWidget =
+				SNew(SNumericEntryBox<float>)
+				.ToolTipText(LOCTEXT("", ""))
+				.MinValue(0.2f)
+				.MaxValue(1.f)
+				.MaxSliderValue(1.f)
+				.AllowSpin(true)
+				.MaxFractionalDigits(2)
+				.Font(FAppStyle::GetFontStyle(TEXT("MenuItem.Font")))
+				.OnValueChanged_Lambda([this](float InValue)
+					{
+						SetFaceNormalThickness(InValue);
+					})
+				.Value_Lambda([this]()
+					{
+						return GetFaceNormalThickness();
+					});
+
+			TSharedRef<SWidget> FaceNormalColorEntryWidget =
+				SNew(SColorBlock)
+				.ToolTipText(LOCTEXT("", ""))
+				.Color_Lambda([this]()
+					{
+						return GetFaceNormalColor();
+					})
+				.OnMouseButtonDown_Lambda([this](const FGeometry&, const FPointerEvent&)
+					{
+						FColorPickerArgs PickerArgs;
+						PickerArgs.bUseAlpha = false;
+						FLinearColor Color = GetFaceNormalColor();
+						PickerArgs.InitialColor = Color;
+						PickerArgs.OnColorCommitted = FOnLinearColorValueChanged::CreateLambda([this](FLinearColor Color)
+							{
+								SetFaceNormalColor(Color);
+							});
+						OpenColorPicker(PickerArgs);
+						return FReply::Handled();
+					});
+
+			TSharedRef<SWidget> DistanceCutoffNumericEntryWidget =
+				SNew(SNumericEntryBox<float>)
+				.ToolTipText(LOCTEXT("", ""))
+				.MinValue(1.f)
+				.MaxValue(1000.f)
+				.MaxSliderValue(1000.f)
+				.AllowSpin(true)
+				.MaxFractionalDigits(1)
+				.Font(FAppStyle::GetFontStyle(TEXT("MenuItem.Font")))
+				.OnValueChanged_Lambda([this](float InValue)
+					{
+						SetDistanceCutoff(InValue);
+					})
+				.Value_Lambda([this]()
+					{
+						return GetDistanceCutoff();
+					});
+
+			///////////////////////////////////////////////////////////////////////////////////////////////////////////////
+
+			const FUIAction MeshVertexNumbersToggleAction(
+				FExecuteAction::CreateLambda([this, ViewportClient]()
+					{
+						bMeshVertexIDVisualizationEnabled = !bMeshVertexIDVisualizationEnabled;
+						if (ViewportClient)
+						{
+							ViewportClient->Invalidate();
+						}
+					}),
+				FCanExecuteAction(),
+				FIsActionChecked::CreateLambda([this]()
+					{
+						return bMeshVertexIDVisualizationEnabled;
+					}));
+
+			const FUIAction MeshFaceNumbersToggleAction(
 				FExecuteAction::CreateLambda([this, ViewportClient]()
 					{
-						bMeshNormalsVisualizationEnabled = !bMeshNormalsVisualizationEnabled;
+						bMeshFaceIDVisualizationEnabled = !bMeshFaceIDVisualizationEnabled;
 						if (ViewportClient)
 						{
 							ViewportClient->Invalidate();
@@ -28,54 +292,390 @@ namespace UE::Dataflow
 				FCanExecuteAction(),
 				FIsActionChecked::CreateLambda([this]()
 					{
-						return bMeshNormalsVisualizationEnabled;
+						return bMeshFaceIDVisualizationEnabled;
 					}));
 
-			MenuBuilder.AddMenuEntry(LOCTEXT("MeshVisualization_NormalsEnabled", "Normals"),
-				LOCTEXT("MeshVisualization_NormalsEnabled_TooltipText", "Draw mesh vertex normals"),
+			const FUIAction MeshVertexNormalsToggleAction(
+				FExecuteAction::CreateLambda([this, ViewportClient]()
+					{
+						bMeshVertexNormalsVisualizationEnabled = !bMeshVertexNormalsVisualizationEnabled;
+						if (ViewportClient)
+						{
+							ViewportClient->Invalidate();
+						}
+					}),
+				FCanExecuteAction(),
+				FIsActionChecked::CreateLambda([this]()
+					{
+						return bMeshVertexNormalsVisualizationEnabled;
+					}));
+
+			const FUIAction MeshFaceNormalsToggleAction(
+				FExecuteAction::CreateLambda([this, ViewportClient]()
+					{
+						bMeshFaceNormalsVisualizationEnabled = !bMeshFaceNormalsVisualizationEnabled;
+						if (ViewportClient)
+						{
+							ViewportClient->Invalidate();
+						}
+					}),
+				FCanExecuteAction(),
+				FIsActionChecked::CreateLambda([this]()
+					{
+						return bMeshFaceNormalsVisualizationEnabled;
+					}));
+
+			const FUIAction IgnoreOccludedTrianglesToggleAction(
+				FExecuteAction::CreateLambda([this, ViewportClient]()
+					{
+						bIgnoreOccludedTriangles = !bIgnoreOccludedTriangles;
+						if (ViewportClient)
+						{
+							ViewportClient->Invalidate();
+						}
+					}),
+				FCanExecuteAction(),
+				FIsActionChecked::CreateLambda([this]()
+					{
+						return bIgnoreOccludedTriangles;
+					}));
+
+			///////////////////////////////////////////////////////////////////////////////////////////////////////////////
+
+			MenuBuilder.AddMenuEntry(LOCTEXT("MeshVisualization_VertexNumbersEnabled", "Vertex ID"),
+				LOCTEXT("MeshVisualization_VertexNumbersEnabled_TooltipText", "Display vertex ID"),
+				FSlateIcon(),
+				MeshVertexNumbersToggleAction,
+				NAME_None,
+				EUserInterfaceActionType::ToggleButton);
+
+			MenuBuilder.AddWidget(
+				CreateColorEntryWidget(VertexIDColorEntryWidget, LOCTEXT("MeshVisualization_VertexIDColor", "Color")),
+				LOCTEXT("", ""), true);
+
+			MenuBuilder.AddWidget(SNew(SSpacer).Size(SpacerSize), FText::GetEmpty(), false, false);
+
+			MenuBuilder.AddMenuEntry(LOCTEXT("MeshVisualization_FaceNumbersEnabled", "Face ID"),
+				LOCTEXT("MeshVisualization_FaceNumbersEnabled_TooltipText", "Display face ID"),
 				FSlateIcon(),
-				MeshNormalsToggleAction,
+				MeshFaceNumbersToggleAction,
+				NAME_None,
+				EUserInterfaceActionType::ToggleButton);
+
+			MenuBuilder.AddWidget(
+				CreateColorEntryWidget(FaceIDColorEntryWidget, LOCTEXT("MeshVisualization_FaceIDColor", "Color")),
+				LOCTEXT("", ""), true);
+
+			MenuBuilder.AddWidget(SNew(SSpacer).Size(SpacerSize), FText::GetEmpty(), false, false);
+
+			MenuBuilder.AddMenuEntry(LOCTEXT("MeshVisualization_VertexNormalsEnabled", "Vertex Normals"),
+				LOCTEXT("MeshVisualization_VertexNormalsEnabled_TooltipText", "Display vertex normals"),
+				FSlateIcon(),
+				MeshVertexNormalsToggleAction,
+				NAME_None,
+				EUserInterfaceActionType::ToggleButton);
+
+			MenuBuilder.AddWidget(
+				CreateNumericEntryWidget(VertexNormalLengthNumericEntryWidget, LOCTEXT("MeshVisualization_VertexNormalLength", "Length")),
+				LOCTEXT("", ""), true);
+
+			MenuBuilder.AddWidget(
+				CreateNumericEntryWidget(VertexNormalThicknessNumericEntryWidget, LOCTEXT("MeshVisualization_VertexNormalThickness", "Thickness")),
+				LOCTEXT("", ""), true);
+
+			MenuBuilder.AddWidget(
+				CreateColorEntryWidget(VertexNormalColorEntryWidget, LOCTEXT("MeshVisualization_VertexNormalColor", "Color")),
+				LOCTEXT("", ""), true);
+
+			MenuBuilder.AddWidget(SNew(SSpacer).Size(SpacerSize), FText::GetEmpty(), false, false);
+
+			MenuBuilder.AddMenuEntry(LOCTEXT("MeshVisualization_FaceNormalsEnabled", "Face Normals"),
+				LOCTEXT("MeshVisualization_FaceNormalsEnabled_TooltipText", "Display face normals"),
+				FSlateIcon(),
+				MeshFaceNormalsToggleAction,
+				NAME_None,
+				EUserInterfaceActionType::ToggleButton);
+
+			MenuBuilder.AddWidget(
+				CreateNumericEntryWidget(FaceNormalLengthNumericEntryWidget, LOCTEXT("MeshVisualization_FaceNormalLength", "Length")),
+				LOCTEXT("", ""), true);
+
+			MenuBuilder.AddWidget(
+				CreateNumericEntryWidget(FaceNormalThicknessNumericEntryWidget, LOCTEXT("MeshVisualization_FaceNormalThickness", "Thickness")),
+				LOCTEXT("", ""), true);
+
+			MenuBuilder.AddWidget(
+				CreateColorEntryWidget(FaceNormalColorEntryWidget, LOCTEXT("MeshVisualization_FaceNormalColor", "Color")),
+				LOCTEXT("", ""), true);
+
+			MenuBuilder.AddWidget(SNew(SSpacer).Size(SpacerSize), FText::GetEmpty(), false, false);
+			MenuBuilder.AddWidget(SNew(SSpacer).Size(SpacerSize), FText::GetEmpty(), false, false);
+
+			MenuBuilder.AddWidget(
+				CreateNumericEntryWidget(DistanceCutoffNumericEntryWidget, LOCTEXT("MeshVisualization_DistanceCutoff", "Distance Cutoff")),
+				LOCTEXT("", ""), true);
+
+			MenuBuilder.AddMenuEntry(LOCTEXT("MeshVisualization_IgnoreOccludedTriangles", "Ignore Occluded Triangles"),
+				LOCTEXT("MeshVisualization_IgnoreOccludedTriangles_TooltipText", "Ignore Occluded Triangles"),
+				FSlateIcon(),
+				IgnoreOccludedTrianglesToggleAction,
 				NAME_None,
 				EUserInterfaceActionType::ToggleButton);
 		}
 		MenuBuilder.EndSection();
 	}
 
-	void FMeshConstructionVisualization::Draw(const FDataflowConstructionScene* ConstructionScene, FPrimitiveDrawInterface* PDI)
+	static void DrawText(FCanvas* Canvas, const FSceneView* SceneView, const FVector& Pos, const FText& Text, const FLinearColor& Color, const float Scale = 1.f)
+	{
+#if WITH_EDITOR
+		if (Canvas && SceneView)
+		{
+			FVector2D PixelLocation;
+			if (SceneView->WorldToPixel(Pos, PixelLocation))
+			{
+				// WorldToPixel doesn't account for DPIScale
+				const float DPIScale = Canvas->GetDPIScale();
+				FCanvasTextItem TextItem(PixelLocation / DPIScale, Text, GEngine->GetSmallFont(), Color);
+				TextItem.Scale = FVector2D::UnitVector * Scale;
+				TextItem.EnableShadow(FLinearColor::Black);
+				TextItem.Draw(Canvas);
+			}
+		}
+#endif
+	}
+
+	static bool IsTriangleVisible(const FDataflowConstructionScene* ConstructionScene, const FSceneView* SceneView, const int32 TriangleID, float& DistanceMin, float& DistanceMax, TMap<int32, float>& TriangleDistances, const bool bIgnoreOccludedTriangles)
+	{
+		const UE::Geometry::FDynamicMesh3& ResultMesh = ConstructionScene->DebugMesh.ResultMesh;
+		const UE::Geometry::FDynamicMeshAABBTree3& Spatial = ConstructionScene->DebugMesh.Spatial;
+
+		const FVector TriangleCentroid = ResultMesh.GetTriCentroid(TriangleID);
+		FVector LocalEyeDirection = TriangleCentroid - SceneView->ViewLocation;
+
+		FRay3d LocalEyeRay;
+		LocalEyeRay.Origin = SceneView->ViewLocation;
+		FVector LocalPosition = ResultMesh.GetTriCentroid(TriangleID);
+		LocalEyeRay.Direction = UE::Geometry::Normalized(LocalPosition - SceneView->ViewLocation);
+		if (LocalEyeRay.Direction.Dot(LocalEyeDirection) < 0)
+		{
+			return false;
+		}
+
+		double LocalRayHitT;
+		int HitTriangleID;
+		FVector3d HitBaryCoords;
+		if (bIgnoreOccludedTriangles)
+		{
+			if (Spatial.FindNearestHitTriangle(LocalEyeRay, LocalRayHitT, HitTriangleID, HitBaryCoords))
+			{
+				if (HitTriangleID != TriangleID)
+				{
+					return false;
+				}
+				else
+				{
+					if (LocalRayHitT < DistanceMin)
+					{
+						DistanceMin = LocalRayHitT;
+					}
+					else if (LocalRayHitT > DistanceMax)
+					{
+						DistanceMax = LocalRayHitT;
+					}
+
+					TriangleDistances.Add(TriangleID, LocalRayHitT);
+				}
+			}
+		}
+
+		return true;
+	}
+
+	static void ComputeVisibleTriangles(const FDataflowConstructionScene* ConstructionScene, const FSceneView* SceneView, const UE::Geometry::FDynamicMesh3& ResultMesh, TArray<int32>& VisibleTriangles, float& DistanceMin, float& DistanceMax, TMap<int32, float>& TriangleDistances, const bool bIgnoreOccludedTriangles)
+	{
+		for (int32 TriangleID : ResultMesh.TriangleIndicesItr())
+		{
+			if (IsTriangleVisible(ConstructionScene, SceneView, TriangleID, DistanceMin, DistanceMax, TriangleDistances, bIgnoreOccludedTriangles))
+			{
+				VisibleTriangles.Add(TriangleID);
+			}
+		}
+	}
+
+	void FMeshConstructionVisualization::DrawCanvas(const FDataflowConstructionScene* ConstructionScene, FCanvas* Canvas, const FSceneView* SceneView)
+	{
+		constexpr float CEndColorProgress = 0.8f;
+
+		if (ConstructionScene && Canvas && SceneView)
+		{
+			const UE::Geometry::FDynamicMesh3& ResultMesh = ConstructionScene->DebugMesh.ResultMesh;
+
+			if (ResultMesh.TriangleCount() > 0)
+			{
+				TMap<int32, float> TriangleDistances;
+				float DistanceMin = FLT_MAX, DistanceMax = -FLT_MAX;
+
+				TArray<int32> VisibleTriangles;
+				ComputeVisibleTriangles(ConstructionScene, SceneView, ResultMesh, VisibleTriangles, DistanceMin, DistanceMax, TriangleDistances, bIgnoreOccludedTriangles);
+
+				if (bMeshVertexIDVisualizationEnabled && VisibleTriangles.Num() > 0)
+				{
+					TSet<int32> VerticesSet;
+					TMap<int32, int32> VertexToTriangle;
+					for (int32 TriangleID : VisibleTriangles)
+					{
+						if (bIgnoreOccludedTriangles && (TriangleDistances[TriangleID] > GetDistanceCutoff()))
+						{
+							continue;
+						}
+
+						const UE::Geometry::FIndex3i TriangleIndices = ResultMesh.GetTriangle(TriangleID);
+
+						auto SetVertex = [&VerticesSet, &VertexToTriangle](const int32 VertexID, const int32 TriangleID) {
+							if (!VerticesSet.Contains(VertexID))
+							{
+								VerticesSet.Add(VertexID);
+								VertexToTriangle.Add(VertexID, TriangleID);
+							}
+						};
+
+						SetVertex(TriangleIndices.A, TriangleID);
+						SetVertex(TriangleIndices.B, TriangleID);
+						SetVertex(TriangleIndices.C, TriangleID);
+					}
+
+					for (int32 VertexID : VerticesSet)
+					{
+						const FVector VertexPos = ResultMesh.GetVertex(VertexID) + FVector(0.15, 0, 0.15);
+						const FText Text = FText::AsNumber(ConstructionScene->DebugMesh.VertexMap[VertexID]);
+
+						if (bIgnoreOccludedTriangles)
+						{
+							FLinearColor EndColor = FLinearColor::LerpUsingHSV(GetVertexIDColor(), FLinearColor::Black, CEndColorProgress);
+							float Progress = (TriangleDistances[VertexToTriangle[VertexID]] - DistanceMin) / (DistanceMax - DistanceMin);
+							FLinearColor Color = FLinearColor::LerpUsingHSV(GetVertexIDColor(), EndColor, Progress);
+
+							DrawText(Canvas, SceneView, VertexPos, Text, Color);
+						}
+						else
+						{
+							DrawText(Canvas, SceneView, VertexPos, Text, GetVertexIDColor());
+						}
+					}
+				}
+
+				if (bMeshFaceIDVisualizationEnabled && VisibleTriangles.Num() > 0)
+				{
+					for (int32 TriangleID : VisibleTriangles)
+					{
+						const FVector TriangleCentroid = ResultMesh.GetTriCentroid(TriangleID);
+						const FText Text = FText::AsNumber(ConstructionScene->DebugMesh.FaceMap[TriangleID]);
+
+						if (bIgnoreOccludedTriangles)
+						{
+							if (TriangleDistances[TriangleID] > GetDistanceCutoff())
+							{
+								continue;
+							}
+
+							FLinearColor EndColor = FLinearColor::LerpUsingHSV(GetFaceIDColor(), FLinearColor::Black, CEndColorProgress);
+							float Progress = (TriangleDistances[TriangleID] - DistanceMin) / (DistanceMax - DistanceMin);
+							FLinearColor Color = FLinearColor::LerpUsingHSV(GetFaceIDColor(), EndColor, Progress);
+
+							DrawText(Canvas, SceneView, TriangleCentroid, Text, Color);
+						}
+						else
+						{
+							DrawText(Canvas, SceneView, TriangleCentroid, Text, GetFaceIDColor());
+						}
+					}
+				}
+			}
+		}
+	}
+
+	void FMeshConstructionVisualization::Draw(const FDataflowConstructionScene* ConstructionScene, FPrimitiveDrawInterface* PDI, const FSceneView* SceneView)
 	{
-		if (bMeshNormalsVisualizationEnabled && ConstructionScene && PDI)
+		if (ConstructionScene && PDI && SceneView)
 		{
-			for (const UDynamicMeshComponent* const DynamicMeshComponent : ConstructionScene->GetDynamicMeshComponents())
+			const UE::Geometry::FDynamicMesh3& ResultMesh = ConstructionScene->DebugMesh.ResultMesh;
+
+			if (ResultMesh.TriangleCount() > 0)
 			{
-				if (const UE::Geometry::FDynamicMesh3* const Mesh = DynamicMeshComponent->GetMesh())
+				TArray<int32> VisibleTriangles;
+				TMap<int32, float> TriangleDistances;
+				float DistanceMin = FLT_MAX, DistanceMax = -FLT_MAX;
+				ComputeVisibleTriangles(ConstructionScene, SceneView, ResultMesh, VisibleTriangles, DistanceMin, DistanceMax, TriangleDistances, bIgnoreOccludedTriangles);
+
+				if (bMeshVertexNormalsVisualizationEnabled && VisibleTriangles.Num() > 0)
 				{
-					if (Mesh->HasAttributes())
+					if (ResultMesh.HasAttributes())
 					{
-						if (const UE::Geometry::FDynamicMeshNormalOverlay* const NormalOverlay = Mesh->Attributes()->PrimaryNormals())
+						if (const UE::Geometry::FDynamicMeshNormalOverlay* const NormalOverlay = ResultMesh.Attributes()->PrimaryNormals())
 						{
-							PDI->AddReserveLines(SDPG_World, NormalOverlay->ElementCount());
+							TSet<int32> VerticesSet;
+							TArray<FVector> Normals;
+							TArray<FVector> VertexPositions;
 
-							for (int32 ElementIndex = 0; ElementIndex < NormalOverlay->MaxElementID(); ++ElementIndex)
+							for (int32 TriangleID : VisibleTriangles)
 							{
-								if (NormalOverlay->IsElement(ElementIndex))
+								if (bIgnoreOccludedTriangles && (TriangleDistances[TriangleID] > GetDistanceCutoff()))
 								{
-									constexpr float LineLength = 2.0f;
-									constexpr float LineThickness = 0.01f;
-									constexpr FColor LineColor(15, 15, 240);
-
-									const int32 ParentVtx = NormalOverlay->GetParentVertex(ElementIndex);
-									const FVector Normal = (FVector)NormalOverlay->GetElement(ElementIndex);
-									const FVector Origin = (FVector)Mesh->GetVertex(ParentVtx);
-									PDI->DrawLine(Origin, Origin + LineLength * Normal, LineColor , SDPG_World, LineThickness);
+									continue;
 								}
+
+								const UE::Geometry::FIndex3i TriangleIndices = ResultMesh.GetTriangle(TriangleID);
+
+								auto SetVertex = [&VerticesSet, &Normals, &VertexPositions, &NormalOverlay , &ResultMesh](const int32 VertexID, const int32 TriangleID) {
+									if (!VerticesSet.Contains(VertexID))
+									{
+										VerticesSet.Add(VertexID);
+
+										const FVector Normal = (FVector)NormalOverlay->GetElementAtVertex(TriangleID, VertexID);
+										Normals.Add(Normal);
+
+										const FVector VertexPosition = (FVector)ResultMesh.GetVertex(VertexID);
+										VertexPositions.Add(VertexPosition);
+									}
+								};
+
+								SetVertex(TriangleIndices.A, TriangleID);
+								SetVertex(TriangleIndices.B, TriangleID);
+								SetVertex(TriangleIndices.C, TriangleID);
+							}
+
+							PDI->AddReserveLines(SDPG_World, Normals.Num());
+
+							for (int32 Idx = 0; Idx < Normals.Num(); ++Idx)
+							{
+								PDI->DrawLine(VertexPositions[Idx], VertexPositions[Idx] + GetVertexNormalLength() * Normals[Idx], GetVertexNormalColor().ToFColor(true), SDPG_World, GetVertexNormalThickness());
 							}
 						}
 					}
 				}
+
+				if (bMeshFaceNormalsVisualizationEnabled && VisibleTriangles.Num() > 0)
+				{
+					const int32 NumTriangles = ResultMesh.TriangleCount();
+
+					PDI->AddReserveLines(SDPG_World, NumTriangles);
+
+					for (int32 TriangleID : VisibleTriangles)
+					{
+						if (bIgnoreOccludedTriangles && (TriangleDistances[TriangleID] > GetDistanceCutoff()))
+						{
+							continue;
+						}
+
+						const FVector TriangleCentroid = ResultMesh.GetTriCentroid(TriangleID);
+						const FVector TriangleNormal = ResultMesh.GetTriNormal(TriangleID);
+
+						PDI->DrawLine(TriangleCentroid, TriangleCentroid + GetFaceNormalLength() * TriangleNormal, GetFaceNormalColor().ToFColor(true), SDPG_World, GetFaceNormalThickness());
+					}
+				}
 			}
 		}
 	}
-
 }
 
 #undef LOCTEXT_NAMESPACE 
diff --git a/Engine/Plugins/Experimental/Dataflow/Source/DataflowEditor/Private/Dataflow/MeshConstructionVisualization.h b/Engine/Plugins/Experimental/Dataflow/Source/DataflowEditor/Private/Dataflow/MeshConstructionVisualization.h
--- a/Engine/Plugins/Experimental/Dataflow/Source/DataflowEditor/Private/Dataflow/MeshConstructionVisualization.h
+++ b/Engine/Plugins/Experimental/Dataflow/Source/DataflowEditor/Private/Dataflow/MeshConstructionVisualization.h
@@ -2,6 +2,7 @@
 #pragma once
 
 #include "Dataflow/DataflowConstructionVisualization.h"
+#include "Math/Color.h" 
 
 namespace UE::Dataflow
 {
@@ -11,15 +12,58 @@ namespace UE::Dataflow
 	public:
 		static FName Name;
 
+		FLinearColor GetVertexIDColor() const { return VertexIDColor; }
+		void SetVertexIDColor(const FLinearColor Value) { VertexIDColor = Value; }
+
+		FLinearColor GetFaceIDColor() const { return FaceIDColor; }
+		void SetFaceIDColor(const FLinearColor Value) { FaceIDColor = Value; }
+
+		float GetVertexNormalLength() const { return VertexNormalLength; }
+		void SetVertexNormalLength(const float Value) { VertexNormalLength = Value; }
+		float GetVertexNormalThickness() const { return VertexNormalThickness; }
+		void SetVertexNormalThickness(const float Value) { VertexNormalThickness = Value; }
+		FLinearColor GetVertexNormalColor() const { return VertexNormalColor; }
+		void SetVertexNormalColor(const FLinearColor Value) { VertexNormalColor = Value; }
+
+		float GetFaceNormalLength() const { return FaceNormalLength; }
+		void SetFaceNormalLength(const float Value) { FaceNormalLength = Value; }
+		float GetFaceNormalThickness() const { return FaceNormalThickness; }
+		void SetFaceNormalThickness(const float Value) { FaceNormalThickness = Value; }
+		FLinearColor GetFaceNormalColor() const { return FaceNormalColor; }
+		void SetFaceNormalColor(const FLinearColor Value) { FaceNormalColor = Value; }
+
+		float GetDistanceCutoff() const { return DistanceCutoff; }
+		void SetDistanceCutoff(const float Value) { DistanceCutoff = Value; }
+
 	private:
 		virtual FName GetName() const override
 		{
 			return Name;
 		}
 		virtual void ExtendViewportShowMenu(const TSharedPtr<FDataflowConstructionViewportClient>& ViewportClient, FMenuBuilder& MenuBuilder) override;
-		virtual void Draw(const FDataflowConstructionScene* ConstructionScene, FPrimitiveDrawInterface* PDI) override;
+		virtual void Draw(const FDataflowConstructionScene* ConstructionScene, FPrimitiveDrawInterface* PDI, const FSceneView* SceneView = nullptr) override;
+		virtual void DrawCanvas(const FDataflowConstructionScene* ConstructionScene, FCanvas* Canvas, const FSceneView* SceneView) override;
+
+		bool bMeshVertexIDVisualizationEnabled = false;
+		bool bMeshFaceIDVisualizationEnabled = false;
+		bool bMeshVertexNormalsVisualizationEnabled = false;
+		bool bMeshFaceNormalsVisualizationEnabled = false;
+
+		FLinearColor VertexIDColor = FLinearColor::Green;
+		
+		FLinearColor FaceIDColor = FLinearColor::Yellow;
+
+		float VertexNormalLength = 3.f;
+		float VertexNormalThickness = 0.3f;
+		FLinearColor VertexNormalColor = FLinearColor::Blue;
+
+		float FaceNormalLength = 3.f;
+		float FaceNormalThickness = 0.3f;
+		FLinearColor FaceNormalColor = FLinearColor::Red;
+
+		float DistanceCutoff = 500.f;
 
-		bool bMeshNormalsVisualizationEnabled = false;
+		bool bIgnoreOccludedTriangles = true;
 	};
 }
 
diff --git a/Engine/Plugins/Experimental/Dataflow/Source/DataflowEditor/Public/Dataflow/DataflowConstructionScene.h b/Engine/Plugins/Experimental/Dataflow/Source/DataflowEditor/Public/Dataflow/DataflowConstructionScene.h
--- a/Engine/Plugins/Experimental/Dataflow/Source/DataflowEditor/Public/Dataflow/DataflowConstructionScene.h
+++ b/Engine/Plugins/Experimental/Dataflow/Source/DataflowEditor/Public/Dataflow/DataflowConstructionScene.h
@@ -5,6 +5,7 @@
 #include "Dataflow/DataflowContent.h"
 #include "Dataflow/DataflowEditorPreviewSceneBase.h"
 #include "Dataflow/DataflowNodeParameters.h"
+#include "DynamicMesh/DynamicMeshAABBTree3.h"
 
 class UDataflowEditor;
 class USkeletalMeshComponent;
@@ -29,6 +30,17 @@ class DATAFLOWEDITOR_API FDataflowConstructionScene : public FDataflowPreviewSce
 	/*(const void* ObjectPointer, const int32 ObjectIndex)*/
 	typedef TPair<TObjectPtr<const UDataflowEdNode>, int32> FDataflowRenderKey;
 
+	struct FDebugMesh
+	{
+		UE::Geometry::FDynamicMesh3 ResultMesh;
+		UE::Geometry::FDynamicMeshAABBTree3 Spatial;
+		TMap<int32, int32> VertexMap;
+		TMap<int32, int32> FaceMap;
+
+		void Reset();
+		void Build(const TArray<TObjectPtr<UDynamicMeshComponent>>& InDynamicMeshComponents);
+	};
+
 public:
 
 	FDataflowConstructionScene(FPreviewScene::ConstructionValues ConstructionValues, UDataflowEditor* Editor);
@@ -58,6 +70,8 @@ public:
 	/** Array of all DynamicMeshComponents in the scene */
 	TArray<TObjectPtr<UDynamicMeshComponent>> GetDynamicMeshComponents() const;
 
+	FDebugMesh DebugMesh;
+
 private:
 
 	/** Reset the scene mesh visualizer */
diff --git a/Engine/Plugins/Experimental/Dataflow/Source/DataflowEditor/Public/Dataflow/DataflowConstructionVisualization.h b/Engine/Plugins/Experimental/Dataflow/Source/DataflowEditor/Public/Dataflow/DataflowConstructionVisualization.h
--- a/Engine/Plugins/Experimental/Dataflow/Source/DataflowEditor/Public/Dataflow/DataflowConstructionVisualization.h
+++ b/Engine/Plugins/Experimental/Dataflow/Source/DataflowEditor/Public/Dataflow/DataflowConstructionVisualization.h
@@ -22,7 +22,7 @@ namespace UE::Dataflow
 
 		DATAFLOWEDITOR_API virtual FName GetName() const = 0;
 		DATAFLOWEDITOR_API virtual void ExtendViewportShowMenu(const TSharedPtr<FDataflowConstructionViewportClient>& ViewportClient, FMenuBuilder& MenuBuilder) {};
-		DATAFLOWEDITOR_API virtual void Draw(const FDataflowConstructionScene* ConstructionScene, FPrimitiveDrawInterface* PDI) {};
+		DATAFLOWEDITOR_API virtual void Draw(const FDataflowConstructionScene* ConstructionScene, FPrimitiveDrawInterface* PDI, const FSceneView* View = nullptr) {};
 		DATAFLOWEDITOR_API virtual void DrawCanvas(const FDataflowConstructionScene* ConstructionScene, FCanvas* Canvas, const FSceneView* SceneView) {};
 	};
 
