Dataflow Editor : Fix simulation window auto focus to ignore the sprites as they have very bounds somehow and causing the camera to zoom out quite far when no object is selected
Also make sure the construction view re-ajust its camera when a node is selected

diff --git a/Engine/Plugins/Experimental/Dataflow/Source/DataflowEditor/Private/Dataflow/DataflowEditorMode.cpp b/Engine/Plugins/Experimental/Dataflow/Source/DataflowEditor/Private/Dataflow/DataflowEditorMode.cpp
--- a/Engine/Plugins/Experimental/Dataflow/Source/DataflowEditor/Private/Dataflow/DataflowEditorMode.cpp
+++ b/Engine/Plugins/Experimental/Dataflow/Source/DataflowEditor/Private/Dataflow/DataflowEditorMode.cpp
@@ -598,6 +598,7 @@ void UDataflowEditorMode::ModeTick(float DeltaTime)
 		if (EditorContent->IsConstructionDirty())
 		{
 			ConstructionScene->UpdateConstructionScene();
+			RefocusConstructionViewportClient();
 		}
 	}
 
diff --git a/Engine/Plugins/Experimental/Dataflow/Source/DataflowEditor/Private/Dataflow/DataflowSimulationScene.cpp b/Engine/Plugins/Experimental/Dataflow/Source/DataflowEditor/Private/Dataflow/DataflowSimulationScene.cpp
--- a/Engine/Plugins/Experimental/Dataflow/Source/DataflowEditor/Private/Dataflow/DataflowSimulationScene.cpp
+++ b/Engine/Plugins/Experimental/Dataflow/Source/DataflowEditor/Private/Dataflow/DataflowSimulationScene.cpp
@@ -19,6 +19,7 @@
 #include "IContentBrowserSingleton.h"
 #include "ContentBrowserModule.h"
 #include "AssetRegistry/AssetRegistryModule.h"
+#include "Components/BillboardComponent.h"
 #include "Components/SkeletalMeshComponent.h"
 #include "Engine/SkeletalMesh.h"
 #include "GeometryCache.h"
@@ -346,14 +347,32 @@ FBox FDataflowSimulationScene::GetBoundingBox() const
 			{
 				if (const UPrimitiveComponent* const SelectedComponent = Cast<const UPrimitiveComponent>(SelectedObject))
 				{
-					SceneBounds += SelectedComponent->Bounds.GetBox();
+					// Ignore billboard ones as they have a very large bounds (-256, 256)
+					const UBillboardComponent* BillboardComponent = Cast<UBillboardComponent>(SelectedComponent);
+					if (BillboardComponent == nullptr)
+					{
+						SceneBounds += SelectedComponent->Bounds.GetBox();
+					}
 				}
 			}
 		}
-		else if(PreviewActor)
+
+		// if no selection of results in an invalid box 
+		const bool bInvalidBounds = SceneBounds.GetExtent().IsNearlyZero(UE_SMALL_NUMBER);
+		if (bInvalidBounds && PreviewActor)
 		{
-			constexpr bool bIncludeNonColliding = true;
-			SceneBounds += PreviewActor->GetComponentsBoundingBox(bIncludeNonColliding);
+			if (PreviewActor)
+			{
+				PreviewActor->ForEachComponent<UPrimitiveComponent>(true, [this, &SceneBounds](UPrimitiveComponent* PrimComponent)
+					{
+						// Ignore billboard ones as they have a very large bounds (-256, 256)
+						const UBillboardComponent* BillboardComponent = Cast<UBillboardComponent>(PrimComponent);
+						if (BillboardComponent == nullptr)
+						{
+							SceneBounds += PrimComponent->Bounds.GetBox();
+						}
+					});
+			}
 		}
 	}
 	return SceneBounds;
