Dataflow : fix empty LOCTEXT and performance in Mesh construction  visualization
- make sure the visible face check is done only if the options are enabled

diff --git a/Engine/Plugins/Experimental/Dataflow/Source/DataflowEditor/Private/Dataflow/MeshConstructionVisualization.cpp b/Engine/Plugins/Experimental/Dataflow/Source/DataflowEditor/Private/Dataflow/MeshConstructionVisualization.cpp
--- a/Engine/Plugins/Experimental/Dataflow/Source/DataflowEditor/Private/Dataflow/MeshConstructionVisualization.cpp
+++ b/Engine/Plugins/Experimental/Dataflow/Source/DataflowEditor/Private/Dataflow/MeshConstructionVisualization.cpp
@@ -91,7 +91,7 @@ namespace UE::Dataflow
 		{
 			TSharedRef<SWidget> VertexIDColorEntryWidget =
 				SNew(SColorBlock)
-				.ToolTipText(LOCTEXT("", ""))
+				.ToolTipText(LOCTEXT("VertexId", "Show Vertex identifier"))
 				.Color_Lambda([this]()
 					{
 						return GetVertexIDColor();
@@ -112,7 +112,7 @@ namespace UE::Dataflow
 
 			TSharedRef<SWidget> FaceIDColorEntryWidget =
 				SNew(SColorBlock)
-				.ToolTipText(LOCTEXT("", ""))
+				.ToolTipText(LOCTEXT("FaceId", "Show Face identifier"))
 				.Color_Lambda([this]()
 					{
 						return GetFaceIDColor();
@@ -133,7 +133,7 @@ namespace UE::Dataflow
 
 			TSharedRef<SWidget> VertexNormalLengthNumericEntryWidget =
 				SNew(SNumericEntryBox<float>)
-				.ToolTipText(LOCTEXT("", ""))
+				.ToolTipText(LOCTEXT("VertexNormalLength", "Adjust the length of the per-Vertex normals"))
 				.MinValue(1.f)
 				.MaxValue(10.f)
 				.MaxSliderValue(10.f)
@@ -151,7 +151,7 @@ namespace UE::Dataflow
 
 			TSharedRef<SWidget> VertexNormalThicknessNumericEntryWidget =
 				SNew(SNumericEntryBox<float>)
-				.ToolTipText(LOCTEXT("", ""))
+				.ToolTipText(LOCTEXT("VertexNormalThickness", "Adjust the thickness of the per-Vertex normals"))
 				.MinValue(0.2f)
 				.MaxValue(1.f)
 				.MaxSliderValue(1.f)
@@ -169,7 +169,7 @@ namespace UE::Dataflow
 
 			TSharedRef<SWidget> VertexNormalColorEntryWidget =
 				SNew(SColorBlock)
-				.ToolTipText(LOCTEXT("", ""))
+				.ToolTipText(LOCTEXT("VertexNormalColor", "Adjust the color of the per-Vertex normals"))
 				.Color_Lambda([this]()
 					{
 						return GetVertexNormalColor();
@@ -190,7 +190,7 @@ namespace UE::Dataflow
 
 			TSharedRef<SWidget> FaceNormalLengthNumericEntryWidget =
 				SNew(SNumericEntryBox<float>)
-				.ToolTipText(LOCTEXT("", ""))
+				.ToolTipText(LOCTEXT("FaceNormalLength", "Adjust the length of the per-Face normals"))
 				.MinValue(1.f)
 				.MaxValue(10.f)
 				.MaxSliderValue(10.f)
@@ -208,7 +208,7 @@ namespace UE::Dataflow
 
 			TSharedRef<SWidget> FaceNormalThicknessNumericEntryWidget =
 				SNew(SNumericEntryBox<float>)
-				.ToolTipText(LOCTEXT("", ""))
+				.ToolTipText(LOCTEXT("FaceNormalLength", "Adjust the thickness of the per-Face normals"))
 				.MinValue(0.2f)
 				.MaxValue(1.f)
 				.MaxSliderValue(1.f)
@@ -226,7 +226,7 @@ namespace UE::Dataflow
 
 			TSharedRef<SWidget> FaceNormalColorEntryWidget =
 				SNew(SColorBlock)
-				.ToolTipText(LOCTEXT("", ""))
+				.ToolTipText(LOCTEXT("FaceNormalColor", "Adjust the Color of the per-Face normals"))
 				.Color_Lambda([this]()
 					{
 						return GetFaceNormalColor();
@@ -247,7 +247,7 @@ namespace UE::Dataflow
 
 			TSharedRef<SWidget> DistanceCutoffNumericEntryWidget =
 				SNew(SNumericEntryBox<float>)
-				.ToolTipText(LOCTEXT("", ""))
+				.ToolTipText(LOCTEXT("DistanceCutoff", "Adjust the distance at which the visualization elements are not display anymore"))
 				.MinValue(1.f)
 				.MaxValue(1000.f)
 				.MaxSliderValue(1000.f)
@@ -351,7 +351,7 @@ namespace UE::Dataflow
 
 			MenuBuilder.AddWidget(
 				CreateColorEntryWidget(VertexIDColorEntryWidget, LOCTEXT("MeshVisualization_VertexIDColor", "Color")),
-				LOCTEXT("", ""), true);
+				FText::GetEmpty(), true);
 
 			MenuBuilder.AddWidget(SNew(SSpacer).Size(SpacerSize), FText::GetEmpty(), false, false);
 
@@ -364,7 +364,7 @@ namespace UE::Dataflow
 
 			MenuBuilder.AddWidget(
 				CreateColorEntryWidget(FaceIDColorEntryWidget, LOCTEXT("MeshVisualization_FaceIDColor", "Color")),
-				LOCTEXT("", ""), true);
+				FText::GetEmpty(), true);
 
 			MenuBuilder.AddWidget(SNew(SSpacer).Size(SpacerSize), FText::GetEmpty(), false, false);
 
@@ -377,15 +377,15 @@ namespace UE::Dataflow
 
 			MenuBuilder.AddWidget(
 				CreateNumericEntryWidget(VertexNormalLengthNumericEntryWidget, LOCTEXT("MeshVisualization_VertexNormalLength", "Length")),
-				LOCTEXT("", ""), true);
+				FText::GetEmpty(), true);
 
 			MenuBuilder.AddWidget(
 				CreateNumericEntryWidget(VertexNormalThicknessNumericEntryWidget, LOCTEXT("MeshVisualization_VertexNormalThickness", "Thickness")),
-				LOCTEXT("", ""), true);
+				FText::GetEmpty(), true);
 
 			MenuBuilder.AddWidget(
 				CreateColorEntryWidget(VertexNormalColorEntryWidget, LOCTEXT("MeshVisualization_VertexNormalColor", "Color")),
-				LOCTEXT("", ""), true);
+				FText::GetEmpty(), true);
 
 			MenuBuilder.AddWidget(SNew(SSpacer).Size(SpacerSize), FText::GetEmpty(), false, false);
 
@@ -398,22 +398,22 @@ namespace UE::Dataflow
 
 			MenuBuilder.AddWidget(
 				CreateNumericEntryWidget(FaceNormalLengthNumericEntryWidget, LOCTEXT("MeshVisualization_FaceNormalLength", "Length")),
-				LOCTEXT("", ""), true);
+				FText::GetEmpty(), true);
 
 			MenuBuilder.AddWidget(
 				CreateNumericEntryWidget(FaceNormalThicknessNumericEntryWidget, LOCTEXT("MeshVisualization_FaceNormalThickness", "Thickness")),
-				LOCTEXT("", ""), true);
+				FText::GetEmpty(), true);
 
 			MenuBuilder.AddWidget(
 				CreateColorEntryWidget(FaceNormalColorEntryWidget, LOCTEXT("MeshVisualization_FaceNormalColor", "Color")),
-				LOCTEXT("", ""), true);
+				FText::GetEmpty(), true);
 
 			MenuBuilder.AddWidget(SNew(SSpacer).Size(SpacerSize), FText::GetEmpty(), false, false);
 			MenuBuilder.AddWidget(SNew(SSpacer).Size(SpacerSize), FText::GetEmpty(), false, false);
 
 			MenuBuilder.AddWidget(
 				CreateNumericEntryWidget(DistanceCutoffNumericEntryWidget, LOCTEXT("MeshVisualization_DistanceCutoff", "Distance Cutoff")),
-				LOCTEXT("", ""), true);
+				FText::GetEmpty(), true);
 
 			MenuBuilder.AddMenuEntry(LOCTEXT("MeshVisualization_IgnoreOccludedTriangles", "Ignore Occluded Triangles"),
 				LOCTEXT("MeshVisualization_IgnoreOccludedTriangles_TooltipText", "Ignore Occluded Triangles"),
@@ -506,7 +506,7 @@ namespace UE::Dataflow
 	{
 		constexpr float CEndColorProgress = 0.8f;
 
-		if (ConstructionScene && Canvas && SceneView)
+		if (ConstructionScene && Canvas && SceneView && (bMeshVertexIDVisualizationEnabled || bMeshFaceIDVisualizationEnabled))
 		{
 			const UE::Geometry::FDynamicMesh3& ResultMesh = ConstructionScene->DebugMesh.ResultMesh;
 
@@ -596,7 +596,7 @@ namespace UE::Dataflow
 
 	void FMeshConstructionVisualization::Draw(const FDataflowConstructionScene* ConstructionScene, FPrimitiveDrawInterface* PDI, const FSceneView* SceneView)
 	{
-		if (ConstructionScene && PDI && SceneView)
+		if (ConstructionScene && PDI && SceneView && (bMeshVertexNormalsVisualizationEnabled || bMeshFaceNormalsVisualizationEnabled))
 		{
 			const UE::Geometry::FDynamicMesh3& ResultMesh = ConstructionScene->DebugMesh.ResultMesh;
 
