Dataflow:
- Disabled hiding occluded triangles feature for now in displaying vertex/triangle ID/normals. It is really not performant on hires geometry and freezes the editor.
Once it will be optimized and it will be enabled.

diff --git a/Engine/Plugins/Experimental/Dataflow/Source/DataflowEditor/Private/Dataflow/DataflowConstructionScene.cpp b/Engine/Plugins/Experimental/Dataflow/Source/DataflowEditor/Private/Dataflow/DataflowConstructionScene.cpp
--- a/Engine/Plugins/Experimental/Dataflow/Source/DataflowEditor/Private/Dataflow/DataflowConstructionScene.cpp
+++ b/Engine/Plugins/Experimental/Dataflow/Source/DataflowEditor/Private/Dataflow/DataflowConstructionScene.cpp
@@ -149,7 +149,11 @@ void FDataflowConstructionScene::FDebugMesh::Build(const TArray<TObjectPtr<UDyna
 		}
 	}
 
+	// Disable it for now
+	// TODO: Optimize this
+#if 0
 	Spatial = UE::Geometry::FDynamicMeshAABBTree3(&ResultMesh, true);
+#endif
 }
 
 void FDataflowConstructionScene::UpdateDynamicMeshComponents()
diff --git a/Engine/Plugins/Experimental/Dataflow/Source/DataflowEditor/Private/Dataflow/MeshConstructionVisualization.cpp b/Engine/Plugins/Experimental/Dataflow/Source/DataflowEditor/Private/Dataflow/MeshConstructionVisualization.cpp
--- a/Engine/Plugins/Experimental/Dataflow/Source/DataflowEditor/Private/Dataflow/MeshConstructionVisualization.cpp
+++ b/Engine/Plugins/Experimental/Dataflow/Source/DataflowEditor/Private/Dataflow/MeshConstructionVisualization.cpp
@@ -408,6 +408,9 @@ namespace UE::Dataflow
 				CreateColorEntryWidget(FaceNormalColorEntryWidget, LOCTEXT("MeshVisualization_FaceNormalColor", "Color")),
 				FText::GetEmpty(), true);
 
+			// Disabled for now
+			// TODO: Optimize this functionality
+#if 0
 			MenuBuilder.AddWidget(SNew(SSpacer).Size(SpacerSize), FText::GetEmpty(), false, false);
 			MenuBuilder.AddWidget(SNew(SSpacer).Size(SpacerSize), FText::GetEmpty(), false, false);
 
@@ -421,6 +424,7 @@ namespace UE::Dataflow
 				IgnoreOccludedTrianglesToggleAction,
 				NAME_None,
 				EUserInterfaceActionType::ToggleButton);
+#endif
 		}
 		MenuBuilder.EndSection();
 	}
diff --git a/Engine/Plugins/Experimental/Dataflow/Source/DataflowEditor/Private/Dataflow/MeshConstructionVisualization.h b/Engine/Plugins/Experimental/Dataflow/Source/DataflowEditor/Private/Dataflow/MeshConstructionVisualization.h
--- a/Engine/Plugins/Experimental/Dataflow/Source/DataflowEditor/Private/Dataflow/MeshConstructionVisualization.h
+++ b/Engine/Plugins/Experimental/Dataflow/Source/DataflowEditor/Private/Dataflow/MeshConstructionVisualization.h
@@ -63,7 +63,7 @@ namespace UE::Dataflow
 
 		float DistanceCutoff = 500.f;
 
-		bool bIgnoreOccludedTriangles = true;
+		bool bIgnoreOccludedTriangles = false;
 	};
 }
 
