Dataflow:
- Fixed issue with SetError() causing nodes returning default value instead of eveluating the outputs
- Bug fixes

diff --git a/Engine/Plugins/Experimental/Dataflow/Source/DataflowEditor/Private/Dataflow/DataflowCollectionSpreadSheetHelpers.cpp b/Engine/Plugins/Experimental/Dataflow/Source/DataflowEditor/Private/Dataflow/DataflowCollectionSpreadSheetHelpers.cpp
--- a/Engine/Plugins/Experimental/Dataflow/Source/DataflowEditor/Private/Dataflow/DataflowCollectionSpreadSheetHelpers.cpp
+++ b/Engine/Plugins/Experimental/Dataflow/Source/DataflowEditor/Private/Dataflow/DataflowCollectionSpreadSheetHelpers.cpp
@@ -331,7 +331,7 @@ namespace UE::Dataflow::CollectionSpreadSheetHelpers
 		}
 		else if (InAttr == "SimulationType")
 		{
-			const TArray<FString> SimTypes = { "None", "Rigid", "Clusterd" };	
+			const TArray<FString> SimTypes = { "None", "Rigid", "Clustered" };	
 			FString AttrValueStr = SimTypes[InCollection->GetAttribute<int32>("SimulationType", InGroup)[InItemIndex]];
 
 			return SNew(SHorizontalBox)
diff --git a/Engine/Source/Runtime/Experimental/Dataflow/Core/Private/Dataflow/DataflowNodeParameters.cpp b/Engine/Source/Runtime/Experimental/Dataflow/Core/Private/Dataflow/DataflowNodeParameters.cpp
--- a/Engine/Source/Runtime/Experimental/Dataflow/Core/Private/Dataflow/DataflowNodeParameters.cpp
+++ b/Engine/Source/Runtime/Experimental/Dataflow/Core/Private/Dataflow/DataflowNodeParameters.cpp
@@ -498,7 +498,9 @@ namespace UE::Dataflow
 			OnNodeFinishEvaluateMulticast.Broadcast(Connection.GetOwningNode(), &Connection);
 		}
 
-		return !GetNumErrors() && RetValue;
+		const bool bHasError = NodeHasError(Connection.GetOwningNode()) || NodeFailed(Connection.GetOwningNode());
+
+		return !bHasError && RetValue;
 	}
 
 	void FContextCache::Serialize(FArchive& Ar)
