Per asset dataflow tools registry

diff --git a/Engine/Plugins/ChaosClothAssetEditor/Source/ChaosClothAssetEditor/Private/ChaosClothAsset/AssetDefinition_ClothAsset.cpp b/Engine/Plugins/ChaosClothAssetEditor/Source/ChaosClothAssetEditor/Private/ChaosClothAsset/AssetDefinition_ClothAsset.cpp
--- a/Engine/Plugins/ChaosClothAssetEditor/Source/ChaosClothAssetEditor/Private/ChaosClothAsset/AssetDefinition_ClothAsset.cpp
+++ b/Engine/Plugins/ChaosClothAssetEditor/Source/ChaosClothAssetEditor/Private/ChaosClothAsset/AssetDefinition_ClothAsset.cpp
@@ -200,6 +200,7 @@ EAssetCommandResult UAssetDefinition_ClothAsset::OpenAssets(const FAssetOpenArgs
 				UDataflowEditor* const AssetEditor = NewObject<UDataflowEditor>(AssetEditorSubsystem, NAME_None, RF_Transient);
 				const TSubclassOf<AActor> ActorClass = StaticLoadClass(AActor::StaticClass(), nullptr,
 					TEXT("/ChaosClothAssetEditor/BP_ClothPreview.BP_ClothPreview_C"), nullptr, LOAD_None, nullptr);
+				AssetEditor->RegisterToolCategories({"General", "Cloth"});
 				AssetEditor->Initialize({ ClothObjects[0] }, ActorClass);
 			}
 		}
diff --git a/Engine/Plugins/ChaosClothAssetEditor/Source/ChaosClothAssetEditor/Private/ChaosClothAsset/ChaosClothAssetEditorModule.cpp b/Engine/Plugins/ChaosClothAssetEditor/Source/ChaosClothAssetEditor/Private/ChaosClothAsset/ChaosClothAssetEditorModule.cpp
--- a/Engine/Plugins/ChaosClothAssetEditor/Source/ChaosClothAssetEditor/Private/ChaosClothAsset/ChaosClothAssetEditorModule.cpp
+++ b/Engine/Plugins/ChaosClothAssetEditor/Source/ChaosClothAssetEditor/Private/ChaosClothAsset/ChaosClothAssetEditorModule.cpp
@@ -93,6 +93,7 @@ namespace UE::Chaos::ClothAsset
 							AssetsToEdit.Append(Context->LoadSelectedObjects<UObject>());
 
 							UDataflowEditor* const AssetEditor = NewObject<UDataflowEditor>(DataflowEditorSubsystem, NAME_None, RF_Transient);
+							AssetEditor->RegisterToolCategories({"General", "Cloth"});
 
 							// Macke sure the cloth asset has a Dataflow asset
 							UChaosClothAsset* const ClothAsset = CastChecked<UChaosClothAsset>(AssetsToEdit[0]);
diff --git a/Engine/Plugins/ChaosClothAssetEditor/Source/ChaosClothAssetEditorTools/Private/ChaosClothAsset/ChaosClothAssetEditorToolsModule.cpp b/Engine/Plugins/ChaosClothAssetEditor/Source/ChaosClothAssetEditorTools/Private/ChaosClothAsset/ChaosClothAssetEditorToolsModule.cpp
--- a/Engine/Plugins/ChaosClothAssetEditor/Source/ChaosClothAssetEditorTools/Private/ChaosClothAsset/ChaosClothAssetEditorToolsModule.cpp
+++ b/Engine/Plugins/ChaosClothAssetEditor/Source/ChaosClothAssetEditorTools/Private/ChaosClothAsset/ChaosClothAssetEditorToolsModule.cpp
@@ -26,19 +26,19 @@ namespace UE::Chaos::ClothAsset
 				NewObject<UClothEditorWeightMapPaintToolBuilder>(),                                 // Tool builder
 				ClothToolActions,                                                                   // Tool action bindings
 				FSlateIcon(FName("ClothStyle"), FName("ChaosClothAssetEditor.AddWeightMapNode")),   // Button icon (see FChaosClothAssetEditorStyle)
-				LOCTEXT("AddWeightMapNodeButtonText", "Weight Map"));                               // Button text
+				LOCTEXT("AddWeightMapNodeButtonText", "Cloth Weight Map"), FName("Cloth"));                               // Button text
 
 			ToolRegistry.AddNodeToToolMapping(FChaosClothAssetSelectionNode_v2::StaticType(), 
 				NewObject<UClothMeshSelectionToolBuilder>(), 
 				ClothToolActions, 
 				FSlateIcon(FName("ClothStyle"), FName("ChaosClothAssetEditor.AddMeshSelectionNode")),
-				LOCTEXT("AddSelectionNodeButtonText", "Selection"));
+				LOCTEXT("AddSelectionNodeButtonText", "Cloth Mesh Selection"), FName("Cloth"));
 
 			ToolRegistry.AddNodeToToolMapping(FChaosClothAssetTransferSkinWeightsNode::StaticType(), 
 				NewObject<UClothTransferSkinWeightsToolBuilder>(), 
 				ClothToolActions, 
 				FSlateIcon(FName("ClothStyle"), FName("ChaosClothAssetEditor.AddTransferSkinWeightsNode")),
-				LOCTEXT("AddTransferSkinWeightNodeButtonText", "Transfer Skin Weights"));
+				LOCTEXT("AddTransferSkinWeightNodeButtonText", "Cloth Skinning Transfer"), FName("Cloth"));
 		}
 
 		virtual void ShutdownModule() override
diff --git a/Engine/Plugins/Experimental/ChaosFlesh/Source/ChaosFleshEditor/Private/ChaosFlesh/Asset/AssetDefinition_FleshAsset.cpp b/Engine/Plugins/Experimental/ChaosFlesh/Source/ChaosFleshEditor/Private/ChaosFlesh/Asset/AssetDefinition_FleshAsset.cpp
--- a/Engine/Plugins/Experimental/ChaosFlesh/Source/ChaosFleshEditor/Private/ChaosFlesh/Asset/AssetDefinition_FleshAsset.cpp
+++ b/Engine/Plugins/Experimental/ChaosFlesh/Source/ChaosFleshEditor/Private/ChaosFlesh/Asset/AssetDefinition_FleshAsset.cpp
@@ -73,6 +73,7 @@ EAssetCommandResult UAssetDefinition_FleshAsset::OpenAssets(const FAssetOpenArgs
 			{
 				UAssetEditorSubsystem* const AssetEditorSubsystem = GEditor->GetEditorSubsystem<UAssetEditorSubsystem>();
 				UDataflowEditor* const AssetEditor = NewObject<UDataflowEditor>(AssetEditorSubsystem, NAME_None, RF_Transient);
+				AssetEditor->RegisterToolCategories({"General"});
 
 				const TSubclassOf<AActor> ActorClass = StaticLoadClass(AActor::StaticClass(), nullptr,
 					TEXT("/ChaosFlesh/BP_FleshPreview.BP_FleshPreview_C"), nullptr, LOAD_None, nullptr);
diff --git a/Engine/Plugins/Experimental/ChaosOutfitAsset/Source/ChaosOutfitAssetEditor/Private/ChaosOutfitAsset/AssetDefinition_OutfitAsset.cpp b/Engine/Plugins/Experimental/ChaosOutfitAsset/Source/ChaosOutfitAssetEditor/Private/ChaosOutfitAsset/AssetDefinition_OutfitAsset.cpp
--- a/Engine/Plugins/Experimental/ChaosOutfitAsset/Source/ChaosOutfitAssetEditor/Private/ChaosOutfitAsset/AssetDefinition_OutfitAsset.cpp
+++ b/Engine/Plugins/Experimental/ChaosOutfitAsset/Source/ChaosOutfitAssetEditor/Private/ChaosOutfitAsset/AssetDefinition_OutfitAsset.cpp
@@ -47,6 +47,7 @@ EAssetCommandResult UAssetDefinition_OutfitAsset::OpenAssets(const FAssetOpenArg
 	{
 		UAssetEditorSubsystem* const AssetEditorSubsystem = GEditor->GetEditorSubsystem<UAssetEditorSubsystem>();
 		UDataflowEditor* const AssetEditor = NewObject<UDataflowEditor>(AssetEditorSubsystem, NAME_None, RF_Transient);
+		AssetEditor->RegisterToolCategories({"General", "Cloth"});
 
 		static const TCHAR* const Path = TEXT("/ChaosOutfitAsset/BP_OutfitPreview.BP_OutfitPreview_C");
 		const TSubclassOf<AActor> PreviewClass = StaticLoadClass(AActor::StaticClass(), nullptr, Path);
diff --git a/Engine/Plugins/Experimental/Dataflow/Source/DataflowEditor/Private/Dataflow/AssetDefinition_DataflowAsset.cpp b/Engine/Plugins/Experimental/Dataflow/Source/DataflowEditor/Private/Dataflow/AssetDefinition_DataflowAsset.cpp
--- a/Engine/Plugins/Experimental/Dataflow/Source/DataflowEditor/Private/Dataflow/AssetDefinition_DataflowAsset.cpp
+++ b/Engine/Plugins/Experimental/Dataflow/Source/DataflowEditor/Private/Dataflow/AssetDefinition_DataflowAsset.cpp
@@ -190,6 +190,7 @@ EAssetCommandResult UAssetDefinition_DataflowAsset::OpenAssets(const FAssetOpenA
 		{
 			UAssetEditorSubsystem* const AssetEditorSubsystem = GEditor->GetEditorSubsystem<UAssetEditorSubsystem>();
 			UDataflowEditor* const AssetEditor = NewObject<UDataflowEditor>(AssetEditorSubsystem, NAME_None, RF_Transient);
+			AssetEditor->RegisterToolCategories({"General"});
 
 			// Validate the asset
 			if (UDataflow* const DataflowAsset = CastChecked<UDataflow>(DataflowObjects[0]))
diff --git a/Engine/Plugins/Experimental/Dataflow/Source/DataflowEditor/Private/Dataflow/DataflowEditorMode.cpp b/Engine/Plugins/Experimental/Dataflow/Source/DataflowEditor/Private/Dataflow/DataflowEditorMode.cpp
--- a/Engine/Plugins/Experimental/Dataflow/Source/DataflowEditor/Private/Dataflow/DataflowEditorMode.cpp
+++ b/Engine/Plugins/Experimental/Dataflow/Source/DataflowEditor/Private/Dataflow/DataflowEditorMode.cpp
@@ -260,6 +260,30 @@ bool UDataflowEditorMode::CanAddNode(FName NewNodeType) const
 	return false;
 }
 
+const TArray<FName>& UDataflowEditorMode::GetToolCategories()
+{
+	const UDataflowEditor* LocalEditor = DataflowEditor;
+	if(!LocalEditor)
+	{
+		if(ConstructionScene)
+		{
+			LocalEditor = ConstructionScene->GetDataflowEditor();
+		}
+		else if(FAssetEditorModeManager* ModeManager = StaticCast<FAssetEditorModeManager*>(Owner))
+		{
+			// The local construction scene is not yet set when calling that function.
+			// We then have to rely on the mode manager one 
+			if(FDataflowConstructionScene* ManagerScene = StaticCast<FDataflowConstructionScene*>(ModeManager->GetPreviewScene()))
+			{
+				LocalEditor = ManagerScene->GetDataflowEditor();
+			}
+		}
+	}
+	static TArray<FName> EmptyToolCategories;
+	return LocalEditor ? LocalEditor->GetToolCategories() : EmptyToolCategories;
+}
+
+
 void UDataflowEditorMode::RegisterTools()
 {
 	const FDataflowEditorCommandsImpl& CommandInfos = FDataflowEditorCommands::Get();
@@ -268,23 +292,28 @@ void UDataflowEditorMode::RegisterTools()
 
 	UE::Dataflow::FDataflowToolRegistry& ToolRegistry = UE::Dataflow::FDataflowToolRegistry::Get();
 	const TArray<FName> NodeNames = ToolRegistry.GetNodeNames();
+	const TArray<FName> ToolCategories = GetToolCategories();
+	
 	for (const FName& RegisteredNodeName : NodeNames)
 	{
-		const TSharedPtr<FUICommandInfo> CommandInfo = ToolRegistry.GetToolCommandForNode(RegisteredNodeName);
-		UInteractiveToolBuilder* const Builder = ToolRegistry.GetToolBuilderForNode(RegisteredNodeName);
-
-		// TODO: This is here only so the Tool can hide the all meshes in the DataflowConstructionScene. That should probably be handed in this class instead.
-		if (UDataflowEditorWeightMapPaintToolBuilder* WeightMapPaintToolBuilder = Cast<UDataflowEditorWeightMapPaintToolBuilder>(Builder))
+		if(ToolCategories.IsEmpty() || ToolCategories.Contains(ToolRegistry.GetToolCategoryForNode(RegisteredNodeName)))
 		{
-			WeightMapPaintToolBuilder->SetEditorMode(this);
-		}
+			const TSharedPtr<FUICommandInfo> CommandInfo = ToolRegistry.GetToolCommandForNode(RegisteredNodeName);
+			UInteractiveToolBuilder* const Builder = ToolRegistry.GetToolBuilderForNode(RegisteredNodeName);
 
-		RegisterDataflowTool(CommandInfo, RegisteredNodeName.ToString() + FString(TEXT("Tool")), Builder, ConstructionViewportToolsContext);
+			// TODO: This is here only so the Tool can hide the all meshes in the DataflowConstructionScene. That should probably be handed in this class instead.
+			if (UDataflowEditorWeightMapPaintToolBuilder* WeightMapPaintToolBuilder = Cast<UDataflowEditorWeightMapPaintToolBuilder>(Builder))
+			{
+				WeightMapPaintToolBuilder->SetEditorMode(this);
+			}
 
-		NodeTypeToToolCommandMap.Add(RegisteredNodeName, CommandInfo);
+			RegisterDataflowTool(CommandInfo, RegisteredNodeName.ToString() + FString(TEXT("Tool")), Builder, ConstructionViewportToolsContext);
 
-		// Register "Add Node" commands for buttons in the UI. The EditorToolkit will construct the actual toolbar buttons.
-		NodeTypeToAddNodeCommandMap.Add(RegisteredNodeName, ToolRegistry.GetAddNodeCommandForNode(RegisteredNodeName));
+			NodeTypeToToolCommandMap.Add(RegisteredNodeName, CommandInfo);
+
+			// Register "Add Node" commands for buttons in the UI. The EditorToolkit will construct the actual toolbar buttons.
+			NodeTypeToAddNodeCommandMap.Add(RegisteredNodeName, ToolRegistry.GetAddNodeCommandForNode(RegisteredNodeName));
+		}
 	}
 
 }
diff --git a/Engine/Plugins/Experimental/Dataflow/Source/DataflowEditor/Private/Dataflow/DataflowToolRegistry.cpp b/Engine/Plugins/Experimental/Dataflow/Source/DataflowEditor/Private/Dataflow/DataflowToolRegistry.cpp
--- a/Engine/Plugins/Experimental/Dataflow/Source/DataflowEditor/Private/Dataflow/DataflowToolRegistry.cpp
+++ b/Engine/Plugins/Experimental/Dataflow/Source/DataflowEditor/Private/Dataflow/DataflowToolRegistry.cpp
@@ -18,14 +18,15 @@ namespace UE::Dataflow
 		TLazySingleton<FDataflowToolRegistry>::TearDown();
 	}
 
-	void FDataflowToolRegistry::AddNodeToToolMapping(const FName& NodeName, TObjectPtr<UInteractiveToolBuilder> ToolBuilder, const TSharedRef<const IDataflowToolActionCommands>& ToolActionCommands, const FSlateIcon& AddNodeButtonIcon, const FText& AddNodeButtonText)
+	void FDataflowToolRegistry::AddNodeToToolMapping(const FName& NodeName, TObjectPtr<UInteractiveToolBuilder> ToolBuilder, const TSharedRef<const IDataflowToolActionCommands>& ToolActionCommands,
+		const FSlateIcon& AddNodeButtonIcon, const FText& AddNodeButtonText, const FName& ToolCategory)
 	{
 		ensureMsgf(!FDataflowEditorCommands::IsRegistered(), TEXT("FDataflowToolRegistry: DataflowEditorCommands have already been registered. \
 				Newly registered Tools may not be available in the Editor. \
 				Ensure that AddNodeToToolMapping is called before the DataflowEditor module is loaded."));
 
 		// FUICommandInfo is uninitialized, it will be created later in FDataflowEditorCommandsImpl::RegisterCommands
-		NodeTypeToToolMap.Add(NodeName, FToolInfo{ ToolBuilder, ToolActionCommands, AddNodeButtonIcon, AddNodeButtonText, nullptr, nullptr });
+		NodeTypeToToolMap.Add(NodeName, FToolInfo{ ToolBuilder, ToolActionCommands, AddNodeButtonIcon, AddNodeButtonText, nullptr, nullptr, ToolCategory });
 	}
 
 	PRAGMA_DISABLE_DEPRECATION_WARNINGS
@@ -38,7 +39,7 @@ namespace UE::Dataflow
 				Ensure that AddNodeToToolMapping is called before the DataflowEditor module is loaded."));
 
 		// FUICommandInfo is uninitialized, it will be created later in FDataflowEditorCommandsImpl::RegisterCommands
-		NodeTypeToToolMap.Add(NodeName, FToolInfo{ ToolBuilder, ToolActionCommands, FSlateIcon(), FText(), nullptr, nullptr });
+		NodeTypeToToolMap.Add(NodeName, FToolInfo{ ToolBuilder, ToolActionCommands, FSlateIcon(), FText(), nullptr, nullptr, FName("General") });
 	}
 	PRAGMA_ENABLE_DEPRECATION_WARNINGS
 
@@ -78,6 +79,12 @@ namespace UE::Dataflow
 		return NodeTypeToToolMap[NodeName].ToolCommand;
 	}
 
+	const FName& FDataflowToolRegistry::GetToolCategoryForNode(const FName& NodeName)
+	{
+		check(NodeTypeToToolMap.Contains(NodeName));
+		return NodeTypeToToolMap[NodeName].ToolCategory;
+	}
+
 	UInteractiveToolBuilder* FDataflowToolRegistry::GetToolBuilderForNode(const FName& NodeName)
 	{
 		check(NodeTypeToToolMap.Contains(NodeName));
diff --git a/Engine/Plugins/Experimental/Dataflow/Source/DataflowEditor/Public/Dataflow/DataflowEditor.h b/Engine/Plugins/Experimental/Dataflow/Source/DataflowEditor/Public/Dataflow/DataflowEditor.h
--- a/Engine/Plugins/Experimental/Dataflow/Source/DataflowEditor/Public/Dataflow/DataflowEditor.h
+++ b/Engine/Plugins/Experimental/Dataflow/Source/DataflowEditor/Public/Dataflow/DataflowEditor.h
@@ -60,6 +60,12 @@ public:
 	TArray<TObjectPtr<UDataflowBaseContent>>& GetTerminalContents() { return TerminalContents; }
 	const TArray<TObjectPtr<UDataflowBaseContent>>& GetTerminalContents() const { return TerminalContents; }
 
+	/** Get the registered tool catagories */
+	const TArray<FName>& GetToolCategories() const { return ToolCategories; }
+
+	/** Registere tool categories available for that construction scene */
+	void RegisterToolCategories(const TArray<FName>& InToolCategrories) { ToolCategories = InToolCategrories; }
+
 private :
 
 	friend class FDataflowEditorToolkit;
@@ -82,6 +88,9 @@ private :
 	/** List of dataflow contents available in the graph and coming from all the terminal nodes */
 	UPROPERTY()
 	TArray<TObjectPtr<UDataflowBaseContent>> TerminalContents;
+
+	/** List of tool categories registered for that editor */
+	TArray<FName> ToolCategories;
 };
 
 DECLARE_LOG_CATEGORY_EXTERN(LogDataflowEditor, Log, All);
diff --git a/Engine/Plugins/Experimental/Dataflow/Source/DataflowEditor/Public/Dataflow/DataflowEditorMode.h b/Engine/Plugins/Experimental/Dataflow/Source/DataflowEditor/Public/Dataflow/DataflowEditorMode.h
--- a/Engine/Plugins/Experimental/Dataflow/Source/DataflowEditor/Public/Dataflow/DataflowEditorMode.h
+++ b/Engine/Plugins/Experimental/Dataflow/Source/DataflowEditor/Public/Dataflow/DataflowEditorMode.h
@@ -179,6 +179,9 @@ private:
 		bPendingNodeSelectionChanged = bIsPending;
 	}
 
+	/** Get the tool categories from the construction scene */
+	const TArray<FName>& GetToolCategories();
+
 	/** Dataflow construction scene from the toolkit */
 	FDataflowConstructionScene* ConstructionScene = nullptr;
 
@@ -243,6 +246,5 @@ private:
 	// Each node has a set of view modes it can use. For each distinct set of available view modes, we store the last mode that the user has chosen.
 	// This will enable switching between different kinds of nodes and not have the view mode reset to a default every time.
 	TMap<TArray<FName>, FName> NodeViewModeHistory;
-	
 };
 
diff --git a/Engine/Plugins/Experimental/Dataflow/Source/DataflowEditor/Public/Dataflow/DataflowEditorPreviewSceneBase.h b/Engine/Plugins/Experimental/Dataflow/Source/DataflowEditor/Public/Dataflow/DataflowEditorPreviewSceneBase.h
--- a/Engine/Plugins/Experimental/Dataflow/Source/DataflowEditor/Public/Dataflow/DataflowEditorPreviewSceneBase.h
+++ b/Engine/Plugins/Experimental/Dataflow/Source/DataflowEditor/Public/Dataflow/DataflowEditorPreviewSceneBase.h
@@ -147,6 +147,9 @@ public:
 
 	/** Return the scene selected components */
 	static USelection* GetSelectedComponents(FDataflowPreviewSceneBase* DataflowScene);
+
+	/** Get the dataflow editor */
+	const UDataflowEditor* GetDataflowEditor() const { return DataflowEditor; }
 	
 protected:
 
diff --git a/Engine/Plugins/Experimental/Dataflow/Source/DataflowEditor/Public/Dataflow/DataflowToolRegistry.h b/Engine/Plugins/Experimental/Dataflow/Source/DataflowEditor/Public/Dataflow/DataflowToolRegistry.h
--- a/Engine/Plugins/Experimental/Dataflow/Source/DataflowEditor/Public/Dataflow/DataflowToolRegistry.h
+++ b/Engine/Plugins/Experimental/Dataflow/Source/DataflowEditor/Public/Dataflow/DataflowToolRegistry.h
@@ -34,8 +34,9 @@ namespace UE::Dataflow
 		void AddNodeToToolMapping(const FName& NodeName, 
 			TObjectPtr<UInteractiveToolBuilder> ToolBuilder, 
 			const TSharedRef<const IDataflowToolActionCommands>& ToolActionCommands, 
-			const FSlateIcon& AddNodeButtonIcon, 
-			const FText& AddNodeButtonText);
+			const FSlateIcon& AddNodeButtonIcon,
+			const FText& AddNodeButtonText, 
+			const FName& ToolCategory = FName("General"));
 
 		UE_DEPRECATED(5.6, "Please use the version of AddNodeToToolMapping taking AddNodeButtonIcon and AddNodeButtonText")
 		void AddNodeToToolMapping(const FName& NodeName,
@@ -52,6 +53,9 @@ namespace UE::Dataflow
 
 		TSharedPtr<FUICommandInfo>& GetToolCommandForNode(const FName& NodeName);
 
+		/** Return the category the tool has been registered in */
+		const FName& GetToolCategoryForNode(const FName& NodeName);
+
 		UInteractiveToolBuilder* GetToolBuilderForNode(const FName& NodeName);
 		const UInteractiveToolBuilder* GetToolBuilderForNode(const FName& NodeName) const;
 
@@ -71,6 +75,9 @@ namespace UE::Dataflow
 			// Constructed automatically in FDataflowEditorCommandsImpl::RegisterCommands
 			TSharedPtr<FUICommandInfo> AddNodeCommand;
 			TSharedPtr<FUICommandInfo> ToolCommand;
+
+			/** Tool category for filtering */
+			FName ToolCategory;
 		};
 
 		TMap<FName, FToolInfo> NodeTypeToToolMap;
diff --git a/Engine/Plugins/Experimental/GeometryCollectionPlugin/Source/GeometryCollectionEditor/Private/GeometryCollection/AssetDefinition_GeometryCollection.cpp b/Engine/Plugins/Experimental/GeometryCollectionPlugin/Source/GeometryCollectionEditor/Private/GeometryCollection/AssetDefinition_GeometryCollection.cpp
--- a/Engine/Plugins/Experimental/GeometryCollectionPlugin/Source/GeometryCollectionEditor/Private/GeometryCollection/AssetDefinition_GeometryCollection.cpp
+++ b/Engine/Plugins/Experimental/GeometryCollectionPlugin/Source/GeometryCollectionEditor/Private/GeometryCollection/AssetDefinition_GeometryCollection.cpp
@@ -85,6 +85,7 @@ EAssetCommandResult UAssetDefinition_GeometryCollection::OpenAssets(const FAsset
 				UAssetEditorSubsystem* const AssetEditorSubsystem = GEditor->GetEditorSubsystem<UAssetEditorSubsystem>();
 				
 				UDataflowEditor* const AssetEditor = NewObject<UDataflowEditor>(AssetEditorSubsystem, NAME_None, RF_Transient);
+				AssetEditor->RegisterToolCategories({"General"});
 				const TSubclassOf<AActor> ActorClass = StaticLoadClass(AActor::StaticClass(), nullptr,
 					TEXT("/GeometryCollectionPlugin/BP_GeometryCollectionPreview.BP_GeometryCollectionPreview_C"), nullptr, LOAD_None, nullptr);
 				AssetEditor->Initialize({ GeometryCollection }, ActorClass);
diff --git a/Engine/Plugins/Runtime/HairStrands/Source/HairStrandsEditor/Private/GroomActions.cpp b/Engine/Plugins/Runtime/HairStrands/Source/HairStrandsEditor/Private/GroomActions.cpp
--- a/Engine/Plugins/Runtime/HairStrands/Source/HairStrandsEditor/Private/GroomActions.cpp
+++ b/Engine/Plugins/Runtime/HairStrands/Source/HairStrandsEditor/Private/GroomActions.cpp
@@ -63,6 +63,7 @@ EAssetCommandResult UAssetDefinition_GroomAsset::OpenAssets(const FAssetOpenArgs
 			{
 				UAssetEditorSubsystem* const AssetEditorSubsystem = GEditor->GetEditorSubsystem<UAssetEditorSubsystem>();
 				UDataflowEditor* const AssetEditor = NewObject<UDataflowEditor>(AssetEditorSubsystem, NAME_None, RF_Transient);
+				AssetEditor->RegisterToolCategories({"General"});
 				
 				const TSubclassOf<AActor> ActorClass = StaticLoadClass(AActor::StaticClass(), nullptr,
 					TEXT("/HairStrands/BP_PreviewGroom.BP_PreviewGroom_C"), nullptr, LOAD_None, nullptr);
