Add Replace/Add override mode to dataflow painting tool

diff --git a/Engine/Plugins/Experimental/Dataflow/Source/DataflowEditor/Private/DataflowEditorTools/DataflowEditorWeightMapPaintTool.cpp b/Engine/Plugins/Experimental/Dataflow/Source/DataflowEditor/Private/DataflowEditorTools/DataflowEditorWeightMapPaintTool.cpp
--- a/Engine/Plugins/Experimental/Dataflow/Source/DataflowEditor/Private/DataflowEditorTools/DataflowEditorWeightMapPaintTool.cpp
+++ b/Engine/Plugins/Experimental/Dataflow/Source/DataflowEditor/Private/DataflowEditorTools/DataflowEditorWeightMapPaintTool.cpp
@@ -387,9 +387,28 @@ void UDataflowEditorWeightMapPaintTool::Setup()
 	Mesh->Attributes()->SetNumWeightLayers(NumAttributeLayers + 1);
 	ActiveWeightMap = Mesh->Attributes()->GetWeightLayer(NumAttributeLayers);
 	ActiveWeightMap->SetName(FName("PaintLayer"));
+	
+	// Setup DynamicMeshToWeight conversion and get Input weight map (if it exists)
+	const int32 NumExpectedWeights = bHaveDynamicMeshToWeightConversion ? WeightToDynamicMesh.Num() : GetSculptMesh()->MaxVertexID();
+	SetupWeights.SetNumZeroed(NumExpectedWeights);
+
+	if(DataflowEditorContextObject && WeightMapNodeToUpdate)
+	{
+		// Find the map if it exists.
+		if (TSharedPtr<UE::Dataflow::FEngineContext> DataflowContext = DataflowEditorContextObject->GetDataflowContext())
+		{
+			// Get the attribute key
+			FCollectionAttributeKey AttributeKey = WeightMapNodeToUpdate->GetWeightAttributeKey(*DataflowContext);
+
+			// Fill the attribute values
+			WeightMapNodeToUpdate->FillAttributeWeights(DataflowEditorContextObject->GetSelectedCollection(), AttributeKey, SetupWeights);
+		}
+	}
 
 	// Copy weights from selected node to the preview mesh
-	const TArray<float>& CurrentWeights = WeightMapNodeToUpdate->VertexWeights;
+	TArray<float> CurrentWeights;
+	CurrentWeights.SetNumZeroed(NumExpectedWeights);
+	WeightMapNodeToUpdate->ExtractVertexWeights(SetupWeights, TArrayView<float>(CurrentWeights));
 	
 	if (bHaveDynamicMeshToWeightConversion)
 	{
@@ -1771,19 +1790,17 @@ void UDataflowEditorWeightMapPaintTool::UpdateSelectedNode()
 
 	if (bHaveDynamicMeshToWeightConversion)
 	{
-		if (WeightToDynamicMesh.Num() != WeightMapNodeToUpdate->VertexWeights.Num())
-		{
-			WeightMapNodeToUpdate->VertexWeights.Init(0.f, WeightToDynamicMesh.Num());
-		}
-
+		TArray<float> FinalWeights;
+		FinalWeights.Init(0.f, WeightToDynamicMesh.Num());
 		for (int32 DynamicMeshIdx = 0; DynamicMeshIdx < CurrentWeights.Num(); ++DynamicMeshIdx)
 		{
-			WeightMapNodeToUpdate->VertexWeights[DynamicMeshToWeight[DynamicMeshIdx]] = CurrentWeights[DynamicMeshIdx];
+			FinalWeights[DynamicMeshToWeight[DynamicMeshIdx]] = CurrentWeights[DynamicMeshIdx];
 		}
+		WeightMapNodeToUpdate->ReportVertexWeights(SetupWeights, FinalWeights);
 	}
 	else
 	{
-		WeightMapNodeToUpdate->VertexWeights = CurrentWeights;
+		WeightMapNodeToUpdate->ReportVertexWeights(SetupWeights, CurrentWeights);
 	}
 	
 	WeightMapNodeToUpdate->Name = UpdateWeightMapProperties->Name;
diff --git a/Engine/Plugins/Experimental/Dataflow/Source/DataflowEditor/Public/DataflowEditorTools/DataflowEditorWeightMapPaintTool.h b/Engine/Plugins/Experimental/Dataflow/Source/DataflowEditor/Public/DataflowEditorTools/DataflowEditorWeightMapPaintTool.h
--- a/Engine/Plugins/Experimental/Dataflow/Source/DataflowEditor/Public/DataflowEditorTools/DataflowEditorWeightMapPaintTool.h
+++ b/Engine/Plugins/Experimental/Dataflow/Source/DataflowEditor/Public/DataflowEditorTools/DataflowEditorWeightMapPaintTool.h
@@ -86,12 +86,6 @@ enum class EDataflowEditorWeightMapPaintInteractionType : uint8
 	LastValue UMETA(Hidden)
 };
 
-
-
-
-
-
-
 /** Mesh Sculpting Brush Types */
 UENUM()
 enum class EDataflowEditorWeightMapPaintBrushType : uint8
@@ -125,8 +119,6 @@ enum class EDataflowEditorWeightMapPaintVisibilityType : uint8
 	Unoccluded
 };
 
-
-
 // TODO: Look at EditConditions for all these properties. Which ones make sense for which SubToolType?
 
 UCLASS()
@@ -204,10 +196,6 @@ public:
 
 };
 
-
-
-
-
 UENUM()
 enum class EDataflowEditorWeightMapPaintToolActions
 {
@@ -218,7 +206,6 @@ enum class EDataflowEditorWeightMapPaintToolActions
 	InvertCurrentSurface
 };
 
-
 UCLASS()
 class DATAFLOWEDITOR_API UDataflowEditorMeshWeightMapPaintToolActions : public UInteractiveToolPropertySet
 {
@@ -502,6 +489,9 @@ protected:
 	TArray<int32> DynamicMeshToWeight; 
 	TArray<TArray<int32>> WeightToDynamicMesh;
 
+	/** Setup weights used to store the initial values */
+	TArray<float> SetupWeights;
+
 protected:
 	virtual bool ShowWorkPlane() const override { return false; }
 
diff --git a/Engine/Plugins/Experimental/Dataflow/Source/DataflowNodes/Private/Dataflow/DataflowCollectionAddScalarVertexPropertyNode.cpp b/Engine/Plugins/Experimental/Dataflow/Source/DataflowNodes/Private/Dataflow/DataflowCollectionAddScalarVertexPropertyNode.cpp
--- a/Engine/Plugins/Experimental/Dataflow/Source/DataflowNodes/Private/Dataflow/DataflowCollectionAddScalarVertexPropertyNode.cpp
+++ b/Engine/Plugins/Experimental/Dataflow/Source/DataflowNodes/Private/Dataflow/DataflowCollectionAddScalarVertexPropertyNode.cpp
@@ -155,20 +155,12 @@ TArray<UE::Dataflow::FRenderingParameter> FDataflowCollectionAddScalarVertexProp
 	return FDataflowAddScalarVertexPropertyCallbackRegistry::Get().GetRenderingParameters(TargetGroup.Name);
 }
 
-
 void FDataflowCollectionAddScalarVertexPropertyNode::Evaluate(UE::Dataflow::FContext& Context, const FDataflowOutput* Out) const
 {
 	using namespace UE::Dataflow;
 
 	// Get the pin value if plugged
-	FCollectionAttributeKey Key = GetValue(Context, &AttributeKey, AttributeKey);
-
-	// If nothing set used the local value
-	if(Key.Attribute.IsEmpty() && Key.Group.IsEmpty())
-	{
-		Key.Group = TargetGroup.Name.ToString();
-		Key.Attribute = Name;
-	}
+	FCollectionAttributeKey Key = GetWeightAttributeKey(Context);
 
 	if (Out->IsA<FManagedArrayCollection>(&Collection))
 	{
@@ -179,23 +171,21 @@ void FDataflowCollectionAddScalarVertexPropertyNode::Evaluate(UE::Dataflow::FCon
 		{
 			const FName InName(Key.Attribute);
 			const FName InGroup(Key.Group);
-			TManagedArray<float>& Scalar = InCollection.AddAttribute<float>(InName, InGroup);
+			TManagedArray<float>& ScalarWeights = InCollection.AddAttribute<float>(InName, InGroup);
 
-			const int32 MaxWeightIndex = FMath::Min(VertexWeights.Num(), Scalar.Num());
-			if (VertexWeights.Num() > 0 && VertexWeights.Num() != Scalar.Num())
+			if (VertexWeights.Num() > 0 && VertexWeights.Num() != ScalarWeights.Num())
 			{
 				FDataflowTools::LogAndToastWarning(*this,
 					LOCTEXT("VertexCountMismatchHeadline", "Vertex count mismatch."),
 					FText::Format(LOCTEXT("VertexCountMismatchDetails", "Vertex weights in the node: {0}\n Vertices in group \"{1}\" in the Collection: {2}"),
 						VertexWeights.Num(),
 						FText::FromName(InGroup),
-						Scalar.Num()));
+						ScalarWeights.Num()));
 			}
 
-			for (int32 VertexID = 0; VertexID < MaxWeightIndex; ++VertexID)
-			{
-				Scalar[VertexID] = VertexWeights[VertexID];
-			}
+			const TArray<float> SetupWeights = ScalarWeights.GetConstArray();
+			ExtractVertexWeights(SetupWeights, TArrayView<float>(ScalarWeights.GetData(), ScalarWeights.Num()));
+			
 		}
 
 		SetValue(Context, MoveTemp(InCollection), &Collection);
@@ -206,4 +196,92 @@ void FDataflowCollectionAddScalarVertexPropertyNode::Evaluate(UE::Dataflow::FCon
 	}
 }
 
+bool FDataflowCollectionAddScalarVertexPropertyNode::FillAttributeWeights( 
+	const TSharedPtr<const FManagedArrayCollection> SelectedCollection, const FCollectionAttributeKey& InAttributeKey, TArray<float>& OutAttributeValues) const
+{
+	if (!InAttributeKey.Attribute.IsEmpty())
+	{
+		const FName InName(InAttributeKey.Attribute);
+		const FName InGroup(InAttributeKey.Group);
+		
+		if(const TManagedArray<float>* AttributeArray = SelectedCollection->FindAttributeTyped<float>(InName, InGroup))
+		{
+			OutAttributeValues = AttributeArray->GetConstArray();
+			return true;
+		}
+	}
+	OutAttributeValues.Init(0.0f, OutAttributeValues.Num());
+	return false;
+}
+
+FCollectionAttributeKey FDataflowCollectionAddScalarVertexPropertyNode::GetWeightAttributeKey(UE::Dataflow::FContext& Context) const
+{
+	// Get the pin value if plugged
+	FCollectionAttributeKey Key = GetValue(Context, &AttributeKey, AttributeKey);
+
+	// If nothing set used the local value
+	if(Key.Attribute.IsEmpty() && Key.Group.IsEmpty())
+	{
+		Key.Group = TargetGroup.Name.ToString();
+		Key.Attribute = Name;
+	}
+	return Key;
+}
+
+void FDataflowCollectionAddScalarVertexPropertyNode::ReportVertexWeights(const TArray<float>& SetupWeights, const TArray<float>& FinalWeights)
+{
+	check(SetupWeights.Num() == FinalWeights.Num());
+	VertexWeights.SetNumZeroed(FinalWeights.Num());
+	for (int32 WeightIndex = 0, NumWeights = FinalWeights.Num(); WeightIndex <NumWeights; ++WeightIndex)
+	{
+		switch (OverrideType)
+		{
+			case EDataflowWeightMapOverrideType::ReplaceAll:
+				VertexWeights[WeightIndex] = FinalWeights[WeightIndex];
+				break;
+			case EDataflowWeightMapOverrideType::ReplaceChanged:
+				VertexWeights[WeightIndex] = (SetupWeights[WeightIndex] == FinalWeights[WeightIndex]) ? ReplaceChangedPassthroughValue : FinalWeights[WeightIndex];
+				break;
+			case EDataflowWeightMapOverrideType::AddDifference:
+				VertexWeights[WeightIndex] = FinalWeights[WeightIndex] - SetupWeights[WeightIndex];
+				break;
+			default: unimplemented();
+		}
+	}
+}
+
+void FDataflowCollectionAddScalarVertexPropertyNode::ExtractVertexWeights(const TArray<float>& SetupWeights, TArrayView<float> FinalWeights) const
+{
+	check(SetupWeights.Num() == FinalWeights.Num());
+
+	int32 NumWeights = FMath::Min(VertexWeights.Num(), SetupWeights.Num());
+	if(NumWeights != 0)
+	{ 
+		for (int32 WeightIndex = 0; WeightIndex < NumWeights; ++WeightIndex)
+		{
+			switch (OverrideType)
+			{
+				case EDataflowWeightMapOverrideType::ReplaceAll:
+					FinalWeights[WeightIndex] = FMath::Clamp(VertexWeights[WeightIndex], 0.f, 1.f);
+					break;
+				case EDataflowWeightMapOverrideType::ReplaceChanged:
+					FinalWeights[WeightIndex] = FMath::Clamp( (VertexWeights[WeightIndex] == ReplaceChangedPassthroughValue) ? SetupWeights[WeightIndex] : VertexWeights[WeightIndex], 0.0f, 1.0f);
+					break;
+				case EDataflowWeightMapOverrideType::AddDifference:
+					FinalWeights[WeightIndex] = FMath::Clamp(SetupWeights[WeightIndex] + VertexWeights[WeightIndex], 0.f, 1.f);
+					break;
+				default: unimplemented();
+			}
+		}
+	}
+	else
+	{
+		NumWeights = SetupWeights.Num();
+		for (int32 WeightIndex = 0; WeightIndex < NumWeights; ++WeightIndex)
+		{
+			FinalWeights[WeightIndex] = SetupWeights[WeightIndex];
+		}
+	}
+}
+
 #undef LOCTEXT_NAMESPACE
diff --git a/Engine/Plugins/Experimental/Dataflow/Source/DataflowNodes/Public/Dataflow/DataflowCollectionAddScalarVertexPropertyNode.h b/Engine/Plugins/Experimental/Dataflow/Source/DataflowNodes/Public/Dataflow/DataflowCollectionAddScalarVertexPropertyNode.h
--- a/Engine/Plugins/Experimental/Dataflow/Source/DataflowNodes/Public/Dataflow/DataflowCollectionAddScalarVertexPropertyNode.h
+++ b/Engine/Plugins/Experimental/Dataflow/Source/DataflowNodes/Public/Dataflow/DataflowCollectionAddScalarVertexPropertyNode.h
@@ -65,13 +65,25 @@ struct FScalarVertexPropertyGroup
 	FName Name = FGeometryCollection::VerticesGroup;
 };
 
+/** How the map stored on the AddWeightMapNode should be applied to an existing map. If no map exists, it is treated as zero.*/
+UENUM()
+enum class EDataflowWeightMapOverrideType : uint8
+{
+	
+	/** Replace all the values.*/
+	ReplaceAll,
+	/** Add the values difference to the input one. */
+	AddDifference,
+	/** Replace only the values that has changed. */
+	ReplaceChanged,
+};
 
 /** Scalar vertex properties. */
 USTRUCT(Meta = (DataflowCollection))
 struct DATAFLOWNODES_API FDataflowCollectionAddScalarVertexPropertyNode : public FDataflowNode
 {
 	GENERATED_USTRUCT_BODY()
-	DATAFLOW_NODE_DEFINE_INTERNAL(FDataflowCollectionAddScalarVertexPropertyNode, "AddScalarVertexProperty", "Collection", "Add a saved scalar property to a collection")
+	DATAFLOW_NODE_DEFINE_INTERNAL(FDataflowCollectionAddScalarVertexPropertyNode, "PaintWeightMap", "Collection", "Paint a weight map and save it to collection")
 
 	virtual TArray<UE::Dataflow::FRenderingParameter> GetRenderParametersImpl() const override;
 
@@ -82,20 +94,40 @@ public:
 	FManagedArrayCollection Collection;
 
 	/** The name to be set as a weight map attribute. */
-	UPROPERTY(EditAnywhere, Category = "Vertex Attribute", meta = (DisplayName = "Attribute"))
+	UPROPERTY(EditAnywhere, Category = "Vertex Attributes", meta = (DisplayName = "Attribute Name"))
 	FString Name;
 
-	UPROPERTY(meta = (DisplayName = "AttributeKey", DataflowInput, DataflowOutput, DataflowPassthrough = "AttributeKey"))
+	UPROPERTY(meta = (DisplayName = "Attribute Key", DataflowInput, DataflowOutput, DataflowPassthrough = "AttributeKey"))
 	FCollectionAttributeKey AttributeKey;
 
 	UPROPERTY()
 	TArray<float> VertexWeights;
 
-	UPROPERTY(EditAnywhere, Category = "Vertex Attribute", meta = (DisplayName = "Group"))
+	UPROPERTY(EditAnywhere, Category = "Vertex Attributes", meta = (DisplayName = "Vertex Group"))
 	FScalarVertexPropertyGroup TargetGroup;
+	
+	/** This type will define how the data are applied to the input data */
+	UPROPERTY(EditAnywhere, Category = "Weight Map")
+	EDataflowWeightMapOverrideType OverrideType = EDataflowWeightMapOverrideType::ReplaceAll;
 
 	FDataflowCollectionAddScalarVertexPropertyNode(const UE::Dataflow::FNodeParameters& InParam, FGuid InGuid = FGuid::NewGuid());
 
+	/** Report the vertex weights back onto the property ones */
+	void ReportVertexWeights(const TArray<float>& SetupWeights, const TArray<float>& FinalWeights);
+
+	/** Extract the vertex weights back from the property ones */
+	void ExtractVertexWeights(const TArray<float>& SetupWeights, TArrayView<float> FinalWeights) const;
+
+	/** Fill the weight attribute values from the collection */
+	bool FillAttributeWeights(const TSharedPtr<const FManagedArrayCollection> SelectedCollection, const FCollectionAttributeKey& InAttributeKey, TArray<float>& OutAttributeValues) const;
+
+	/** Get the weights attribute key to retrieve/set the weight values*/
+	FCollectionAttributeKey GetWeightAttributeKey(UE::Dataflow::FContext& Context) const;
+
 private:
+
+	/** Pass through value to skip replacing the weight map value if nothing has changed */
+	static constexpr float ReplaceChangedPassthroughValue = UE_BIG_NUMBER;
+	
 	virtual void Evaluate(UE::Dataflow::FContext& Context, const FDataflowOutput* Out) const override;
 };
