Dataflow : Add category filtering  per asset in the graph action menu
- Add category registry
- Filtering is on by default
- Register per asset category for Cloth/ Geometry collection / Flesh / Groom
- Add a option for filtering in the contextual action menu

diff --git a/Engine/Plugins/ChaosClothAssetEditor/Source/ChaosClothAssetDataflowNodes/Private/ChaosClothAsset/ChaosClothAssetDataflowNodesModule.cpp b/Engine/Plugins/ChaosClothAssetEditor/Source/ChaosClothAssetDataflowNodes/Private/ChaosClothAsset/ChaosClothAssetDataflowNodesModule.cpp
--- a/Engine/Plugins/ChaosClothAssetEditor/Source/ChaosClothAssetDataflowNodes/Private/ChaosClothAsset/ChaosClothAssetDataflowNodesModule.cpp
+++ b/Engine/Plugins/ChaosClothAssetEditor/Source/ChaosClothAssetDataflowNodes/Private/ChaosClothAsset/ChaosClothAssetDataflowNodesModule.cpp
@@ -6,6 +6,7 @@
 #include "ChaosClothAsset/AttributeNode.h"
 #include "ChaosClothAsset/BindToRootBoneNode.h"
 #include "ChaosClothAsset/BlendVerticesNode.h"
+#include "ChaosClothAsset/ClothAsset.h"
 #include "ChaosClothAsset/ClothCollectionGroup.h"
 #include "ChaosClothAsset/ClothCollectionToDynamicMeshNode.h"
 #include "ChaosClothAsset/ClothCollectionQueryNode.h"
@@ -78,6 +79,7 @@
 #include "ChaosClothAsset/SelectionGroupCustomization.h"
 #include "ChaosClothAsset/WeightMapToSelectionNode.h"
 #include "Dataflow/DataflowCollectionAddScalarVertexPropertyNode.h"
+#include "Dataflow/DataflowCategoryRegistry.h"
 #include "Dataflow/DataflowNodeColorsRegistry.h"
 #include "Dataflow/DataflowNodeFactory.h"
 #include "Dataflow/DataflowRenderingFactory.h"
@@ -360,6 +362,8 @@ PRAGMA_ENABLE_DEPRECATION_WARNINGS
 
 			Private::RegisterDataflowNodes();
 
+			UE_DATAFLOW_REGISTER_CATEGORY_FORASSET_TYPE("Cloth", UChaosClothAsset);
+
 			// Register type customizations
 			if (FPropertyEditorModule* const PropertyModule = FModuleManager::GetModulePtr<FPropertyEditorModule>("PropertyEditor"))
 			{
diff --git a/Engine/Plugins/Experimental/ChaosFlesh/Source/ChaosFleshNodes/Private/Dataflow/ChaosFleshNodesPlugin.cpp b/Engine/Plugins/Experimental/ChaosFlesh/Source/ChaosFleshNodes/Private/Dataflow/ChaosFleshNodesPlugin.cpp
--- a/Engine/Plugins/Experimental/ChaosFlesh/Source/ChaosFleshNodes/Private/Dataflow/ChaosFleshNodesPlugin.cpp
+++ b/Engine/Plugins/Experimental/ChaosFlesh/Source/ChaosFleshNodes/Private/Dataflow/ChaosFleshNodesPlugin.cpp
@@ -33,6 +33,7 @@
 #include "Dataflow/ChaosFleshVertexConstraintNode.h"
 #include "Dataflow/ChaosFleshVisualizeFiberFieldNode.h"
 #include "Dataflow/GeometryCollectionAppendCollectionTransformNode.h"
+#include "Dataflow/DataflowCategoryRegistry.h"
 #include "Modules/ModuleManager.h"
 
 #define LOCTEXT_NAMESPACE "ChaosFleshNodes"
@@ -82,6 +83,9 @@ void IChaosFleshNodesPlugin::StartupModule()
 	DATAFLOW_NODE_REGISTER_CREATION_FACTORY(FVisualizeFiberFieldNode);
 	DATAFLOW_NODE_REGISTER_CREATION_FACTORY(FVisualizePositionTargetsNode);
 	DATAFLOW_NODE_REGISTER_CREATION_FACTORY(FVisualizeKinematicFacesNode);
+
+	// register node category for flesh asset
+	UE_DATAFLOW_REGISTER_CATEGORY_FORASSET_TYPE("Flesh", UFleshAsset);
 }
 
 void IChaosFleshNodesPlugin::ShutdownModule()
diff --git a/Engine/Plugins/Experimental/Dataflow/Source/DataflowEditor/Private/Dataflow/DataflowGraphEditor.cpp b/Engine/Plugins/Experimental/Dataflow/Source/DataflowEditor/Private/Dataflow/DataflowGraphEditor.cpp
--- a/Engine/Plugins/Experimental/Dataflow/Source/DataflowEditor/Private/Dataflow/DataflowGraphEditor.cpp
+++ b/Engine/Plugins/Experimental/Dataflow/Source/DataflowEditor/Private/Dataflow/DataflowGraphEditor.cpp
@@ -346,7 +346,57 @@ FActionMenuContent SDataflowGraphEditor::OnCreateActionMenu(UEdGraph* Graph, con
 		.OnClosedCallback(OnClosed)
 		;
 
-	return FActionMenuContent(ActionMenu, ActionMenu->GetFilterTextBox());
+	TSharedRef<SWidget> ActionMenuWithOptions =
+		SNew(SVerticalBox)
+		+ SVerticalBox::Slot()
+		.Padding(2)
+		.AutoHeight()
+		[
+			SNew(SHorizontalBox)
+
+			+ SHorizontalBox::Slot()
+			.AutoWidth()
+			.HAlign(HAlign_Right)
+			.AutoWidth()
+			.Padding(2.f)
+			[
+				SNew(SCheckBox)
+				.OnCheckStateChanged(this, &SDataflowGraphEditor::OnActionMenuFilterByAssetTypeChanged, ActionMenu.ToWeakPtr())
+				.IsChecked(this, &SDataflowGraphEditor::IsActionMenuFilterByAssetTypeChecked)
+			]
+
+			+ SHorizontalBox::Slot()
+			.HAlign(HAlign_Right)
+			.AutoWidth()
+			.Padding(2.f)
+			[
+				SNew(STextBlock)
+				.Text(LOCTEXT("DataflowActionMenu_FilterByAssetType", "Filter by asset type"))
+			]
+		]
+		+ SVerticalBox::Slot()
+		.Padding(2)
+		.AutoHeight()
+		[
+			ActionMenu
+		];
+
+	return FActionMenuContent(ActionMenuWithOptions, ActionMenu->GetFilterTextBox());
+}
+
+void SDataflowGraphEditor::OnActionMenuFilterByAssetTypeChanged(ECheckBoxState NewState, const TWeakPtr<SGraphEditorActionMenu> WeakActionMenu)
+{
+	bFilterActionMenyByAssetType = (NewState == ECheckBoxState::Checked);
+
+	if (TSharedPtr<SGraphEditorActionMenu> ActionMenu = WeakActionMenu.Pin())
+	{
+		ActionMenu->RefreshAllActions();
+	}
+}
+
+ECheckBoxState SDataflowGraphEditor::IsActionMenuFilterByAssetTypeChecked() const
+{
+	return bFilterActionMenyByAssetType ? ECheckBoxState::Checked : ECheckBoxState::Unchecked;
 }
 
 TSharedPtr<UE::Dataflow::FContext> SDataflowGraphEditor::GetDataflowContext() const
diff --git a/Engine/Plugins/Experimental/Dataflow/Source/DataflowEditor/Private/Dataflow/DataflowSchema.cpp b/Engine/Plugins/Experimental/Dataflow/Source/DataflowEditor/Private/Dataflow/DataflowSchema.cpp
--- a/Engine/Plugins/Experimental/Dataflow/Source/DataflowEditor/Private/Dataflow/DataflowSchema.cpp
+++ b/Engine/Plugins/Experimental/Dataflow/Source/DataflowEditor/Private/Dataflow/DataflowSchema.cpp
@@ -6,6 +6,7 @@
 #include "Dataflow/DataflowEdNode.h"
 #include "Dataflow/DataflowEditorCommands.h"
 #include "Dataflow/DataflowEdNode.h"
+#include "Dataflow/DataflowGraphEditor.h"
 #include "Dataflow/DataflowGraphSchemaAction.h"
 #include "Dataflow/DataflowNodeFactory.h"
 #include "Dataflow/DataflowAssetEditUtils.h"
@@ -19,6 +20,7 @@
 #include "Dataflow/DataflowSettings.h"
 #include "Dataflow/DataflowObject.h"
 #include "Dataflow/DataflowSubGraph.h"
+#include "Dataflow/DataflowCategoryRegistry.h"
 #include "ScopedTransaction.h"
 
 #include UE_INLINE_GENERATED_CPP_BY_NAME(DataflowSchema)
@@ -162,9 +164,54 @@ bool UDataflowSchema::CanPinBeConnectedToNode(const UEdGraphPin* Pin, const UE::
 	return false;
 }
 
+FName UDataflowSchema::GetEditedAssetType()
+{
+	if (TSharedPtr<SDataflowGraphEditor> GraphEditor = SDataflowGraphEditor::GetLastActionMenuGraphEditor().Pin())
+	{
+		if (TSharedPtr<UE::Dataflow::FContext> DataflowContext = GraphEditor->GetDataflowContext())
+		{
+			if (const UE::Dataflow::FEngineContext* EngineContext = DataflowContext->AsType<UE::Dataflow::FEngineContext>())
+			{
+				const FName AssetType = EngineContext->Owner.GetClass()->GetFName();
+				if (AssetType != UDataflow::StaticClass()->GetFName())
+				{
+					return AssetType;
+				}
+			}
+		}
+	}
+	return NAME_None;
+}
+
+bool UDataflowSchema::IsCategorySupported(FName NodeCategory, FName AssetType)
+{
+	bool bFilteringByAssettypeEnable = true;
+	if (TSharedPtr<SDataflowGraphEditor> GraphEditor = SDataflowGraphEditor::GetLastActionMenuGraphEditor().Pin())
+	{
+		bFilteringByAssettypeEnable &= GraphEditor->GetFilterActionMenyByAssetType();
+	}
+	if (!bFilteringByAssettypeEnable)
+	{
+		return true;
+	}
+	if (AssetType.IsNone())
+	{
+		return true;
+	}
+	FString NodeCategoryStr = NodeCategory.ToString();
+	int32 SeparatorPos = NodeCategoryStr.Find("|");
+	if (SeparatorPos != INDEX_NONE)
+	{
+		NodeCategoryStr.LeftInline(SeparatorPos, EAllowShrinking::No);
+		return UE::Dataflow::FCategoryRegistry::Get().IsCategoryForAssetType(FName(NodeCategoryStr), AssetType);
+	}
+	return UE::Dataflow::FCategoryRegistry::Get().IsCategoryForAssetType(NodeCategory, AssetType);
+}
 
 void UDataflowSchema::GetGraphContextActions(FGraphContextMenuBuilder& ContextMenuBuilder) const
 {
+	const FName AssetType = UDataflowSchema::GetEditedAssetType();
+
 	if (UDataflow* DataflowAsset = UDataflow::GetDataflowAssetFromEdGraph(const_cast<UEdGraph*>(ContextMenuBuilder.CurrentGraph)))
 	{
 		const bool bDisplayVariable = (ContextMenuBuilder.FromPin == nullptr || ContextMenuBuilder.FromPin->Direction == EEdGraphPinDirection::EGPD_Input);
@@ -205,13 +252,16 @@ void UDataflowSchema::GetGraphContextActions(FGraphContextMenuBuilder& ContextMe
 	{
 		for (UE::Dataflow::FFactoryParameters NodeParameters : Factory->RegisteredParameters())
 		{
-			// contextual filtering ( if pin is null always return true )
-			if (UDataflowSchema::CanPinBeConnectedToNode(ContextMenuBuilder.FromPin, NodeParameters))
+			if (UDataflowSchema::IsCategorySupported(NodeParameters.Category, AssetType))
 			{
-				if (TSharedPtr<FAssetSchemaAction_Dataflow_CreateNode_DataflowEdNode> Action =
-					FAssetSchemaAction_Dataflow_CreateNode_DataflowEdNode::CreateAction(ContextMenuBuilder.CurrentGraph, NodeParameters.TypeName, NodeParameters.DisplayName))
+				// contextual filtering ( if pin is null always return true )
+				if (UDataflowSchema::CanPinBeConnectedToNode(ContextMenuBuilder.FromPin, NodeParameters))
 				{
-					ContextMenuBuilder.AddAction(Action);
+					if (TSharedPtr<FAssetSchemaAction_Dataflow_CreateNode_DataflowEdNode> Action =
+						FAssetSchemaAction_Dataflow_CreateNode_DataflowEdNode::CreateAction(ContextMenuBuilder.CurrentGraph, NodeParameters.TypeName, NodeParameters.DisplayName))
+					{
+						ContextMenuBuilder.AddAction(Action);
+					}
 				}
 			}
 		}
diff --git a/Engine/Plugins/Experimental/Dataflow/Source/DataflowEditor/Public/Dataflow/DataflowGraphEditor.h b/Engine/Plugins/Experimental/Dataflow/Source/DataflowEditor/Public/Dataflow/DataflowGraphEditor.h
--- a/Engine/Plugins/Experimental/Dataflow/Source/DataflowEditor/Public/Dataflow/DataflowGraphEditor.h
+++ b/Engine/Plugins/Experimental/Dataflow/Source/DataflowEditor/Public/Dataflow/DataflowGraphEditor.h
@@ -16,12 +16,14 @@
 #include "SGraphNode.h"
 #include "Widgets/DeclarativeSyntaxSupport.h"
 #include "Widgets/Layout/SConstraintCanvas.h"
+#include "Widgets/Input/SCheckBox.h"
 #include "Widgets/SWidget.h"
 
 
 class FDataflowEditorToolkit;
 class FDataflowGraphEditorNodeFactory;
 class FDataflowSNodeFactory;
+class SGraphEditorActionMenu;
 class UDataflow;
 class UDataflowEditor;
 struct FDataflowConnection;
@@ -167,6 +169,8 @@ public:
 
 	static const TWeakPtr<SDataflowGraphEditor>& GetLastActionMenuGraphEditor() { return LastActionMenuGraphEditor; }
 
+	bool GetFilterActionMenyByAssetType() const { return bFilterActionMenyByAssetType; }
+
 private:
 	/** Add an additional option pin to all selected Dataflow nodes for those that overrides the AddPin function. */
 	void OnAddOptionPin();
@@ -201,6 +205,8 @@ private:
 	FDataflowEditorCommands::FGraphEvaluationCallback EvaluateGraphCallback;
 
 	FActionMenuContent OnCreateActionMenu(UEdGraph* Graph, const FVector2f& Position, const TArray<UEdGraphPin*>& DraggedPins, bool bAutoExpandActionMenu, SGraphEditor::FActionMenuClosed OnClosed);
+	void OnActionMenuFilterByAssetTypeChanged(ECheckBoxState NewState, const TWeakPtr<SGraphEditorActionMenu> WeakActionMenu);
+	ECheckBoxState IsActionMenuFilterByAssetTypeChecked() const;
 
 	/** The asset that ownes this dataflow graph */
 	TWeakObjectPtr<UObject> AssetOwner;
@@ -232,6 +238,8 @@ private:
 	bool LeftAltKeyDown = false;
 	bool RightAltKeyDown = false;
 
+	bool bFilterActionMenyByAssetType = true;
+
 	FDelegateHandle CVarChangedDelegateHandle;
 	TSharedPtr<SWidget> MessageBar;
 	TSharedPtr<SWidget> EvaluationProgressBar;
diff --git a/Engine/Plugins/Experimental/Dataflow/Source/DataflowEditor/Public/Dataflow/DataflowSchema.h b/Engine/Plugins/Experimental/Dataflow/Source/DataflowEditor/Public/Dataflow/DataflowSchema.h
--- a/Engine/Plugins/Experimental/Dataflow/Source/DataflowEditor/Public/Dataflow/DataflowSchema.h
+++ b/Engine/Plugins/Experimental/Dataflow/Source/DataflowEditor/Public/Dataflow/DataflowSchema.h
@@ -52,7 +52,8 @@ public:
 private:
 	static bool CanPinBeConnectedToNode(const UEdGraphPin* Pin, const UE::Dataflow::FFactoryParameters& NodeParameters);
 	static bool CanConnectPins(const UEdGraphPin& OutputPin, const UEdGraphPin& InputPin);
-	
+	static bool IsCategorySupported(FName NodeCategory, FName AssetType);
+	static FName GetEditedAssetType();
 	static UE::Dataflow::FPin::EDirection GetDirectionFromPinDirection(EEdGraphPinDirection InPinDirection);
 
 	// used by SupportsDropPinOnNode to know which original connection is being dropped
diff --git a/Engine/Plugins/Experimental/GeometryCollectionPlugin/Source/GeometryCollectionNodes/Private/Dataflow/GeometryCollectionNodesPlugin.cpp b/Engine/Plugins/Experimental/GeometryCollectionPlugin/Source/GeometryCollectionNodes/Private/Dataflow/GeometryCollectionNodesPlugin.cpp
--- a/Engine/Plugins/Experimental/GeometryCollectionPlugin/Source/GeometryCollectionNodes/Private/Dataflow/GeometryCollectionNodesPlugin.cpp
+++ b/Engine/Plugins/Experimental/GeometryCollectionPlugin/Source/GeometryCollectionNodes/Private/Dataflow/GeometryCollectionNodesPlugin.cpp
@@ -5,6 +5,7 @@
 #include "CoreMinimal.h"
 
 #include "Dataflow/DataflowNodeFactory.h"
+#include "Dataflow/DataflowCategoryRegistry.h"
 #include "Dataflow/CreateColorArrayFromFloatArrayNode.h"
 #include "Dataflow/GeometryCollectionArrayNodes.h"
 #include "Dataflow/GeometryCollectionAssetNodes.h"
@@ -38,6 +39,8 @@
 #include "Dataflow/SetVertexColorFromVertexIndicesNode.h"
 #include "Dataflow/SetVertexColorFromVertexSelectionNode.h"
 #include "Dataflow/DataflowNodeAndPinTypeColors.h"
+#include "Dataflow/DataflowCategoryRegistry.h"
+#include "GeometryCollection/GeometryCollectionObject.h"
 
 #define LOCTEXT_NAMESPACE "DataflowNodes"
 
@@ -84,6 +87,9 @@ PRAGMA_ENABLE_DEPRECATION_WARNINGS
 
 	// Register Dataflow node and PinType colors
 	UE::Dataflow::RegisterGeometryCollectionNodesColors();
+
+	// register category for GeometryCollection asset
+	UE_DATAFLOW_REGISTER_CATEGORY_FORASSET_TYPE("GeometryCollection", UGeometryCollection);
 }
 
 void IGeometryCollectionNodesPlugin::ShutdownModule()
diff --git a/Engine/Plugins/Runtime/HairStrands/Source/HairStrandsDataflow/Private/HairStrandsDataflowModule.cpp b/Engine/Plugins/Runtime/HairStrands/Source/HairStrandsDataflow/Private/HairStrandsDataflowModule.cpp
--- a/Engine/Plugins/Runtime/HairStrands/Source/HairStrandsDataflow/Private/HairStrandsDataflowModule.cpp
+++ b/Engine/Plugins/Runtime/HairStrands/Source/HairStrandsDataflow/Private/HairStrandsDataflowModule.cpp
@@ -1,4 +1,4 @@
-﻿// Copyright Epic Games, Inc. All Rights Reserved.
+// Copyright Epic Games, Inc. All Rights Reserved.
 
 #include "HairStrandsDataflowModule.h"
 
@@ -15,6 +15,7 @@
 #include "GroomDataflowRendering.h"
 #include "GroomDataflowVisualization.h"
 #include "Dataflow/DataflowCollectionAddScalarVertexPropertyNode.h"
+#include "Dataflow/DataflowCategoryRegistry.h"
 #include "Modules/ModuleManager.h"
 
 #define LOCTEXT_NAMESPACE "HairStrandsDataflow"
@@ -65,6 +66,9 @@ void FHairStrandsDataflowModule::StartupModule()
 
 	UE::Dataflow::FDataflowSimulationVisualizationRegistry::GetInstance().RegisterVisualization(
 		MakeUnique<UE::Groom::FGroomDataflowSimulationVisualization>());
+
+	// register node category for groom asset
+	UE_DATAFLOW_REGISTER_CATEGORY_FORASSET_TYPE("Groom", UGroomAsset);
 }
 
 void FHairStrandsDataflowModule::ShutdownModule()
diff --git a/Engine/Source/Editor/GraphEditor/Private/SGraphEditorActionMenu.cpp b/Engine/Source/Editor/GraphEditor/Private/SGraphEditorActionMenu.cpp
--- a/Engine/Source/Editor/GraphEditor/Private/SGraphEditorActionMenu.cpp
+++ b/Engine/Source/Editor/GraphEditor/Private/SGraphEditorActionMenu.cpp
@@ -49,6 +49,14 @@ void SGraphEditorActionMenu::Construct( const FArguments& InArgs )
 	);
 }
 
+void SGraphEditorActionMenu::RefreshAllActions()
+{
+	if (GraphActionMenu)
+	{
+		GraphActionMenu->RefreshAllActions(true, false);
+	}
+}
+
 void SGraphEditorActionMenu::CollectAllActions(FGraphActionListBuilderBase& OutAllActions)
 {
 	// Build up the context object
diff --git a/Engine/Source/Editor/GraphEditor/Public/SGraphEditorActionMenu.h b/Engine/Source/Editor/GraphEditor/Public/SGraphEditorActionMenu.h
--- a/Engine/Source/Editor/GraphEditor/Public/SGraphEditorActionMenu.h
+++ b/Engine/Source/Editor/GraphEditor/Public/SGraphEditorActionMenu.h
@@ -45,6 +45,8 @@ public:
 
 	~SGraphEditorActionMenu();
 
+	void RefreshAllActions();
+
 	TSharedRef<SEditableTextBox> GetFilterTextBox();
 
 protected:
diff --git a/Engine/Source/Runtime/Experimental/Dataflow/Core/Private/Dataflow/DataflowCategoryRegistry.cpp b/Engine/Source/Runtime/Experimental/Dataflow/Core/Private/Dataflow/DataflowCategoryRegistry.cpp
new file mode 100644
--- /dev/null
+++ b/Engine/Source/Runtime/Experimental/Dataflow/Core/Private/Dataflow/DataflowCategoryRegistry.cpp
@@ -0,0 +1,32 @@
+// Copyright Epic Games, Inc. All Rights Reserved.
+
+#include "Dataflow/DataflowCategoryRegistry.h"
+#include "Misc/LazySingleton.h"
+
+namespace UE::Dataflow
+{
+	FCategoryRegistry& FCategoryRegistry::Get()
+	{
+		return TLazySingleton<FCategoryRegistry>::Get();
+	}
+
+	void FCategoryRegistry::TearDown()
+	{
+		return TLazySingleton<FCategoryRegistry>::TearDown();
+	}
+
+	void FCategoryRegistry::RegisterCategoryForAssetType(const FCategoryName Category, const FAssetType AssetType)
+	{
+		AssetTypesByCategory.FindOrAdd(Category).AddUnique(AssetType);
+	}
+
+	bool FCategoryRegistry::IsCategoryForAssetType(const FCategoryName Category, const FAssetType AssetType) const
+	{
+		if (const TArray<FAssetType>* AssetTypes = AssetTypesByCategory.Find(Category))
+		{
+			return AssetTypes->Contains(AssetType);
+		}
+		// category not found so it's a common one 
+		return true;
+	}
+}
\ No newline at end of file
diff --git a/Engine/Source/Runtime/Experimental/Dataflow/Core/Public/Dataflow/DataflowCategoryRegistry.h b/Engine/Source/Runtime/Experimental/Dataflow/Core/Public/Dataflow/DataflowCategoryRegistry.h
new file mode 100644
--- /dev/null
+++ b/Engine/Source/Runtime/Experimental/Dataflow/Core/Public/Dataflow/DataflowCategoryRegistry.h
@@ -0,0 +1,34 @@
+// Copyright Epic Games, Inc. All Rights Reserved.
+
+#pragma once
+
+#include "CoreMinimal.h"
+
+class FLazySingleton;
+
+namespace UE::Dataflow
+{
+	// registry for categories
+	// This is used for node in editor filtering
+	struct FCategoryRegistry
+	{
+	public:
+		static DATAFLOWCORE_API FCategoryRegistry& Get();
+		static DATAFLOWCORE_API void TearDown();
+
+		using FCategoryName = FName;
+		using FAssetType = FName;
+
+		/** register a category name for a specific type of asset */
+		DATAFLOWCORE_API void RegisterCategoryForAssetType(const FCategoryName Category, const FAssetType AssetType);
+		DATAFLOWCORE_API bool IsCategoryForAssetType(const FCategoryName Category, const FAssetType AssetType) const;
+
+	private:
+		TMap<FCategoryName, TArray<FAssetType>> AssetTypesByCategory;
+
+		friend FLazySingleton;
+	};
+}
+
+#define UE_DATAFLOW_REGISTER_CATEGORY_FORASSET_TYPE(CATEGORY_NAME, ASSET_TYPE) \
+	UE::Dataflow::FCategoryRegistry::Get().RegisterCategoryForAssetType(CATEGORY_NAME, ASSET_TYPE::StaticClass()->GetFName());
