dataflow : fix issue where asynchronous execution "evaluating" string is not showing unless the perf data is enabled

diff --git a/Engine/Plugins/Experimental/Dataflow/Source/DataflowEditor/Private/Dataflow/DataflowSNode.cpp b/Engine/Plugins/Experimental/Dataflow/Source/DataflowEditor/Private/Dataflow/DataflowSNode.cpp
--- a/Engine/Plugins/Experimental/Dataflow/Source/DataflowEditor/Private/Dataflow/DataflowSNode.cpp
+++ b/Engine/Plugins/Experimental/Dataflow/Source/DataflowEditor/Private/Dataflow/DataflowSNode.cpp
@@ -243,47 +243,44 @@ TArray<FOverlayWidgetInfo> SDataflowEdNode::GetOverlayWidgets(bool bSelected, co
 
 		if (TSharedPtr<UE::Dataflow::FContext> Context = DataflowInterface->GetDataflowContext())
 		{
-			if (Context->IsPerfDataEnabled())
+			if (TSharedPtr<const FDataflowNode> DataflowNode = DataflowGraphNode->GetDataflowNode())
 			{
-				if (TSharedPtr<const FDataflowNode> DataflowNode = DataflowGraphNode->GetDataflowNode())
+				UE::Dataflow::FContextPerfData::FData PerfData = Context->GetPerfDataForNode(*DataflowNode);
+				if (DataflowNode->IsAsyncEvaluating())
 				{
-					UE::Dataflow::FContextPerfData::FData PerfData = Context->GetPerfDataForNode(*DataflowNode);
-					if (DataflowNode->IsAsyncEvaluating())
-					{
-						constexpr double BlinkPeriod = 2000;
-						const double ElapsedMs = FPlatformTime::ToMilliseconds64(FPlatformTime::Cycles64());
-						const double Fraction = FMath::Fractional(ElapsedMs / BlinkPeriod);
-						const double LerpFactor = FMath::Sin(Fraction * 2.0 * UE_PI) * 0.5 + 0.5;
-						const FLinearColor Color = FMath::Lerp(FLinearColor::Red, FLinearColor::White, LerpFactor);
-						PerfWidget->SetColorAndOpacity(Color);
-						PerfWidget->SetVisibility(EVisibility::Visible);
-						PerfWidget->SetText(LOCTEXT("DataflowNodeEvaluatingMessage", "Evaluating..."));
-					}
-					else if (PerfData.ExclusiveTimeMs > 0 && PerfData.InclusiveTimeMs > 0)
-					{
-						constexpr double FadeTimeMs = 5000;
-						const double MsSinceUpdate = FPlatformTime::ToMilliseconds64(FPlatformTime::Cycles64() - PerfData.LastTimestamp.Value);
-						const double LerpFactor = FMath::Clamp(MsSinceUpdate / FadeTimeMs, 0.0, 1.0);
-						const FLinearColor Color = FMath::Lerp(FLinearColor::Red, FLinearColor::White, LerpFactor);
-						PerfWidget->SetColorAndOpacity(Color);
-						PerfWidget->SetVisibility(EVisibility::Visible);
-						FNumberFormattingOptions FmtOptions;
-						FmtOptions.MaximumFractionalDigits = 2;
-						FmtOptions.MinimumFractionalDigits = 2;
-
-						PerfWidget->SetText(FText::Format(LOCTEXT("DataflowNodePerfDataFormat", "Time: {0} ms"), FText::AsNumber(PerfData.ExclusiveTimeMs, &FmtOptions)));
-					}
-					else
-					{
-						PerfWidget->SetVisibility(EVisibility::Collapsed);
-						PerfWidget->SetText(FText());
-					}
-
-					FOverlayWidgetInfo PerfWidgetInfo;
-					PerfWidgetInfo.OverlayOffset = FVector2f{ 0.0f, -20.0f };
-					PerfWidgetInfo.Widget = PerfWidget;
-					Widgets.Add(PerfWidgetInfo);
+					constexpr double BlinkPeriod = 2000;
+					const double ElapsedMs = FPlatformTime::ToMilliseconds64(FPlatformTime::Cycles64());
+					const double Fraction = FMath::Fractional(ElapsedMs / BlinkPeriod);
+					const double LerpFactor = FMath::Sin(Fraction * 2.0 * UE_PI) * 0.5 + 0.5;
+					const FLinearColor Color = FMath::Lerp(FLinearColor::Red, FLinearColor::White, LerpFactor);
+					PerfWidget->SetColorAndOpacity(Color);
+					PerfWidget->SetVisibility(EVisibility::Visible);
+					PerfWidget->SetText(LOCTEXT("DataflowNodeEvaluatingMessage", "Evaluating..."));
+				}
+				else if (Context->IsPerfDataEnabled() && PerfData.ExclusiveTimeMs > 0 && PerfData.InclusiveTimeMs > 0)
+				{
+					constexpr double FadeTimeMs = 5000;
+					const double MsSinceUpdate = FPlatformTime::ToMilliseconds64(FPlatformTime::Cycles64() - PerfData.LastTimestamp.Value);
+					const double LerpFactor = FMath::Clamp(MsSinceUpdate / FadeTimeMs, 0.0, 1.0);
+					const FLinearColor Color = FMath::Lerp(FLinearColor::Red, FLinearColor::White, LerpFactor);
+					PerfWidget->SetColorAndOpacity(Color);
+					PerfWidget->SetVisibility(EVisibility::Visible);
+					FNumberFormattingOptions FmtOptions;
+					FmtOptions.MaximumFractionalDigits = 2;
+					FmtOptions.MinimumFractionalDigits = 2;
+
+					PerfWidget->SetText(FText::Format(LOCTEXT("DataflowNodePerfDataFormat", "Time: {0} ms"), FText::AsNumber(PerfData.ExclusiveTimeMs, &FmtOptions)));
 				}
+				else
+				{
+					PerfWidget->SetVisibility(EVisibility::Collapsed);
+					PerfWidget->SetText(FText());
+				}
+
+				FOverlayWidgetInfo PerfWidgetInfo;
+				PerfWidgetInfo.OverlayOffset = FVector2f{ 0.0f, -20.0f };
+				PerfWidgetInfo.Widget = PerfWidget;
+				Widgets.Add(PerfWidgetInfo);
 			}
 		}
 
