Add primitive component rendering to dataflow

diff --git a/Engine/Plugins/ChaosClothAssetEditor/Source/ChaosClothAssetDataflowNodes/Private/ChaosClothAsset/ChaosClothAssetDataflowNodesModule.cpp b/Engine/Plugins/ChaosClothAssetEditor/Source/ChaosClothAssetDataflowNodes/Private/ChaosClothAsset/ChaosClothAssetDataflowNodesModule.cpp
--- a/Engine/Plugins/ChaosClothAssetEditor/Source/ChaosClothAssetDataflowNodes/Private/ChaosClothAsset/ChaosClothAssetDataflowNodesModule.cpp
+++ b/Engine/Plugins/ChaosClothAssetEditor/Source/ChaosClothAssetDataflowNodes/Private/ChaosClothAsset/ChaosClothAssetDataflowNodesModule.cpp
@@ -84,6 +84,7 @@
 #include "Dataflow/DataflowNodeFactory.h"
 #include "Dataflow/DataflowRenderingFactory.h"
 #include "Modules/ModuleManager.h"
+#include "GeometryCollection/Facades/CollectionRenderingFacade.h"
 
 namespace UE::Chaos::ClothAsset
 {
@@ -194,7 +195,9 @@ PRAGMA_ENABLE_DEPRECATION_WARNINGS
 				const FName& ViewModeName = ViewMode.GetName();
 				return (ViewModeName == FCloth2DSimViewMode::Name ||
 						ViewModeName == FCloth3DSimViewMode::Name ||
-						ViewModeName == FClothRenderViewMode::Name);
+						ViewModeName == FClothRenderViewMode::Name ||
+						ViewModeName == Dataflow::FDataflowConstruction3DViewMode::Name ||
+						ViewModeName == Dataflow::FDataflowConstruction2DViewMode::Name);
 			}
 
 			virtual void Render(GeometryCollection::Facades::FRenderingFacade& RenderCollection, const UE::Dataflow::FGraphRenderingState& State) override
@@ -220,13 +223,13 @@ PRAGMA_ENABLE_DEPRECATION_WARNINGS
 					TArray<TArray<FVector2f>> UVs;
 					TArray<int32> RenderMaterialIDs;
 
-					if (State.GetViewMode().GetName() == FCloth3DSimViewMode::Name)
+					if (State.GetViewMode().GetName() == FCloth3DSimViewMode::Name || State.GetViewMode().GetName() == Dataflow::FDataflowConstruction3DViewMode::Name)
 					{
 						Vertices = ClothFacade.GetSimPosition3D();
 						Indices = ClothFacade.GetSimIndices3D();
 						Normals = ClothFacade.GetSimNormal();
 					}
-					else if (State.GetViewMode().GetName() == FCloth2DSimViewMode::Name)
+					else if (State.GetViewMode().GetName() == FCloth2DSimViewMode::Name || State.GetViewMode().GetName() == Dataflow::FDataflowConstruction2DViewMode::Name)
 					{
 						//
 						// Flip the Y coordinate to get the desired visualization with our LVT_OrthoXY viewport
diff --git a/Engine/Plugins/Experimental/Dataflow/Source/DataflowEditor/Private/Dataflow/DataflowConstructionScene.cpp b/Engine/Plugins/Experimental/Dataflow/Source/DataflowEditor/Private/Dataflow/DataflowConstructionScene.cpp
--- a/Engine/Plugins/Experimental/Dataflow/Source/DataflowEditor/Private/Dataflow/DataflowConstructionScene.cpp
+++ b/Engine/Plugins/Experimental/Dataflow/Source/DataflowEditor/Private/Dataflow/DataflowConstructionScene.cpp
@@ -13,6 +13,9 @@
 #include "Elements/Framework/EngineElementsLibrary.h"
 #include "Selection.h"
 #include "AssetViewerSettings.h"
+#include "Dataflow/DataflowPrimitiveNode.h"
+#include "GeometryCollection/Facades/CollectionRenderingFacade.h"
+#include "Selection/GeometrySelector.h"
 
 #define LOCTEXT_NAMESPACE "FDataflowConstructionScene"
 
@@ -29,7 +32,7 @@ FDataflowConstructionScene::FDataflowConstructionScene(FPreviewScene::Constructi
 FDataflowConstructionScene::~FDataflowConstructionScene()
 {
 	ResetWireframeMeshElementsVisualizer();
-	ResetDynamicMeshComponents();
+	ResetSceneComponents();
 }
 
 TArray<TObjectPtr<UDynamicMeshComponent>> FDataflowConstructionScene::GetDynamicMeshComponents() const
@@ -163,10 +166,7 @@ void FDataflowConstructionScene::UpdateDynamicMeshComponents()
 	// The preview scene for the construction view will be
 	// cleared and rebuilt from scratch. This will generate a 
 	// list of UPrimitiveComponents for rendering.
-	ResetDynamicMeshComponents();
-
-	AddSceneObject(RootSceneActor, true);
-
+	
 	if (const TObjectPtr<UDataflowBaseContent>& EditorContent = GetEditorContent())
 	{
 		const TObjectPtr<UDataflow>& DataflowAsset = EditorContent->GetDataflowAsset();
@@ -178,59 +178,62 @@ void FDataflowConstructionScene::UpdateDynamicMeshComponents()
 				if (Target)
 				{
 					TSharedPtr<FManagedArrayCollection> RenderCollection(new FManagedArrayCollection);
-					GeometryCollection::Facades::FRenderingFacade Facade(*RenderCollection);
-					Facade.DefineSchema();
-
-					UE::Dataflow::RenderNodeOutput(Facade, *Target, *EditorContent);
+					GeometryCollection::Facades::FRenderingFacade RenderingFacade(*RenderCollection);
+					RenderingFacade.DefineSchema();
 
-					const int32 NumGeometry = Facade.NumGeometry();
+					const bool bHasPrimitves = UE::Dataflow::RenderNodeOutput(RenderingFacade, *Target, *EditorContent);
 
-					DebugMesh.Reset();
-
-					int32 VertexIndex = 0, FaceIndex = 0;
-
-					const TManagedArray<int32>& VertexStart = Facade.GetVertexStart();
-					const TManagedArray<int32>& VertexCount = Facade.GetVertexCount();
-					const TManagedArray<int32>& FaceStart = Facade.GetIndicesStart();
-					const TManagedArray<int32>& FaceCount = Facade.GetIndicesCount();
-
-					for (int32 MeshIndex = 0; MeshIndex < NumGeometry; ++MeshIndex)
+					if (Target == EditorContent->GetSelectedNode())
 					{
-						FDynamicMesh3 DynamicMesh;
-						UE::Dataflow::Conversion::RenderingFacadeToDynamicMesh(Facade, MeshIndex, DynamicMesh);
+						EditorContent->SetRenderCollection(RenderCollection);
+					}
+					if(!bHasPrimitves)
+					{
+						const int32 NumGeometry = RenderingFacade.NumGeometry();
+	
+						DebugMesh.Reset();
+						int32 VertexIndex = 0, FaceIndex = 0;
 
-						if (DynamicMesh.VertexCount())
+						const TManagedArray<int32>& VertexStart = RenderingFacade.GetVertexStart();
+						const TManagedArray<int32>& VertexCount = RenderingFacade.GetVertexCount();
+						const TManagedArray<int32>& FaceStart = RenderingFacade.GetIndicesStart();
+						const TManagedArray<int32>& FaceCount = RenderingFacade.GetIndicesCount();
+
+						for (int32 MeshIndex = 0; MeshIndex < NumGeometry; ++MeshIndex)
 						{
-							if (Target == EditorContent->GetSelectedNode())
+							FDynamicMesh3 DynamicMesh;
+							UE::Dataflow::Conversion::RenderingFacadeToDynamicMesh(RenderingFacade, MeshIndex, DynamicMesh);
+
+							if (DynamicMesh.VertexCount())
 							{
-								EditorContent->SetRenderCollection(RenderCollection);
-							}
-							const FString MeshName = Facade.GetGeometryName()[MeshIndex];
+								const FString MeshName = RenderingFacade.GetGeometryName()[MeshIndex];
 
-							const TManagedArray<FString>& MaterialPaths = Facade.GetMaterialPaths();
-							const int32 MaterialStart = Facade.GetMaterialStart()[MeshIndex];
-							const int32 MaterialCount = Facade.GetMaterialCount()[MeshIndex];
+								const TManagedArray<FString>& MaterialPaths = RenderingFacade.GetMaterialPaths();
+								const int32 MaterialStart = RenderingFacade.GetMaterialStart()[MeshIndex];
+								const int32 MaterialCount = RenderingFacade.GetMaterialCount()[MeshIndex];
 
-							TArray<UMaterialInterface*> Materials;
-							for (int32 MaterialIndex = MaterialStart; MaterialIndex < MaterialStart + MaterialCount; ++MaterialIndex)
-							{
-								const FString& Path = MaterialPaths[MaterialIndex];
-								UMaterialInterface* const Material = LoadObject<UMaterialInterface>(nullptr, *Path);
-								Materials.Add(Material);
-							}
+								TArray<UMaterialInterface*> Materials;
+								for (int32 MaterialIndex = MaterialStart; MaterialIndex < MaterialStart + MaterialCount; ++MaterialIndex)
+								{
+									const FString& Path = MaterialPaths[MaterialIndex];
+									UMaterialInterface* const Material = LoadObject<UMaterialInterface>(nullptr, *Path);
+									Materials.Add(Material);
+								}
 
-							AddDynamicMeshComponent({Target, MeshIndex}, MeshName, MoveTemp(DynamicMesh), Materials);
+								AddDynamicMeshComponent({Target, MeshIndex}, MeshName, MoveTemp(DynamicMesh), Materials);
 
-							for (int32 Idx = 0; Idx < VertexCount[MeshIndex]; ++Idx)
-							{
-								DebugMesh.VertexMap.Add(VertexIndex, VertexStart[MeshIndex] + Idx);
-								VertexIndex++;
-							}
 
-							for (int32 Idx = 0; Idx < FaceCount[MeshIndex]; ++Idx)
-							{
-								DebugMesh.FaceMap.Add(FaceIndex, FaceStart[MeshIndex] + Idx);
-								FaceIndex++;
+								for (int32 Idx = 0; Idx < VertexCount[MeshIndex]; ++Idx)
+								{
+									DebugMesh.VertexMap.Add(VertexIndex, VertexStart[MeshIndex] + Idx);
+									VertexIndex++;
+								}
+
+								for (int32 Idx = 0; Idx < FaceCount[MeshIndex]; ++Idx)
+								{
+									DebugMesh.FaceMap.Add(FaceIndex, FaceStart[MeshIndex] + Idx);
+									FaceIndex++;
+								}
 							}
 						}
 					}
@@ -262,38 +265,42 @@ void FDataflowConstructionScene::UpdateDynamicMeshComponents()
 				if (Target)
 				{
 					TSharedPtr<FManagedArrayCollection> RenderCollection(new FManagedArrayCollection);
-					GeometryCollection::Facades::FRenderingFacade Facade(*RenderCollection);
-					Facade.DefineSchema();
+					GeometryCollection::Facades::FRenderingFacade RenderingFacade(*RenderCollection);
+					RenderingFacade.DefineSchema();
 
-					UE::Dataflow::RenderNodeOutput(Facade, *Target, *EditorContent);
+					const bool bHasPrimitves = UE::Dataflow::RenderNodeOutput(RenderingFacade, *Target, *EditorContent);
 
-					const int32 NumGeometry = Facade.NumGeometry();
+					if (Target == EditorContent->GetSelectedNode())
 
-					for (int32 MeshIndex = 0; MeshIndex < NumGeometry; ++MeshIndex)
 					{
-						FDataflowRenderKey WireframeDynamicMeshKey{ Target, MeshIndex };
+						EditorContent->SetRenderCollection(RenderCollection);
+					}
 
-						if (DynamicMeshComponents.Contains(WireframeDynamicMeshKey))
-						{
-							UDynamicMeshComponent* const ExistingMeshComponent = DynamicMeshComponents[WireframeDynamicMeshKey];
-							MeshComponentsForWireframeRendering.Add(ExistingMeshComponent);
-						}
-						else
+					if(!bHasPrimitves)
+					{
+						const int32 NumGeometry = RenderingFacade.NumGeometry();
+						for (int32 MeshIndex = 0; MeshIndex < NumGeometry; ++MeshIndex)
 						{
-							FDynamicMesh3 DynamicMesh;
-							UE::Dataflow::Conversion::RenderingFacadeToDynamicMesh(Facade, MeshIndex, DynamicMesh);
+							FDataflowRenderKey WireframeDynamicMeshKey{ Target, MeshIndex };
 
-							if (DynamicMesh.VertexCount())
+							if (DynamicMeshComponents.Contains(WireframeDynamicMeshKey))
 							{
-								if (Target == EditorContent->GetSelectedNode())
+								UDynamicMeshComponent* const ExistingMeshComponent = DynamicMeshComponents[WireframeDynamicMeshKey];
+								MeshComponentsForWireframeRendering.Add(ExistingMeshComponent);
+							}
+							else
+							{
+								FDynamicMesh3 DynamicMesh;
+								UE::Dataflow::Conversion::RenderingFacadeToDynamicMesh(RenderingFacade, MeshIndex, DynamicMesh);
+
+								if (DynamicMesh.VertexCount())
 								{
-									EditorContent->SetRenderCollection(RenderCollection);
+									const FString MeshName = RenderingFacade.GetGeometryName()[MeshIndex];
+									const FString UniqueObjectName = MakeUniqueObjectName(RootSceneActor, UDataflowEditorCollectionComponent::StaticClass(), FName(MeshName)).ToString();
+									UDynamicMeshComponent* const NewDynamicMeshComponent = AddDynamicMeshComponent(WireframeDynamicMeshKey, UniqueObjectName, MoveTemp(DynamicMesh), {});
+									NewDynamicMeshComponent->SetVisibility(false);
+									MeshComponentsForWireframeRendering.Add(NewDynamicMeshComponent);
 								}
-								const FString MeshName = Facade.GetGeometryName()[MeshIndex];
-								const FString UniqueObjectName = MakeUniqueObjectName(RootSceneActor, UDataflowEditorCollectionComponent::StaticClass(), FName(MeshName)).ToString();
-								UDynamicMeshComponent* const NewDynamicMeshComponent = AddDynamicMeshComponent(WireframeDynamicMeshKey, UniqueObjectName, MoveTemp(DynamicMesh), {});
-								NewDynamicMeshComponent->SetVisibility(false);
-								MeshComponentsForWireframeRendering.Add(NewDynamicMeshComponent);
 							}
 						}
 					}
@@ -329,24 +336,65 @@ void FDataflowConstructionScene::UpdateDynamicMeshComponents()
 	bPreviewSceneDirty = true;
 }
 
-void FDataflowConstructionScene::ResetDynamicMeshComponents()
+void FDataflowConstructionScene::UpdatePrimitiveComponents()
 {
-	USelection* SelectedComponents = DataflowModeManager->GetSelectedComponents();
-	for (FRenderElement RenderElement : DynamicMeshComponents)
+	if (const TObjectPtr<UDataflowBaseContent>& EditorContent = GetEditorContent())
 	{
-		TObjectPtr<UDynamicMeshComponent>& DynamicMeshComponent = RenderElement.Value;
+		if (const TSharedPtr<UE::Dataflow::FEngineContext> DataflowContext = EditorContent->GetDataflowContext())
+		{
+			if (UDataflowEdNode* SelectedNode = EditorContent->GetSelectedNode())
+			{
+				if (TSharedPtr<FDataflowNode> SelectedDataflowNode = SelectedNode->GetDataflowNode())
+				{
+					if(SelectedDataflowNode->IsA(FDataflowPrimitiveNode::StaticType()))
+					{
+						StaticCastSharedPtr<FDataflowPrimitiveNode>(SelectedDataflowNode)->AddPrimitiveComponents( 
+							EditorContent->GetRenderCollection(), DataflowContext->Owner, RootSceneActor, PrimitiveComponents);
+					}
+				}
+			}
+		}
+	}
+	for(TObjectPtr<UPrimitiveComponent> PrimitiveComponent : PrimitiveComponents)
+	{
+		PrimitiveComponent->SelectionOverrideDelegate = UPrimitiveComponent::FSelectionOverride::CreateRaw(this, &FDataflowPreviewSceneBase::IsComponentSelected);
+		PrimitiveComponent->UpdateBounds();
+
+		AddComponent(PrimitiveComponent, PrimitiveComponent->GetRelativeTransform());	
+		AddSceneObject(PrimitiveComponent, true);
+	}
+	bPreviewSceneDirty = true;
+}
 
-		DynamicMeshComponent->SelectionOverrideDelegate.Unbind();
-		if (SelectedComponents->IsSelected(DynamicMeshComponent))
+void FDataflowConstructionScene::RemoveSceneComponent(USelection* SelectedComponents, UPrimitiveComponent* PrimitiveComponent)
+{
+	if(PrimitiveComponent)
+	{
+		PrimitiveComponent->SelectionOverrideDelegate.Unbind();
+		if (SelectedComponents->IsSelected(PrimitiveComponent))
 		{
-			SelectedComponents->Deselect(DynamicMeshComponent);
-			DynamicMeshComponent->PushSelectionToProxy();
+			SelectedComponents->Deselect(PrimitiveComponent);
+			PrimitiveComponent->PushSelectionToProxy();
 		}
-		RemoveSceneObject(DynamicMeshComponent);
-		RemoveComponent(DynamicMeshComponent);
-		DynamicMeshComponent->DestroyComponent();
+		RemoveSceneObject(PrimitiveComponent);
+		RemoveComponent(PrimitiveComponent);
+		PrimitiveComponent->DestroyComponent();
+	}
+}
+
+void FDataflowConstructionScene::ResetSceneComponents()
+{
+	USelection* SelectedComponents = DataflowModeManager->GetSelectedComponents();
+	for (FRenderElement RenderElement : DynamicMeshComponents)
+	{
+		RemoveSceneComponent(SelectedComponents, RenderElement.Value);
+	}
+	for (TObjectPtr<UPrimitiveComponent> PrimitiveComponent : PrimitiveComponents)
+	{
+		RemoveSceneComponent(SelectedComponents, PrimitiveComponent);
 	}
 	DynamicMeshComponents.Reset();
+	PrimitiveComponents.Reset();
 	bPreviewSceneDirty = true;
 
 	RemoveSceneObject(RootSceneActor);
@@ -498,13 +546,16 @@ void FDataflowConstructionScene::ResetConstructionScene()
 
 	ResetWireframeMeshElementsVisualizer();
 
-	ResetDynamicMeshComponents();
+	ResetSceneComponents();
 }
 
 void FDataflowConstructionScene::UpdateConstructionScene()
 {
 	ResetConstructionScene();
 
+	// Add root actor to TEDS
+	AddSceneObject(RootSceneActor, true);
+
 	// The preview scene for the construction view will be
 	// cleared and rebuilt from scratch. This will generate a 
 	// list of UPrimitiveComponents for rendering.
@@ -513,6 +564,9 @@ void FDataflowConstructionScene::UpdateConstructionScene()
 	// Attach a wireframe renderer to the DynamicMeshComponents
 	UpdateWireframeMeshElementsVisualizer();
 
+	// Update all the primitive components potentially added by the selected node
+	UpdatePrimitiveComponents();
+
 	for (const UDynamicMeshComponent* const DynamicMeshComponent : MeshComponentsForWireframeRendering)
 	{
 		WireframeElements[DynamicMeshComponent]->Settings->bVisible = true;
diff --git a/Engine/Plugins/Experimental/Dataflow/Source/DataflowEditor/Private/Dataflow/DataflowConstructionViewportClient.cpp b/Engine/Plugins/Experimental/Dataflow/Source/DataflowEditor/Private/Dataflow/DataflowConstructionViewportClient.cpp
--- a/Engine/Plugins/Experimental/Dataflow/Source/DataflowEditor/Private/Dataflow/DataflowConstructionViewportClient.cpp
+++ b/Engine/Plugins/Experimental/Dataflow/Source/DataflowEditor/Private/Dataflow/DataflowConstructionViewportClient.cpp
@@ -162,24 +162,13 @@ void FDataflowConstructionViewportClient::OnViewportClicked(HHitProxy* HitProxy)
 						{
 							if (SelectedComponents && SelectedComponents->Num() == 1)
 							{
-								if (UDataflowEditorCollectionComponent* CollectionComponent =
-									Cast< UDataflowEditorCollectionComponent>(SelectedComponents->GetSelectedObject(0)))
-								{
-									if (CollectionComponent->Node == SelectedNode)
-									{
-										// Start the corresponding tool
-										DataflowMode->StartToolForSelectedNode(SelectedNode);
-
-										return CollectionComponent;
-									}
-								}
+								DataflowMode->StartToolForSelectedNode(SelectedNode);
 							}
 						}
 					}
 				}
 			}
 		}
-		return (UDataflowEditorCollectionComponent*)nullptr;
 	};
 
 	auto UpdateSelectedComponentInViewport = [&](USelection* SelectedComponents)
@@ -235,15 +224,6 @@ void FDataflowConstructionViewportClient::OnViewportClicked(HHitProxy* HitProxy)
 		return false;
 	};
 
-	auto IsolateComponent = [&](UDataflowEditorCollectionComponent* SelectedComponent)
-	{
-		if( FDataflowConstructionScene* Scene = static_cast<FDataflowConstructionScene*>(PreviewScene) )
-		{ 
-			Scene->SetVisibility(false);
-			Scene->SetVisibility(true,SelectedComponent);
-		}
-	};
-
 	TArray<UPrimitiveComponent*> CurrentlySelectedComponents;
 	if (!IsInteractiveToolActive())
 	{
diff --git a/Engine/Plugins/Experimental/Dataflow/Source/DataflowEditor/Private/Dataflow/DataflowEditorPreviewSceneBase.cpp b/Engine/Plugins/Experimental/Dataflow/Source/DataflowEditor/Private/Dataflow/DataflowEditorPreviewSceneBase.cpp
--- a/Engine/Plugins/Experimental/Dataflow/Source/DataflowEditor/Private/Dataflow/DataflowEditorPreviewSceneBase.cpp
+++ b/Engine/Plugins/Experimental/Dataflow/Source/DataflowEditor/Private/Dataflow/DataflowEditorPreviewSceneBase.cpp
@@ -66,7 +66,10 @@ FDataflowPreviewSceneBase::FDataflowPreviewSceneBase(FPreviewScene::Construction
 }
 
 FDataflowPreviewSceneBase::~FDataflowPreviewSceneBase()
-{}
+{
+	UnregisterSceneElements();
+	ModifySceneElements().Reset();
+}
 
 TObjectPtr<UDataflowBaseContent>& FDataflowPreviewSceneBase::GetEditorContent() 
 { 
diff --git a/Engine/Plugins/Experimental/Dataflow/Source/DataflowEditor/Private/Dataflow/DataflowEditorUtil.cpp b/Engine/Plugins/Experimental/Dataflow/Source/DataflowEditor/Private/Dataflow/DataflowEditorUtil.cpp
--- a/Engine/Plugins/Experimental/Dataflow/Source/DataflowEditor/Private/Dataflow/DataflowEditorUtil.cpp
+++ b/Engine/Plugins/Experimental/Dataflow/Source/DataflowEditor/Private/Dataflow/DataflowEditorUtil.cpp
@@ -14,6 +14,7 @@
 #include "DynamicMesh/MeshNormals.h"
 #include "Engine/SkeletalMesh.h"
 #include "Animation/Skeleton.h"
+#include "GeometryCollection/Facades/CollectionRenderingFacade.h"
 #include "Materials/Material.h"
 
 
@@ -141,12 +142,13 @@ namespace UE::Dataflow
 		return false;
 	}
 
-	void RenderNodeOutput(GeometryCollection::Facades::FRenderingFacade& Facade, const UDataflowEdNode& Node, const UDataflowBaseContent& EditorContent)
+	bool RenderNodeOutput(GeometryCollection::Facades::FRenderingFacade& RenderingFacade, const UDataflowEdNode& Node, const UDataflowBaseContent& EditorContent)
 	{
 		const TObjectPtr<UDataflow>& DataflowAsset = EditorContent.GetDataflowAsset();
 		const IDataflowConstructionViewMode* ConstructionViewMode = EditorContent.GetConstructionViewMode();
 		const TSharedPtr<FEngineContext>& DataflowContext = EditorContent.GetDataflowContext();
 
+		bool bHasPrimitives = false;
 		if (DataflowAsset && DataflowContext && ConstructionViewMode)
 		{
 			if (FRenderingFactory* const Factory = FRenderingFactory::GetInstance())
@@ -157,12 +159,14 @@ namespace UE::Dataflow
 					{
 						if (TSharedPtr<const FDataflowNode> NodeTarget = Graph->FindBaseNode(FName(Node.GetName())))
 						{
-							Factory->RenderNodeOutput(Facade, FGraphRenderingState{ Node.GetDataflowNodeGuid(), NodeTarget.Get(), Parameter, *DataflowContext, *ConstructionViewMode });
+							Factory->RenderNodeOutput(RenderingFacade, FGraphRenderingState{ Node.GetDataflowNodeGuid(), NodeTarget.Get(), Parameter, *DataflowContext, *ConstructionViewMode });
+							bHasPrimitives = bHasPrimitives || NodeTarget->HasPrimitives();
 						}
 					}
 				}
 			}
 		}
+		return bHasPrimitives;
 	}
 
 
diff --git a/Engine/Plugins/Experimental/Dataflow/Source/DataflowEditor/Private/Dataflow/DataflowEditorUtil.h b/Engine/Plugins/Experimental/Dataflow/Source/DataflowEditor/Private/Dataflow/DataflowEditorUtil.h
--- a/Engine/Plugins/Experimental/Dataflow/Source/DataflowEditor/Private/Dataflow/DataflowEditorUtil.h
+++ b/Engine/Plugins/Experimental/Dataflow/Source/DataflowEditor/Private/Dataflow/DataflowEditorUtil.h
@@ -6,7 +6,6 @@
 #include "Containers/UnrealString.h"
 #include "DynamicMesh/DynamicMesh3.h"
 #include "Dataflow/DataflowObjectInterface.h"
-#include "GeometryCollection/Facades/CollectionRenderingFacade.h"
 #include "Templates/SharedPointer.h"
 
 class UDataflow;
@@ -18,6 +17,10 @@ class UAnimationAsset;
 class UMaterial;
 class UDataflowEdNode;
 namespace UE::Dataflow { class IDataflowConstructionViewMode; }
+namespace GeometryCollection::Facades
+{
+	class FRenderingFacade;
+};
 
 namespace UE::Dataflow::Private
 {
@@ -49,6 +52,9 @@ namespace UE::Dataflow
 	TSharedPtr<UE::Dataflow::FEngineContext> GetContext(TObjectPtr<UDataflowBaseContent> Content);
 
 	bool CanRenderNodeOutput(const UDataflowEdNode& EdNode, const UDataflowBaseContent& EditorContent, const UE::Dataflow::IDataflowConstructionViewMode& ViewMode);
-	void RenderNodeOutput(GeometryCollection::Facades::FRenderingFacade& Facade, const UDataflowEdNode& Node, const UDataflowBaseContent& EditorContent);
+
+	/** This function fills the managed array collection linked to the rendering facade from the callbacks.
+	 * Returns true if the resulting collection should be directly rendered of false if the collection will be used to generate a primitive component */
+	bool RenderNodeOutput(GeometryCollection::Facades::FRenderingFacade& RenderingFacade, const UDataflowEdNode& Node, const UDataflowBaseContent& EditorContent);
 	void GetViewModesForNode(const UDataflowEdNode& EdNode, const UDataflowBaseContent& EditorContent, TArray<FName>& OutValidViewModeNames);
 }
\ No newline at end of file
diff --git a/Engine/Plugins/Experimental/Dataflow/Source/DataflowEditor/Private/Dataflow/DataflowRenderingFactory.cpp b/Engine/Plugins/Experimental/Dataflow/Source/DataflowEditor/Private/Dataflow/DataflowRenderingFactory.cpp
--- a/Engine/Plugins/Experimental/Dataflow/Source/DataflowEditor/Private/Dataflow/DataflowRenderingFactory.cpp
+++ b/Engine/Plugins/Experimental/Dataflow/Source/DataflowEditor/Private/Dataflow/DataflowRenderingFactory.cpp
@@ -5,16 +5,17 @@
 #include "Dataflow/DataflowNode.h"
 #include "Dataflow/DataflowInputOutput.h"
 #include "Misc/MessageDialog.h"
+#include "GeometryCollection/Facades/CollectionRenderingFacade.h"
 
 namespace UE::Dataflow
 {
 	FRenderingFactory* FRenderingFactory::Instance = nullptr;
 
-	void FRenderingFactory::RenderNodeOutput(GeometryCollection::Facades::FRenderingFacade& RenderData, const FGraphRenderingState& State)
+	void FRenderingFactory::RenderNodeOutput(GeometryCollection::Facades::FRenderingFacade& RenderingFacade, const FGraphRenderingState& State)
 	{
 		if (CallbackMap.Contains(State.GetRenderKey()))
 		{
-			CallbackMap[State.GetRenderKey()]->Render(RenderData, State);
+			CallbackMap[State.GetRenderKey()]->Render(RenderingFacade, State);
 		}
 		else
 		{
diff --git a/Engine/Plugins/Experimental/Dataflow/Source/DataflowEditor/Public/Dataflow/DataflowConstructionScene.h b/Engine/Plugins/Experimental/Dataflow/Source/DataflowEditor/Public/Dataflow/DataflowConstructionScene.h
--- a/Engine/Plugins/Experimental/Dataflow/Source/DataflowEditor/Public/Dataflow/DataflowConstructionScene.h
+++ b/Engine/Plugins/Experimental/Dataflow/Source/DataflowEditor/Public/Dataflow/DataflowConstructionScene.h
@@ -83,15 +83,21 @@ private:
 	/** Add a mesh visualizer to the scene */
 	void AddWireframeMeshElementsVisualizer();
 
-	/** Reset all the dynamic mesh components */
-    void ResetDynamicMeshComponents();
+	/** Reset all the scene components (dynamic mesh / primitive)*/
+    void ResetSceneComponents();
 	
 	/** Build the dynamic mesh components from the dataflow graph render targets*/
 	void UpdateDynamicMeshComponents();
+
+	/** Build the primitive components from the dataflow graph selected node */
+    void UpdatePrimitiveComponents();
 	
 	/** Add a dynamic mesh component to the scene */
 	TObjectPtr<UDynamicMeshComponent>& AddDynamicMeshComponent(
 		FDataflowRenderKey InKey, const FString& MeshName, UE::Geometry::FDynamicMesh3&& DynamicMesh,  const TArray<UMaterialInterface*>& MaterialSet);
+
+	/** Remove the scene component from selection and destroy it */
+	void RemoveSceneComponent(USelection* SelectedComponents, UPrimitiveComponent* PrimitiveComponent);
 	
 	/** Dynamic mesh components generated by the tools for edit/manipulation*/
 	typedef TPair<FDataflowRenderKey, TObjectPtr<UDynamicMeshComponent>> FRenderElement;
@@ -112,4 +118,7 @@ private:
 
 	/** Show the wireframe on render */ 
 	bool bConstructionViewWireframe = true;
+	
+	/** Primitive components generated by the nodes and to be used by the tools */
+	TArray<TObjectPtr<UPrimitiveComponent>> PrimitiveComponents;
 };
diff --git a/Engine/Plugins/Experimental/Dataflow/Source/DataflowEditor/Public/Dataflow/DataflowRenderingFactory.h b/Engine/Plugins/Experimental/Dataflow/Source/DataflowEditor/Public/Dataflow/DataflowRenderingFactory.h
--- a/Engine/Plugins/Experimental/Dataflow/Source/DataflowEditor/Public/Dataflow/DataflowRenderingFactory.h
+++ b/Engine/Plugins/Experimental/Dataflow/Source/DataflowEditor/Public/Dataflow/DataflowRenderingFactory.h
@@ -9,8 +9,11 @@
 #include "Dataflow/DataflowNode.h"
 #include "Dataflow/DataflowObjectInterface.h"
 #include "Dataflow/DataflowRenderingViewMode.h"
-#include "GeometryCollection/Facades/CollectionRenderingFacade.h"
 
+namespace GeometryCollection::Facades
+{
+	class FRenderingFacade;
+}
 namespace UE::Dataflow
 {
 	class FContext;
@@ -76,7 +79,7 @@ namespace UE::Dataflow
 			virtual ~ICallbackInterface() = default;
 			virtual FRenderKey GetRenderKey() const = 0;
 			virtual bool CanRender(const UE::Dataflow::IDataflowConstructionViewMode& ViewMode) const = 0;
-			virtual void Render(GeometryCollection::Facades::FRenderingFacade& RenderData, const FGraphRenderingState& State) = 0;
+			virtual void Render(GeometryCollection::Facades::FRenderingFacade& RenderingFacade, const FGraphRenderingState& State) = 0;
 		};
 
 		~FRenderingFactory() { delete Instance; }
@@ -87,7 +90,7 @@ namespace UE::Dataflow
 		DATAFLOWEDITOR_API void DeregisterCallbacks(const FRenderKey& Key);
 
 		DATAFLOWEDITOR_API bool Contains(const FRenderKey& InKey) const { return CallbackMap.Contains(InKey); }
-		DATAFLOWEDITOR_API void RenderNodeOutput(GeometryCollection::Facades::FRenderingFacade& RenderData, const FGraphRenderingState& State);
+		DATAFLOWEDITOR_API void RenderNodeOutput(GeometryCollection::Facades::FRenderingFacade& RenderingFacade, const FGraphRenderingState& State);
 		DATAFLOWEDITOR_API bool CanRenderNodeOutput(const FGraphRenderingState& State) const;
 
 	private:
diff --git a/Engine/Plugins/Experimental/Dataflow/Source/DataflowNodes/Public/Dataflow/DataflowPrimitiveNode.h b/Engine/Plugins/Experimental/Dataflow/Source/DataflowNodes/Public/Dataflow/DataflowPrimitiveNode.h
new file mode 100644
--- /dev/null
+++ b/Engine/Plugins/Experimental/Dataflow/Source/DataflowNodes/Public/Dataflow/DataflowPrimitiveNode.h
@@ -0,0 +1,46 @@
+﻿// Copyright Epic Games, Inc. All Rights Reserved.
+
+#pragma once
+
+#include "Dataflow/DataflowNode.h"
+
+#include "DataflowPrimitiveNode.generated.h"
+
+/**
+* FDataflowPrimitiveNode
+*	Base class for nodes that could add a primitive component to the construction scene within the Dataflow graph. 
+*/
+USTRUCT(Meta = (Experimental))
+struct FDataflowPrimitiveNode : public FDataflowNode
+{
+	GENERATED_BODY()
+
+	FDataflowPrimitiveNode()
+		: Super() { }
+
+	FDataflowPrimitiveNode(const UE::Dataflow::FNodeParameters& Param, FGuid InGuid = FGuid::NewGuid())
+		: Super(Param,InGuid) {
+	}
+
+	/** FDataflowNode interface */
+	virtual ~FDataflowPrimitiveNode() {}
+	
+	static FName StaticType() { return FName("FDataflowPrimitiveNode"); }
+
+	virtual bool IsA(FName InType) const override 
+	{ 
+		return InType.ToString().Equals(StaticType().ToString()) 
+			|| Super::IsA(InType); 
+	}
+
+	virtual bool HasPrimitives() const override
+	{
+		return true;
+	}
+
+	/** Add primitive components to the construction scene */
+	virtual void AddPrimitiveComponents(const TSharedPtr<const FManagedArrayCollection> RenderCollection, TObjectPtr<UObject> NodeOwner, 
+		TObjectPtr<AActor> RootActor, TArray<TObjectPtr<UPrimitiveComponent>>& PrimitiveComponents) {}
+};
+
+
diff --git a/Engine/Source/Runtime/Experimental/Dataflow/Core/Public/Dataflow/DataflowNode.h b/Engine/Source/Runtime/Experimental/Dataflow/Core/Public/Dataflow/DataflowNode.h
--- a/Engine/Source/Runtime/Experimental/Dataflow/Core/Public/Dataflow/DataflowNode.h
+++ b/Engine/Source/Runtime/Experimental/Dataflow/Core/Public/Dataflow/DataflowNode.h
@@ -117,6 +117,12 @@ PRAGMA_ENABLE_DEPRECATION_WARNINGS
 	UE_DEPRECATED(5.4, "FDataflowNode::IsExperimental is deprecated.")
 	virtual bool IsExperimental() { return false; }
 
+	/** Override if the node could generate primitives .*/
+	DATAFLOWCORE_API virtual bool HasPrimitives() const
+	{
+		return false;
+	}
+
 #if WITH_EDITOR
 	struct FDebugDrawParameters
 	{
