Dataflow Skin weight painting tool

diff --git a/Engine/Plugins/Experimental/Dataflow/Dataflow.uplugin b/Engine/Plugins/Experimental/Dataflow/Dataflow.uplugin
--- a/Engine/Plugins/Experimental/Dataflow/Dataflow.uplugin
+++ b/Engine/Plugins/Experimental/Dataflow/Dataflow.uplugin
@@ -22,7 +22,7 @@
 		},
 		{
 			"Name": "DataflowAssetTools",
-			"Type": "Editor",
+			"Type": "Runtime",
 			"LoadingPhase": "Default"
 		},
 		{
diff --git a/Engine/Plugins/Experimental/Dataflow/Source/DataflowAssetTools/DataflowAssetTools.Build.cs b/Engine/Plugins/Experimental/Dataflow/Source/DataflowAssetTools/DataflowAssetTools.Build.cs
--- a/Engine/Plugins/Experimental/Dataflow/Source/DataflowAssetTools/DataflowAssetTools.Build.cs
+++ b/Engine/Plugins/Experimental/Dataflow/Source/DataflowAssetTools/DataflowAssetTools.Build.cs
@@ -26,7 +26,6 @@ public class DataflowAssetTools : ModuleRules
 				"DataflowEnginePlugin",
 				"SlateCore",
 				"Projects",
-				"UnrealEd",
 				"GeometryCore",					// For DynamicMesh conversion
 				"SkeletalMeshDescription"		// For FSkeletalMeshAttributes::DefaultSkinWeightProfileName
 			}
diff --git a/Engine/Plugins/Experimental/Dataflow/Source/DataflowAssetTools/Private/Dataflow/CollectionRenderingPatternUtility.cpp b/Engine/Plugins/Experimental/Dataflow/Source/DataflowAssetTools/Private/Dataflow/CollectionRenderingPatternUtility.cpp
--- a/Engine/Plugins/Experimental/Dataflow/Source/DataflowAssetTools/Private/Dataflow/CollectionRenderingPatternUtility.cpp
+++ b/Engine/Plugins/Experimental/Dataflow/Source/DataflowAssetTools/Private/Dataflow/CollectionRenderingPatternUtility.cpp
@@ -1,9 +1,6 @@
 // Copyright Epic Games, Inc. All Rights Reserved. 
 
 #include "Dataflow/CollectionRenderingPatternUtility.h"
-
-#include "Dataflow/DataflowEdNode.h"
-#include "Dataflow/DataflowObject.h"
 #include "DynamicMesh/MeshNormals.h"
 #include "DynamicMesh/NonManifoldMappingSupport.h"
 #include "Engine/SkeletalMesh.h"
diff --git a/Engine/Plugins/Experimental/Dataflow/Source/DataflowAssetTools/Public/Dataflow/CollectionRenderingPatternUtility.h b/Engine/Plugins/Experimental/Dataflow/Source/DataflowAssetTools/Public/Dataflow/CollectionRenderingPatternUtility.h
--- a/Engine/Plugins/Experimental/Dataflow/Source/DataflowAssetTools/Public/Dataflow/CollectionRenderingPatternUtility.h
+++ b/Engine/Plugins/Experimental/Dataflow/Source/DataflowAssetTools/Public/Dataflow/CollectionRenderingPatternUtility.h
@@ -2,8 +2,6 @@
 #pragma once
 
 #include "CoreMinimal.h"
-#include "Dataflow/DataflowObjectInterface.h"
-
 
 namespace UE::Geometry { class FDynamicMesh3; }
 namespace GeometryCollection::Facades { class FRenderingFacade; }
diff --git a/Engine/Plugins/Experimental/Dataflow/Source/DataflowEditor/DataflowEditor.Build.cs b/Engine/Plugins/Experimental/Dataflow/Source/DataflowEditor/DataflowEditor.Build.cs
--- a/Engine/Plugins/Experimental/Dataflow/Source/DataflowEditor/DataflowEditor.Build.cs
+++ b/Engine/Plugins/Experimental/Dataflow/Source/DataflowEditor/DataflowEditor.Build.cs
@@ -86,7 +86,8 @@ namespace UnrealBuildTool.Rules
 					"ChaosCaching", // UAnimSingleNodeInstance
 					"StructUtilsEditor",
 					"MessageLog",
-					"AppFramework"
+					"AppFramework",
+					"SkeletalMeshDescription"
 				}
 			);
 		}
diff --git a/Engine/Plugins/Experimental/Dataflow/Source/DataflowEditor/Private/Dataflow/DataflowEditorModule.cpp b/Engine/Plugins/Experimental/Dataflow/Source/DataflowEditor/Private/Dataflow/DataflowEditorModule.cpp
--- a/Engine/Plugins/Experimental/Dataflow/Source/DataflowEditor/Private/Dataflow/DataflowEditorModule.cpp
+++ b/Engine/Plugins/Experimental/Dataflow/Source/DataflowEditor/Private/Dataflow/DataflowEditorModule.cpp
@@ -22,6 +22,8 @@
 
 #include "PropertyEditorModule.h"
 #include "EditorModeRegistry.h"
+#include "Dataflow/DataflowCollectionEditSkinWeightsNode.h"
+#include "DataflowEditorTools/DataflowEditorSkinWeightsPaintTool.h"
 
 #define LOCTEXT_NAMESPACE "DataflowEditor"
 
@@ -49,6 +51,22 @@ namespace UE::Dataflow::Private
 		}
 	};
 
+	class FDataflowEditorSkinWeightPaintToolActionCommands : public TInteractiveToolCommands<FDataflowEditorSkinWeightPaintToolActionCommands>
+	{
+	public:
+		FDataflowEditorSkinWeightPaintToolActionCommands() : 
+			TInteractiveToolCommands<FDataflowEditorSkinWeightPaintToolActionCommands>(
+				TEXT("DataflowEditorSkinWeightPaintToolContext"),
+				LOCTEXT("DataflowEditorSkinWeightPaintToolContext", "Dataflow Skin weight Paint Tool Context"),
+				NAME_None,
+				FAppStyle::GetAppStyleSetName())
+		{}
+
+		virtual void GetToolDefaultObjectList(TArray<UInteractiveTool*>& ToolCDOs) override
+		{
+			ToolCDOs.Add(GetMutableDefault<UDataflowEditorSkinWeightsPaintTool>());
+		}
+	};
 
 	class FDataflowToolActionCommandBindings : public UE::Dataflow::FDataflowToolRegistry::IDataflowToolActionCommands
 	{
@@ -56,12 +74,16 @@ namespace UE::Dataflow::Private
 		FDataflowToolActionCommandBindings()
 		{
 			FDataflowEditorWeightMapPaintToolActionCommands::Register();
+			FDataflowEditorSkinWeightPaintToolActionCommands::Register();
 		}
 
 		virtual void UnbindActiveCommands(const TSharedPtr<FUICommandList>& UICommandList) const override
 		{
 			checkf(FDataflowEditorWeightMapPaintToolActionCommands::IsRegistered(), TEXT("Expected WeightMapPaintTool actions to have been registered"));
 			FDataflowEditorWeightMapPaintToolActionCommands::Get().UnbindActiveCommands(UICommandList);
+
+			checkf(FDataflowEditorSkinWeightPaintToolActionCommands::IsRegistered(), TEXT("Expected SkinWeightPaintTool actions to have been registered"));
+			FDataflowEditorSkinWeightPaintToolActionCommands::Get().UnbindActiveCommands(UICommandList);
 		}
 
 		virtual void BindCommandsForCurrentTool(const TSharedPtr<FUICommandList>& UICommandList, UInteractiveTool* Tool) const override
@@ -71,6 +93,11 @@ namespace UE::Dataflow::Private
 				checkf(FDataflowEditorWeightMapPaintToolActionCommands::IsRegistered(), TEXT("Expected WeightMapPaintTool actions to have been registered"));
 				FDataflowEditorWeightMapPaintToolActionCommands::Get().BindCommandsForCurrentTool(UICommandList, Tool);
 			}
+			if (ExactCast<UDataflowEditorSkinWeightsPaintTool>(Tool))
+			{
+				checkf(FDataflowEditorSkinWeightPaintToolActionCommands::IsRegistered(), TEXT("Expected SkinWeightPaintTool actions to have been registered"));
+				FDataflowEditorSkinWeightPaintToolActionCommands::Get().BindCommandsForCurrentTool(UICommandList, Tool);
+			}
 		}
 	};
 }
@@ -95,11 +122,16 @@ void FDataflowEditorModule::StartupModule()
 	UE::Dataflow::RenderingCallbacks();
 
 	UE::Dataflow::FDataflowToolRegistry& ToolRegistry = UE::Dataflow::FDataflowToolRegistry::Get();
-
-	UDataflowEditorWeightMapPaintToolBuilder* const ToolBuilder = NewObject<UDataflowEditorWeightMapPaintToolBuilder>();
 	TSharedRef<const FDataflowToolActionCommandBindings> Actions = MakeShared<FDataflowToolActionCommandBindings>();
-	ToolRegistry.AddNodeToToolMapping(FDataflowCollectionAddScalarVertexPropertyNode::StaticType(), ToolBuilder, Actions, FSlateIcon(), LOCTEXT("AddWeightMapNodeButtonText", "Scalar Vertex Property"));
+	
+	ToolRegistry.AddNodeToToolMapping(FDataflowCollectionAddScalarVertexPropertyNode::StaticType(),
+		NewObject<UDataflowEditorWeightMapPaintToolBuilder>(), Actions,
+		FSlateIcon(FName("DataflowEditorStyle"), FName("Dataflow.PaintWeightMap")), LOCTEXT("AddWeightMapNodeButtonText", "Paint Weight Map"));
 
+	ToolRegistry.AddNodeToToolMapping(FDataflowCollectionEditSkinWeightsNode::StaticType(), 
+		NewObject<UDataflowEditorSkinWeightsPaintToolBuilder>(), Actions,
+		FSlateIcon(FName("DataflowEditorStyle"), FName("Dataflow.EditSkinWeights")), LOCTEXT("AddSkinWeightNodeButtonText", "Edit Skin Weights"));
+	
 	UE::Dataflow::FDataflowConstructionVisualizationRegistry& ConstructionVisualizationRegistry = UE::Dataflow::FDataflowConstructionVisualizationRegistry::GetInstance();
 	ConstructionVisualizationRegistry.RegisterVisualization(MakeUnique<UE::Dataflow::FMeshStatsConstructionVisualization>());
 	ConstructionVisualizationRegistry.RegisterVisualization(MakeUnique<UE::Dataflow::FMeshConstructionVisualization>());
@@ -126,6 +158,7 @@ void FDataflowEditorModule::ShutdownModule()
 
 	UE::Dataflow::FDataflowToolRegistry& ToolRegistry = UE::Dataflow::FDataflowToolRegistry::Get();
 	ToolRegistry.RemoveNodeToToolMapping(FDataflowCollectionAddScalarVertexPropertyNode::StaticType());
+	ToolRegistry.RemoveNodeToToolMapping(FDataflowCollectionEditSkinWeightsNode::StaticType());
 
 	UE::Dataflow::FDataflowConstructionVisualizationRegistry& ConstructionVisualizationRegistry = UE::Dataflow::FDataflowConstructionVisualizationRegistry::GetInstance();
 	ConstructionVisualizationRegistry.DeregisterVisualization(UE::Dataflow::FMeshStatsConstructionVisualization::Name);
diff --git a/Engine/Plugins/Experimental/Dataflow/Source/DataflowEditor/Private/Dataflow/DataflowEditorStyle.h b/Engine/Plugins/Experimental/Dataflow/Source/DataflowEditor/Private/Dataflow/DataflowEditorStyle.h
--- a/Engine/Plugins/Experimental/Dataflow/Source/DataflowEditor/Private/Dataflow/DataflowEditorStyle.h
+++ b/Engine/Plugins/Experimental/Dataflow/Source/DataflowEditor/Private/Dataflow/DataflowEditorStyle.h
@@ -27,30 +27,31 @@ public:
 		SetCoreContentRoot(FPaths::EngineContentDir() / TEXT("Editor/Slate"));
 
 		SetContentRoot(IPluginManager::Get().FindPlugin("Dataflow")->GetBaseDir() / TEXT("Resources"));
-		Set("ClassIcon.Dataflow", new FSlateVectorImageBrush(RootToContentDir(TEXT("DataflowAsset_16.svg")), Icon16x16));
-		Set("ClassThumbnail.Dataflow", new FSlateVectorImageBrush(RootToContentDir(TEXT("DataflowAsset_64.svg")), Icon64x64));
+		Set("ClassIcon.Dataflow", new IMAGE_BRUSH_SVG("DataflowAsset_16", Icon16x16));
+		Set("ClassThumbnail.Dataflow", new IMAGE_BRUSH_SVG("DataflowAsset_64", Icon64x64));
 
-		Set("Dataflow.Render.Unknown", new FSlateImageBrush(RootToContentDir(TEXT("Slate/Switch_Undetermined_56x_28x.png")), Icon28x14));
-		Set("Dataflow.Render.Disabled", new FSlateImageBrush(RootToContentDir(TEXT("Slate/Switch_OFF_56x_28x.png")), Icon28x14));
-		Set("Dataflow.Render.Enabled", new FSlateImageBrush(RootToContentDir(TEXT("Slate/Switch_ON_56x_28x.png")), Icon28x14));
+		Set("Dataflow.Render.Unknown", new IMAGE_BRUSH("Slate/Switch_Undetermined_56x_28x", Icon28x14));
+		Set("Dataflow.Render.Disabled", new IMAGE_BRUSH("Slate/Switch_OFF_56x_28x", Icon28x14));
+		Set("Dataflow.Render.Enabled", new IMAGE_BRUSH("Slate/Switch_ON_56x_28x", Icon28x14));
 
-		Set("Dataflow.Cached.False", new FSlateImageBrush(RootToContentDir(TEXT("Slate/status_grey.png")), Icon16x16));
-		Set("Dataflow.Cached.True", new FSlateImageBrush(RootToContentDir(TEXT("Slate/status_green.png")), Icon16x16));
+		Set("Dataflow.Cached.False", new IMAGE_BRUSH("Slate/status_grey", Icon16x16));
+		Set("Dataflow.Cached.True", new IMAGE_BRUSH("Slate/status_green", Icon16x16));
 
-		Set("Dataflow.SelectObject", new FSlateImageBrush(RootToContentDir(TEXT("Slate/Dataflow_SelectObject_40x.png")), Icon40x40));
-		Set("Dataflow.SelectFace", new FSlateImageBrush(RootToContentDir(TEXT("Slate/Dataflow_SelectFace_40x.png")), Icon40x40));
-		Set("Dataflow.SelectVertex", new FSlateImageBrush(RootToContentDir(TEXT("Slate/Dataflow_SelectVertex_40x.png")), Icon40x40));
+		Set("Dataflow.SelectObject", new IMAGE_BRUSH("Slate/Dataflow_SelectObject_40x", Icon40x40));
+		Set("Dataflow.SelectFace", new IMAGE_BRUSH("Slate/Dataflow_SelectFace_40x", Icon40x40));
+		Set("Dataflow.SelectVertex", new IMAGE_BRUSH("Slate/Dataflow_SelectVertex_40x", Icon40x40));
 
-		Set("Dataflow.FreezeNode", new FSlateImageBrush(RootToContentDir(TEXT("Slate/Dataflow_FreezeNode_24x.png")), Icon24x24));  // There's also a 32x version if it is decided that it is too small 
+		Set("Dataflow.FreezeNode", new IMAGE_BRUSH("Slate/Dataflow_FreezeNode_24x", Icon24x24));  // There's also a 32x version if it is decided that it is too small
+		
+		// Dataflow weight map and skin weights icons 
+		Set("Dataflow.PaintWeightMap", new IMAGE_BRUSH_SVG("Slate/Dataflow_WeightMap", Icon20x20));
+		Set("Dataflow.EditSkinWeights", new IMAGE_BRUSH_SVG("Slate/Dataflow_SkinWeight", Icon20x20));
 
 		// @todo(brice) Remove Example Tools
 		//const FString AttributeEditorPropertyName = "DataflowEditor." + FDataflowEditorCommandsImpl::BeginAttributeEditorToolIdentifier;
 		//Set(*AttributeEditorPropertyName, new FSlateImageBrush(RootToContentDir(TEXT("Slate/Dataflow_SelectObject_20x.png")), Icon20x20));
 		//const FString MeshSelectionPropertyName = "DataflowEditor." + FDataflowEditorCommandsImpl::BeginMeshSelectionToolIdentifier;
 		//Set(*MeshSelectionPropertyName, new FSlateImageBrush(RootToContentDir(TEXT("Slate/Dataflow_SelectVertex_20x.png")), Icon20x20));
-		
-		const FString PropertyNameString = "DataflowEditor." + FDataflowEditorCommandsImpl::AddWeightMapNodeIdentifier;
-		Set(*PropertyNameString, new IMAGE_BRUSH_SVG("Icons/PaintMaps", ToolbarIconSize));
 
 		DefaultMaterial = Cast<UMaterial>(StaticLoadObject(UMaterial::StaticClass(), NULL, TEXT("/Engine/BasicShapes/BasicShapeMaterial")));
 		VertexMaterial = Cast<UMaterial>(StaticLoadObject(UMaterial::StaticClass(), NULL, TEXT("/Engine/EditorMaterials/Dataflow/DataflowVertexMaterial")));
diff --git a/Engine/Plugins/Experimental/Dataflow/Source/DataflowEditor/Private/DataflowEditorTools/DataflowEditorSkinWeightsPaintTool.cpp b/Engine/Plugins/Experimental/Dataflow/Source/DataflowEditor/Private/DataflowEditorTools/DataflowEditorSkinWeightsPaintTool.cpp
new file mode 100644
--- /dev/null
+++ b/Engine/Plugins/Experimental/Dataflow/Source/DataflowEditor/Private/DataflowEditorTools/DataflowEditorSkinWeightsPaintTool.cpp
@@ -0,0 +1,275 @@
+﻿// Copyright Epic Games, Inc. All Rights Reserved.
+
+#include "DataflowEditorTools/DataflowEditorSkinWeightsPaintTool.h"
+
+#include "ContextObjectStore.h"
+#include "Dataflow/DataflowCollectionEditSkinWeightsNode.h"
+#include "Dataflow/DataflowContent.h"
+#include "Dataflow/DataflowObjectInterface.h"
+#include "Dataflow/DataflowSNode.h"
+#include "Engine/World.h"
+#include "InteractiveToolManager.h"
+#include "ModelingToolTargetUtil.h"
+#include "Components/SkeletalMeshComponent.h"
+#include "DynamicMesh/DynamicVertexSkinWeightsAttribute.h"
+#include "TargetInterfaces/MeshDescriptionCommitter.h"
+#include "TargetInterfaces/MeshDescriptionProvider.h"
+#include "TargetInterfaces/PrimitiveComponentBackedTarget.h"
+
+#include UE_INLINE_GENERATED_CPP_BY_NAME(DataflowEditorSkinWeightsPaintTool)
+
+using namespace UE::Geometry;
+
+#define LOCTEXT_NAMESPACE "UDataflowEditorSkinWeightsPaintTool"
+
+void UDataflowEditorSkinWeightsPaintToolBuilder::GetSupportedConstructionViewModes(const UDataflowContextObject& ContextObject, TArray<const UE::Dataflow::IDataflowConstructionViewMode*>& Modes) const
+{}
+
+bool UDataflowEditorSkinWeightsPaintToolBuilder::CanBuildTool(const FToolBuilderState& SceneState) const
+{
+	auto HasManagedArrayCollection = [](const FDataflowNode* InDataflowNode, const TSharedPtr<UE::Dataflow::FEngineContext> Context)
+	{
+		if (InDataflowNode && Context)
+		{
+			for (const FDataflowOutput* const Output : InDataflowNode->GetOutputs())
+			{
+				if (Output->GetType() == FName("FManagedArrayCollection"))
+				{
+					return true;
+				}
+			}
+		}
+
+		return false;
+	};
+
+	if (USkinWeightsPaintToolBuilder::CanBuildTool(SceneState))
+	{
+		if (SceneState.SelectedComponents.Num() == 1 && SceneState.SelectedComponents[0]->IsA<USkeletalMeshComponent>())
+		{
+			if (UDataflowBaseContent* ContextObject = SceneState.ToolManager->GetContextObjectStore()->FindContext<UDataflowBaseContent>())
+			{
+				if (const TSharedPtr<UE::Dataflow::FEngineContext> EvaluationContext = ContextObject->GetDataflowContext())
+				{
+					if (const FDataflowNode* PrimarySelection = ContextObject->GetSelectedNodeOfType<FDataflowCollectionEditSkinWeightsNode>())
+					{
+						return HasManagedArrayCollection(PrimarySelection, EvaluationContext);
+					}
+				}
+			}
+		}
+	}
+	return false;
+}
+
+const FToolTargetTypeRequirements& UDataflowEditorSkinWeightsPaintToolBuilder::GetTargetRequirements() const
+{
+	static FToolTargetTypeRequirements TypeRequirements({
+		UMeshDescriptionProvider::StaticClass(),
+		UMeshDescriptionCommitter::StaticClass(),
+		});
+	return TypeRequirements;
+}
+
+UMeshSurfacePointTool* UDataflowEditorSkinWeightsPaintToolBuilder::CreateNewTool(const FToolBuilderState& SceneState) const
+{
+	UDataflowEditorSkinWeightsPaintTool* PaintTool = NewObject<UDataflowEditorSkinWeightsPaintTool>(SceneState.ToolManager);
+
+	if (UDataflowBaseContent* ContextObject = SceneState.ToolManager->GetContextObjectStore()->FindContext<UDataflowBaseContent>())
+	{
+		if (FDataflowNode* PrimarySelection = ContextObject->GetSelectedNodeOfType<FDataflowCollectionEditSkinWeightsNode>())
+		{
+			FDataflowCollectionEditSkinWeightsNode* SkinWeightNode =
+				StaticCast<FDataflowCollectionEditSkinWeightsNode*>(PrimarySelection);
+			PaintTool->SkinWeightNode = SkinWeightNode;
+
+			SkinWeightNode->OnBoneSelectionChanged.AddLambda(
+				[PaintTool](const TArray<FName>& BoneNames)
+			{
+				PaintTool->GetNotifier().HandleNotification(BoneNames, ESkeletalMeshNotifyType::BonesSelected);
+			});
+
+			PaintTool->SetDataflowEditorContextObject(ContextObject);
+		}
+	}
+	return PaintTool;
+}
+
+FMeshDescription* UDataflowEditorSkinWeightsPaintTool::GetCurrentDescription(const int32 LODIndex) const 
+{
+	if(const IPrimitiveComponentBackedTarget* TargetComponent = Cast<IPrimitiveComponentBackedTarget>(Target))
+	{
+		if(USkeletalMeshComponent* SkeletalMeshComponent = Cast<USkeletalMeshComponent>(TargetComponent->GetOwnerComponent()))
+		{
+			if(USkeletalMesh* SkeletalMesh = SkeletalMeshComponent->GetSkeletalMeshAsset())
+			{
+				return SkeletalMesh->GetMeshDescription(LODIndex);
+			}
+		}
+	}
+	return nullptr;
+}
+
+bool UDataflowEditorSkinWeightsPaintTool::FillVertexOffsets(TArray<int32>& VertexOffsets) const 
+{
+	if(const IPrimitiveComponentBackedTarget* TargetComponent = Cast<IPrimitiveComponentBackedTarget>(Target))
+	{
+		if(USkeletalMeshComponent* SkeletalMeshComponent = Cast<USkeletalMeshComponent>(TargetComponent->GetOwnerComponent()))
+		{
+			if(USkeletalMesh* SkeletalMesh = SkeletalMeshComponent->GetSkeletalMeshAsset())
+			{
+				const int32 NumLODs = SkeletalMesh->GetLODNum();
+				VertexOffsets.Init(0, NumLODs+1);
+
+				int32 NumVertices = 0;
+				for(int32 LODIndex = 0; LODIndex < NumLODs; ++LODIndex)
+				{
+					VertexOffsets[LODIndex] = NumVertices;
+					if(FMeshDescription* MeshDescription = SkeletalMesh->GetMeshDescription(static_cast<int32>(LODIndex)))
+					{
+						NumVertices += MeshDescription->Vertices().Num();
+					}
+				}
+				VertexOffsets[NumLODs] = NumVertices;
+				return true;
+			}
+		}
+	}
+	return false;
+}
+
+bool UDataflowEditorSkinWeightsPaintTool::ExtractSkinWeights(const int32 NumVertices, TArray<TArray<int32>>& CurrentIndices, TArray<TArray<float>>& CurrentWeights)
+{
+	// Setup DynamicMeshToWeight conversion and get Input weight map (if it exists)
+	SetupWeights.SetNumZeroed(NumVertices);
+	SetupIndices.SetNumZeroed(NumVertices);
+
+	if(DataflowEditorContextObject && SkinWeightNode)
+	{
+		// Find the map if it exists.
+		if (TSharedPtr<UE::Dataflow::FEngineContext> DataflowContext = DataflowEditorContextObject->GetDataflowContext())
+		{
+			if(DataflowEditorContextObject->GetSelectedCollection().IsValid())
+			{
+				// Fill the attribute values
+				SkinWeightNode->FillAttributeWeights(*DataflowEditorContextObject->GetSelectedCollection().Get(),
+					SkinWeightNode->GetBoneIndicesKey(*DataflowContext), SkinWeightNode->GetBoneWeightsKey(*DataflowContext),
+					SetupIndices, SetupWeights);
+						
+				CurrentWeights.SetNumZeroed(NumVertices);
+				CurrentIndices.SetNumZeroed(NumVertices);
+						
+				SkinWeightNode->ExtractVertexWeights(SetupIndices, SetupWeights,
+					TArrayView<TArray<int32>>(CurrentIndices), TArrayView<TArray<float>>(CurrentWeights));
+
+				return true;
+			}
+		}
+	}
+	return false;
+}
+
+void UDataflowEditorSkinWeightsPaintTool::Setup()
+{
+	TArray<int32> VertexOffsets;
+	if(FillVertexOffsets(VertexOffsets))
+	{
+		TArray<TArray<float>> CurrentWeights;
+		TArray<TArray<int32>> CurrentIndices;
+		
+		if(ExtractSkinWeights(VertexOffsets.Last(), CurrentIndices, CurrentWeights))
+		{
+			for(int32 LODIndex = 0, NumLODs = VertexOffsets.Num()-1; LODIndex < NumLODs; ++LODIndex)
+			{
+				if(FMeshDescription* MeshDescription = GetCurrentDescription(LODIndex))
+				{
+					const int32 NumVertices = MeshDescription->Vertices().Num();
+				
+					FSkeletalMeshAttributes MeshAttribs(*MeshDescription);
+					FSkinWeightsVertexAttributesRef SkinWeights =
+						MeshAttribs.GetVertexSkinWeights(FSkeletalMeshAttributesShared::DefaultSkinWeightProfileName);
+				
+					UE::AnimationCore::FBoneWeightsSettings WeightsSettings;
+					WeightsSettings.SetNormalizeType(UE::AnimationCore::EBoneWeightNormalizeType::Always);
+
+					const int32 VertexOffset = VertexOffsets[LODIndex];
+					for(int32 VertexIndex = 0; VertexIndex < NumVertices; ++VertexIndex)
+					{
+						TArray<UE::AnimationCore::FBoneWeight> BoneWeights;
+						const int32 CurrentIndex = VertexOffset+VertexIndex;
+					
+						if(CurrentIndex < CurrentWeights.Num() && CurrentIndex < CurrentIndices.Num() &&
+							CurrentWeights[CurrentIndex].Num() == CurrentIndices[CurrentIndex].Num())
+						{
+							BoneWeights.Reserve(CurrentWeights[CurrentIndex].Num());
+							for(int32 WeightIndex = 0, NumWeights = CurrentWeights[CurrentIndex].Num(); WeightIndex < NumWeights; ++WeightIndex)
+							{
+								BoneWeights.Add(UE::AnimationCore::FBoneWeight(CurrentIndices[CurrentIndex][WeightIndex],
+									CurrentWeights[CurrentIndex][WeightIndex]));
+							}
+						}
+						SkinWeights.Set(FVertexID(VertexIndex), UE::AnimationCore::FBoneWeights::Create(BoneWeights, WeightsSettings));
+					}
+				}
+			}
+		}
+	}
+	USkinWeightsPaintTool::Setup();
+}
+
+void UDataflowEditorSkinWeightsPaintTool::OnShutdown(EToolShutdownType ShutdownType)
+{
+	USkinWeightsPaintTool::OnShutdown(ShutdownType);
+
+	if(Target && ShutdownType == EToolShutdownType::Accept)
+	{
+		TArray<int32> VertexOffsets;
+		if(FillVertexOffsets(VertexOffsets))
+		{
+			TArray<TArray<int32>> CurrentIndices;
+			TArray<TArray<float>> CurrentWeights;
+			
+			CurrentIndices.SetNumZeroed(SetupIndices.Num());
+			CurrentWeights.SetNumZeroed(SetupWeights.Num());
+			
+			for(TPair<EMeshLODIdentifier, FCleanedEditMesh>& EditedeMesh : EditedMeshes)
+			{
+				const FMeshDescription& MeshDescription = EditedeMesh.Value.GetEditableMeshDescription();
+				const int32 LODIndex = static_cast<int32>(EditedeMesh.Key);
+			
+				// profile to edit
+				const FName ActiveProfile = WeightToolProperties->GetActiveSkinWeightProfile();
+			
+				FSkeletalMeshConstAttributes MeshAttribs(MeshDescription);
+				FSkinWeightsVertexAttributesConstRef SkinWeights = MeshAttribs.GetVertexSkinWeights(ActiveProfile);
+			
+				const int32 NumVertices = MeshDescription.Vertices().Num();
+				if(LODIndex < VertexOffsets.Num())
+				{
+					const int32 VertexOffset = VertexOffsets[LODIndex];
+					for(int32 VertexIndex = 0; VertexIndex < NumVertices; ++VertexIndex)
+					{
+						FVertexBoneWeightsConst BoneWeights = SkinWeights.Get(FVertexID(VertexIndex));
+						const int32 CurrentIndex = VertexOffset+VertexIndex;
+				
+						if(CurrentIndex < CurrentWeights.Num() && CurrentIndex < CurrentIndices.Num() &&
+							CurrentWeights[CurrentIndex].Num() == CurrentIndices[CurrentIndex].Num())
+						{
+							CurrentIndices[CurrentIndex].Reserve(BoneWeights.Num());
+							CurrentWeights[CurrentIndex].Reserve(BoneWeights.Num());
+							for(int32 WeightIndex = 0, NumWeights = BoneWeights.Num(); WeightIndex < NumWeights; ++WeightIndex)
+							{
+								CurrentIndices[CurrentIndex].Add(BoneWeights[WeightIndex].GetBoneIndex());
+								CurrentWeights[CurrentIndex].Add(BoneWeights[WeightIndex].GetWeight());
+							}
+						}
+					}
+				}
+			}
+			SkinWeightNode->ReportVertexWeights(SetupIndices, SetupWeights, CurrentIndices, CurrentWeights);
+			SkinWeightNode->Invalidate();
+		}
+	}
+}
+
+#undef LOCTEXT_NAMESPACE
diff --git a/Engine/Plugins/Experimental/Dataflow/Source/DataflowEditor/Public/DataflowEditorTools/DataflowEditorSkinWeightsPaintTool.h b/Engine/Plugins/Experimental/Dataflow/Source/DataflowEditor/Public/DataflowEditorTools/DataflowEditorSkinWeightsPaintTool.h
new file mode 100644
--- /dev/null
+++ b/Engine/Plugins/Experimental/Dataflow/Source/DataflowEditor/Public/DataflowEditorTools/DataflowEditorSkinWeightsPaintTool.h
@@ -0,0 +1,85 @@
+﻿// Copyright Epic Games, Inc. All Rights Reserved.
+
+#pragma once
+
+#include "CoreMinimal.h"
+#include "HAL/Platform.h"
+#include "DataflowEditorTools/DataflowEditorToolBuilder.h"
+#include "SkeletalMesh/SkinWeightsPaintTool.h"
+
+#include "DataflowEditorSkinWeightsPaintTool.generated.h"
+
+class UDataflowContextObject;
+class UDataflowEditorMode;
+
+/**
+ * Dataflow skin weights tool builder
+ */
+UCLASS()
+class DATAFLOWEDITOR_API UDataflowEditorSkinWeightsPaintToolBuilder : public USkinWeightsPaintToolBuilder, public IDataflowEditorToolBuilder
+{
+	GENERATED_BODY()
+
+private:
+	//~ Begin IDataflowEditorToolBuilder interface
+	virtual void GetSupportedConstructionViewModes(const UDataflowContextObject& ContextObject, TArray<const UE::Dataflow::IDataflowConstructionViewMode*>& Modes) const override;
+	virtual bool CanSetConstructionViewWireframeActive() const override { return false; }
+	//~ End IDataflowEditorToolBuilder interface
+
+	//~ Begin USkinWeightsPaintToolBuilder interface
+	virtual bool CanBuildTool(const FToolBuilderState& SceneState) const override;
+	virtual UMeshSurfacePointTool* CreateNewTool(const FToolBuilderState& SceneState) const override;
+	//~ End USkinWeightsPaintToolBuilder interface
+
+protected:
+	virtual const FToolTargetTypeRequirements& GetTargetRequirements() const override;
+};
+
+
+/**
+ * Dataflow skin weights painting tool
+ */
+UCLASS()
+class DATAFLOWEDITOR_API UDataflowEditorSkinWeightsPaintTool : public USkinWeightsPaintTool
+{
+	GENERATED_BODY()
+	
+	//~ Begin USkinWeightsPaintTool interface
+	virtual FEditorViewportClient* GetViewportClient() const override {return nullptr;}
+	virtual void Setup() override;
+	virtual void OnShutdown(EToolShutdownType ShutdownType) override;
+	//~ End USkinWeightsPaintTool interface
+
+	friend class UDataflowEditorSkinWeightsPaintToolBuilder;
+
+private:
+
+	/** Set the dataflow context object */
+	void SetDataflowEditorContextObject(TObjectPtr<UDataflowContextObject> InDataflowEditorContextObject)
+	{
+		DataflowEditorContextObject = InDataflowEditorContextObject;
+	}
+
+	/** Extract the skin weights from the node */
+	bool ExtractSkinWeights(const int32 NumVertices, TArray<TArray<int32>>& CurrentIndices, TArray<TArray<float>>& CurrentWeights);
+
+	/** Fill the vertex offsets fromm all the LODs description */
+	bool FillVertexOffsets(TArray<int32>& VertexOffsets) const;
+
+	/** Get the LOD description given a  LOD index */
+	FMeshDescription* GetCurrentDescription(const int32 LODIndex) const;
+	
+	/** Skin weight node associated with that tool */
+	struct FDataflowCollectionEditSkinWeightsNode* SkinWeightNode = nullptr;
+
+	/** Setup weights used to store the initial values */
+	TArray<TArray<float>> SetupWeights;
+
+	/** Setup indices used to store the initial values */
+	TArray<TArray<int32>> SetupIndices;
+
+	/** Dataflow context object to be used to retrieve the node context data*/
+	TObjectPtr<UDataflowContextObject> DataflowEditorContextObject;
+};
+
+
diff --git a/Engine/Plugins/Experimental/Dataflow/Source/DataflowNodes/DataflowNodes.Build.cs b/Engine/Plugins/Experimental/Dataflow/Source/DataflowNodes/DataflowNodes.Build.cs
--- a/Engine/Plugins/Experimental/Dataflow/Source/DataflowNodes/DataflowNodes.Build.cs
+++ b/Engine/Plugins/Experimental/Dataflow/Source/DataflowNodes/DataflowNodes.Build.cs
@@ -16,7 +16,19 @@ namespace UnrealBuildTool.Rules
 					"DataflowCore",
 					"DataflowEngine",
 					"DataflowEnginePlugin",
-					"Engine"
+					"DataflowAssetTools",
+					"Engine",
+					"MeshConversion",
+					"MeshDescription",
+					"SkeletalMeshDescription",
+					"GeometryCore"
+				}
+			);
+
+			PrivateDependencyModuleNames.AddRange(
+				new string[]
+				{
+					"SkeletalMeshUtilitiesCommon"
 				}
 			);
 		}
diff --git a/Engine/Plugins/Experimental/Dataflow/Source/DataflowNodes/Private/Dataflow/DataflowCollectionEditSkinWeightsNode.cpp b/Engine/Plugins/Experimental/Dataflow/Source/DataflowNodes/Private/Dataflow/DataflowCollectionEditSkinWeightsNode.cpp
new file mode 100644
--- /dev/null
+++ b/Engine/Plugins/Experimental/Dataflow/Source/DataflowNodes/Private/Dataflow/DataflowCollectionEditSkinWeightsNode.cpp
@@ -0,0 +1,358 @@
+// Copyright Epic Games, Inc. All Rights Reserved.
+
+#include "Dataflow/DataflowCollectionEditSkinWeightsNode.h"
+#include "Dataflow/DataflowInputOutput.h"
+#include "Dataflow/DataflowDebugDrawObject.h"
+#include "Dataflow/DataflowSkeletalMeshNodes.h"
+#include "Dataflow/CollectionRenderingPatternUtility.h"
+
+#include "DynamicMeshToMeshDescription.h"
+#include "MeshConversionOptions.h"
+#include "SkeletalMeshAttributes.h"
+#include "StaticToSkeletalMeshConverter.h"
+#include "DynamicMesh/DynamicMesh3.h"
+#include "GeometryCollection/Facades/CollectionRenderingFacade.h"
+#include "Rendering/SkeletalMeshModel.h"
+#include "Engine/SkeletalMesh.h"
+#include "Animation/Skeleton.h"
+#include "Components/SkeletalMeshComponent.h"
+
+#include UE_INLINE_GENERATED_CPP_BY_NAME(DataflowCollectionEditSkinWeightsNode)
+
+#define LOCTEXT_NAMESPACE "DataflowCollectionEditSkinWeights"
+
+namespace UE::Dataflow::Private
+{
+
+static bool CreateSkeletalMesh(USkeletalMesh* SkeletalMesh, const TArray<UE::Geometry::FDynamicMesh3>& DynamicMeshes, const FReferenceSkeleton& ReferenceSkeleton)
+{
+#if WITH_EDITOR
+	TArray<const FMeshDescription*> MeshDescriptions;
+
+#if WITH_EDITORONLY_DATA
+	SkeletalMesh->PreEditChange(nullptr);
+	SkeletalMesh->GetImportedModel()->LODModels.Empty();
+#endif
+	SkeletalMesh->ResetLODInfo();
+	
+	TArray<FMeshDescription> LocalDescriptions;
+	for (const UE::Geometry::FDynamicMesh3& DynamicMesh : DynamicMeshes)
+	{
+		// Create a mesh description
+		FMeshDescription& MeshDescription = LocalDescriptions.AddDefaulted_GetRef();
+		
+		// Add skeletal mesh attributes to the mesh description
+		FSkeletalMeshAttributes Attributes(MeshDescription);
+		Attributes.Register();
+
+		// Convert dynamic mesh to the mesh description
+		FConversionToMeshDescriptionOptions ConverterOptions;
+		FDynamicMeshToMeshDescription Converter(ConverterOptions);
+		Converter.Convert(&DynamicMesh, MeshDescription, false);
+
+		// Add the created description to the list
+		MeshDescriptions.Add(&MeshDescription);
+	}
+
+	TArray<FSkeletalMaterial> Materials;
+	TConstArrayView<FSkeletalMaterial> MaterialView;
+
+	// ensure there is at least one material
+	if (MaterialView.IsEmpty())
+	{
+		Materials.Add(FSkeletalMaterial());
+		MaterialView = Materials;
+	}
+	
+	if (!FStaticToSkeletalMeshConverter::InitializeSkeletalMeshFromMeshDescriptions(
+		SkeletalMesh, MeshDescriptions, MaterialView, ReferenceSkeleton,true, true))
+	{
+		return false;
+	}
+	return true;
+#else
+	return false;
+#endif
+}
+	
+static bool BuildSkeletalMesh(USkeletalMesh* SkeletalMesh, 
+	const TSharedPtr<const FManagedArrayCollection>& RenderCollection, const TObjectPtr<USkeleton>& ObjectSkeleton)
+{
+	if(RenderCollection.IsValid() && ObjectSkeleton)
+	{
+		GeometryCollection::Facades::FRenderingFacade Facade(*RenderCollection);
+		if(Facade.IsValid())
+		{
+			TArray<UE::Geometry::FDynamicMesh3> DynamicMeshes;
+			const int32 NumGeometry = Facade.NumGeometry();
+
+			for (int32 MeshIndex = 0; MeshIndex < NumGeometry; ++MeshIndex)
+			{
+				UE::Geometry::FDynamicMesh3& DynamicMesh = DynamicMeshes.AddDefaulted_GetRef();
+				UE::Dataflow::Conversion::RenderingFacadeToDynamicMesh(Facade, MeshIndex, DynamicMesh);
+			}
+	
+			return CreateSkeletalMesh(SkeletalMesh, DynamicMeshes, ObjectSkeleton->GetReferenceSkeleton());
+		}
+	}
+	return false;
+}
+
+	
+template<typename ArrayType>
+static bool SetAttributeValues(FManagedArrayCollection& SelectedCollection,
+	const FCollectionAttributeKey& AttributeKey, const TArray<TArray<ArrayType>>& AttributeValues)
+{
+	if (!AttributeValues.IsEmpty() && !AttributeKey.Attribute.IsEmpty() && !AttributeKey.Group.IsEmpty())
+	{
+		const FName AttributeName(AttributeKey.Attribute);
+		const FName AttributeGroup(AttributeKey.Group);
+		
+		if(TManagedArray<TArray<ArrayType>>* AttributeArray = SelectedCollection.FindAttributeTyped<TArray<ArrayType>>(AttributeName, AttributeGroup))
+		{
+			if(AttributeArray->Num() == AttributeValues.Num())
+			{
+				for(int32 VertexIndex = 0, NumVertices = AttributeArray->Num(); VertexIndex < NumVertices; ++VertexIndex)
+				{
+					AttributeArray->GetData()[VertexIndex] = AttributeValues[VertexIndex];
+				}
+			}
+			return true;
+		}
+	}
+	return false;
+}
+
+template<typename ArrayType>
+static bool GetAttributeValues(FManagedArrayCollection& SelectedCollection,
+	const FCollectionAttributeKey& AttributeKey, TArray<TArray<ArrayType>>& AttributeValues)
+{
+	if (!AttributeKey.Attribute.IsEmpty() && !AttributeKey.Group.IsEmpty())
+	{
+		const FName AttributeName(AttributeKey.Attribute);
+		const FName AttributeGroup(AttributeKey.Group);
+		
+		TManagedArray<TArray<ArrayType>>& AttributeArray = SelectedCollection.AddAttribute<TArray<ArrayType>>(AttributeName, AttributeGroup);
+		AttributeValues = AttributeArray.GetConstArray();
+		return true;
+	}
+	return false;
+}
+
+template<typename ArrayType>
+static bool FillAttributeValues(const FManagedArrayCollection& SelectedCollection,
+	const FCollectionAttributeKey& AttributeKey, TArray<TArray<ArrayType>>& AttributeValues)
+{
+	if (!AttributeKey.Attribute.IsEmpty() && !AttributeKey.Group.IsEmpty())
+	{
+		const FName AttributeName(AttributeKey.Attribute);
+		const FName AttributeGroup(AttributeKey.Group);
+
+		if(const TManagedArray<TArray<ArrayType>>* AttributeArray = SelectedCollection.FindAttributeTyped<TArray<ArrayType>>(AttributeName, AttributeGroup))
+		{
+			AttributeValues = AttributeArray->GetConstArray();
+			return true;
+		}
+	}
+	return false;
+}
+
+}
+
+//
+// FDataflowCollectionEditSkinWeightsNode
+//
+
+FDataflowCollectionEditSkinWeightsNode::FDataflowCollectionEditSkinWeightsNode(const UE::Dataflow::FNodeParameters& InParam, FGuid InGuid)
+	: FDataflowPrimitiveNode(InParam, InGuid)
+{
+	RegisterInputConnection(&Collection);
+	RegisterInputConnection(&BoneIndicesKey);
+	RegisterInputConnection(&BoneWeightsKey);
+	RegisterOutputConnection(&Collection, &Collection);
+	RegisterOutputConnection(&BoneIndicesKey, &BoneIndicesKey);
+	RegisterOutputConnection(&BoneWeightsKey, &BoneWeightsKey);
+}
+
+TArray<UE::Dataflow::FRenderingParameter> FDataflowCollectionEditSkinWeightsNode::GetRenderParametersImpl() const
+{
+	return FDataflowAddScalarVertexPropertyCallbackRegistry::Get().GetRenderingParameters(VertexGroup.Name);
+}
+
+void FDataflowCollectionEditSkinWeightsNode::AddPrimitiveComponents(const TSharedPtr<const FManagedArrayCollection> RenderCollection, TObjectPtr<UObject> NodeOwner, 
+	TObjectPtr<AActor> RootActor, TArray<TObjectPtr<UPrimitiveComponent>>& PrimitiveComponents)  
+{
+	if(RootActor)
+	{
+		if(!SkeletalMesh)
+		{
+			//Create transient skeletal mesh
+			FString SkeletalMeshName("SK_DataflowSkeletalMesh");
+			SkeletalMeshName = MakeUniqueObjectName(NodeOwner, USkeletalMesh::StaticClass(), *SkeletalMeshName, EUniqueObjectNameOptions::GloballyUnique).ToString();
+			SkeletalMesh = NewObject<USkeletalMesh>(NodeOwner, FName(*SkeletalMeshName), RF_Transient);
+		}
+		USkeletalMeshComponent* SkeletalMeshComponent = SkeletalMeshComponent = NewObject<USkeletalMeshComponent>(RootActor, FName("Dataflow_SkeletalMesh"));
+		if(UE::Dataflow::Private::BuildSkeletalMesh(SkeletalMesh, RenderCollection, ObjectSkeleton))
+		{
+			SkeletalMeshComponent->SetSkeletalMesh(SkeletalMesh);
+		}
+		PrimitiveComponents.Add(SkeletalMeshComponent);
+	}
+}
+
+void FDataflowCollectionEditSkinWeightsNode::Evaluate(UE::Dataflow::FContext& Context, const FDataflowOutput* Out) const
+{
+	using namespace UE::Dataflow;
+	
+	// Get the pin value if plugged
+	FCollectionAttributeKey BoneIndicesKeyValue = GetBoneIndicesKey(Context);
+	FCollectionAttributeKey BoneWeightsKeyValue = GetBoneWeightsKey(Context);
+
+	if (Out->IsA<FManagedArrayCollection>(&Collection))
+	{
+		// Evaluate in collection
+		FManagedArrayCollection InCollection = GetValue<FManagedArrayCollection>(Context, &Collection);
+
+		if (!BoneIndicesKeyValue.Attribute.IsEmpty() && !BoneWeightsKeyValue.Attribute.IsEmpty() )
+		{
+			TArray<TArray<float>> SetupWeights, FinalWeights;
+			TArray<TArray<int32>> SetupIndices, FinalIndices;
+			
+			UE::Dataflow::Private::GetAttributeValues<int32>(InCollection, BoneIndicesKeyValue, SetupIndices);
+		    UE::Dataflow::Private::GetAttributeValues<float>(InCollection, BoneWeightsKeyValue, SetupWeights);
+			
+			FinalIndices.SetNum(SetupIndices.Num());
+			FinalWeights.SetNum(SetupWeights.Num());
+			
+			ExtractVertexWeights(SetupIndices, SetupWeights, TArrayView<TArray<int32>>(FinalIndices.GetData(), FinalIndices.Num()),
+			TArrayView<TArray<float>>(FinalWeights.GetData(), FinalWeights.Num()));
+
+			UE::Dataflow::Private::SetAttributeValues<int32>(InCollection, BoneIndicesKeyValue, FinalIndices);
+			UE::Dataflow::Private::SetAttributeValues<float>(InCollection, BoneWeightsKeyValue, FinalWeights);
+		}
+
+		SetValue(Context, MoveTemp(InCollection), &Collection);
+	}
+	else if (Out->IsA<FCollectionAttributeKey>(&BoneIndicesKey))
+	{
+		SetValue(Context, MoveTemp(BoneIndicesKeyValue), &BoneIndicesKey);
+	}
+	else if (Out->IsA<FCollectionAttributeKey>(&BoneWeightsKey))
+	{
+		SetValue(Context, MoveTemp(BoneWeightsKeyValue), &BoneWeightsKey);
+	}
+}
+
+void FDataflowCollectionEditSkinWeightsNode::ReportVertexWeights(const TArray<TArray<int32>>& SetupIndices, const TArray<TArray<float>>& SetupWeights, const TArray<TArray<int32>>& FinalIndices, const TArray<TArray<float>>& FinalWeights)
+{
+	check(SetupWeights.Num() == FinalWeights.Num());
+	check(SetupWeights.Num() == SetupIndices.Num());
+	check(FinalWeights.Num() == FinalIndices.Num());
+	
+	SkinWeights.SetNumZeroed(FinalWeights.Num());
+	
+	for (int32 VertexIndex = 0, NumVertices = FinalWeights.Num(); VertexIndex < NumVertices; ++VertexIndex)
+	{
+		SkinWeights[VertexIndex].BoneWeights = FinalWeights[VertexIndex];
+		SkinWeights[VertexIndex].BoneIndices = FinalIndices[VertexIndex];
+	}
+}
+
+void FDataflowCollectionEditSkinWeightsNode::ExtractVertexWeights(const TArray<TArray<int32>>& SetupIndices, const TArray<TArray<float>>& SetupWeights,
+	TArrayView<TArray<int32>> FinalIndices, TArrayView<TArray<float>> FinalWeights) const
+{
+	check(SetupWeights.Num() == FinalWeights.Num());
+	check(SetupWeights.Num() == SetupIndices.Num());
+	check(FinalWeights.Num() == FinalIndices.Num());
+
+	if(SkinWeights.Num() == FinalWeights.Num())
+	{
+		for (int32 VertexIndex = 0, NumVertices = FinalWeights.Num(); VertexIndex < NumVertices; ++VertexIndex)
+		{
+			if(!SkinWeights[VertexIndex].BoneWeights.IsEmpty() && !SkinWeights[VertexIndex].BoneIndices.IsEmpty())
+			{
+				FinalWeights[VertexIndex] = SkinWeights[VertexIndex].BoneWeights;
+				FinalIndices[VertexIndex] = SkinWeights[VertexIndex].BoneIndices;
+			}
+			else
+			{
+				FinalWeights[VertexIndex] = SetupWeights[VertexIndex];
+				FinalIndices[VertexIndex] = SetupIndices[VertexIndex];
+			}
+		}
+	}
+	else
+	{
+		for (int32 VertexIndex = 0, NumVertices = FinalWeights.Num(); VertexIndex < NumVertices; ++VertexIndex)
+		{
+			FinalWeights[VertexIndex] = SetupWeights[VertexIndex];
+			FinalIndices[VertexIndex] = SetupIndices[VertexIndex];
+		}
+	}
+}
+
+bool FDataflowCollectionEditSkinWeightsNode::FillAttributeWeights(const FManagedArrayCollection& SelectedCollection,
+	const FCollectionAttributeKey& InBoneIndicesKey, const FCollectionAttributeKey& InBoneWeightsKey,
+	TArray<TArray<int32>>& AttributeIndices, TArray<TArray<float>>& AttributeWeights) const
+{
+	return UE::Dataflow::Private::FillAttributeValues<int32>(SelectedCollection, InBoneIndicesKey, AttributeIndices) &&
+		   UE::Dataflow::Private::FillAttributeValues<float>(SelectedCollection, InBoneWeightsKey, AttributeWeights);
+}
+
+FCollectionAttributeKey FDataflowCollectionEditSkinWeightsNode::GetBoneIndicesKey(UE::Dataflow::FContext& Context) const
+{
+	// Get the pin value if plugged
+	FCollectionAttributeKey Key = GetValue(Context, &BoneIndicesKey, BoneIndicesKey);
+
+	// If nothing set used the local value
+	if(Key.Attribute.IsEmpty() && Key.Group.IsEmpty())
+	{
+		Key.Group = VertexGroup.Name.ToString();
+		Key.Attribute = BoneIndicesName;
+	}
+	return Key;
+}
+
+FCollectionAttributeKey FDataflowCollectionEditSkinWeightsNode::GetBoneWeightsKey(UE::Dataflow::FContext& Context) const
+{
+	// Get the pin value if plugged
+	FCollectionAttributeKey Key = GetValue(Context, &BoneWeightsKey, BoneWeightsKey);
+
+	// If nothing set used the local value
+	if(Key.Attribute.IsEmpty() && Key.Group.IsEmpty())
+	{
+		Key.Group = VertexGroup.Name.ToString();
+		Key.Attribute = BoneWeightsName;
+	}
+	return Key;
+}
+
+#if WITH_EDITOR
+
+void FDataflowCollectionEditSkinWeightsNode::DebugDraw(UE::Dataflow::FContext& Context,
+		IDataflowDebugDrawInterface& DataflowRenderingInterface,
+		const FDebugDrawParameters& DebugDrawParameters) const
+{
+	if(ObjectSkeleton)
+	{
+		TRefCountPtr<FDataflowDebugDrawSkeletonObject> SkeletonObject = MakeDebugDrawObject<FDataflowDebugDrawSkeletonObject>(
+			DataflowRenderingInterface.ModifyDataflowElements(), ObjectSkeleton->GetReferenceSkeleton());
+	
+		DataflowRenderingInterface.DrawObject(TRefCountPtr<IDataflowDebugDrawObject>(SkeletonObject));
+
+		SkeletonObject->OnBoneSelectionChanged.AddLambda([this](const TArray<FName>& BoneNames)
+		{
+			OnBoneSelectionChanged.Broadcast(BoneNames);
+		});
+	}
+}
+
+bool FDataflowCollectionEditSkinWeightsNode::CanDebugDrawViewMode(const FName& ViewModeName) const
+{
+	return true;
+}
+
+#endif
+
+
+#undef LOCTEXT_NAMESPACE
diff --git a/Engine/Plugins/Experimental/Dataflow/Source/DataflowNodes/Private/Dataflow/DataflowNodesPlugin.cpp b/Engine/Plugins/Experimental/Dataflow/Source/DataflowNodes/Private/Dataflow/DataflowNodesPlugin.cpp
--- a/Engine/Plugins/Experimental/Dataflow/Source/DataflowNodes/Private/Dataflow/DataflowNodesPlugin.cpp
+++ b/Engine/Plugins/Experimental/Dataflow/Source/DataflowNodes/Private/Dataflow/DataflowNodesPlugin.cpp
@@ -10,6 +10,7 @@
 #include "Dataflow/DataflowSelectionNodes.h"
 #include "Dataflow/DataflowContextOverridesNodes.h"
 #include "Dataflow/DataflowCollectionAddScalarVertexPropertyNode.h"
+#include "Dataflow/DataflowCollectionEditSkinWeightsNode.h"
 
 #define LOCTEXT_NAMESPACE "DataflowNodes"
 
@@ -47,6 +48,7 @@ void IDataflowNodesPlugin::StartupModule()
 	UE::Dataflow::RegisterContextOverridesNodes();
 	UE::Dataflow::DataflowCollectionAttributeKeyNodes();
 	DATAFLOW_NODE_REGISTER_CREATION_FACTORY(FDataflowCollectionAddScalarVertexPropertyNode);
+	DATAFLOW_NODE_REGISTER_CREATION_FACTORY(FDataflowCollectionEditSkinWeightsNode);
 
 	UE::Dataflow::RegisterNodeFilter(FDataflowTerminalNode::StaticType());
 
diff --git a/Engine/Plugins/Experimental/Dataflow/Source/DataflowNodes/Private/Dataflow/DataflowSkeletalMeshNodes.cpp b/Engine/Plugins/Experimental/Dataflow/Source/DataflowNodes/Private/Dataflow/DataflowSkeletalMeshNodes.cpp
--- a/Engine/Plugins/Experimental/Dataflow/Source/DataflowNodes/Private/Dataflow/DataflowSkeletalMeshNodes.cpp
+++ b/Engine/Plugins/Experimental/Dataflow/Source/DataflowNodes/Private/Dataflow/DataflowSkeletalMeshNodes.cpp
@@ -42,7 +42,7 @@ void FDataflowDebugDrawSkeletonObject::PopulateDataflowElements()
 	ElementsSize = ReferenceSkeleton.GetNum();
 	
 	DataflowElements.Reserve(ElementsOffset + ElementsSize);
-	
+	PreviousSelection.Init(false, ElementsSize);
 	for(int32 BoneIndex = 0; BoneIndex < ElementsSize; ++BoneIndex)
 	{
 		int32 ParentIndex = ReferenceSkeleton.GetParentIndex(BoneIndex);
@@ -83,7 +83,7 @@ FBox FDataflowDebugDrawSkeletonObject::ComputeBoundingBox() const
 	return BoundingBox;
 }
 
-void FDataflowDebugDrawSkeletonObject::DrawDataflowElements(FPrimitiveDrawInterface* PDI) const
+void FDataflowDebugDrawSkeletonObject::DrawDataflowElements(FPrimitiveDrawInterface* PDI) 
 {
 	const uint32 NumBones = ReferenceSkeleton.GetNum();
 
@@ -92,10 +92,12 @@ void FDataflowDebugDrawSkeletonObject::DrawDataflowElements(FPrimitiveDrawInterf
 	TArray<FLinearColor> BoneColors;
 	TArray<FTransform> WorldTransforms;
 	TArray<int32> SelectedBones;
+	TArray<FName> BoneNames;
 	
 	ReferenceSkeleton.GetBoneAbsoluteTransforms(WorldTransforms);
 	RequiredBones.Reserve(NumBones);
 
+	bool bSelectionChanged = false;
 	if(NumBones == ElementsSize)
 	{
 		for(int32 BoneIndex = 0; BoneIndex < ElementsSize; ++BoneIndex)
@@ -107,9 +109,15 @@ void FDataflowDebugDrawSkeletonObject::DrawDataflowElements(FPrimitiveDrawInterf
 				{
 					RequiredBones.Add(BoneIndex);
 				}
+				if(DataflowElement->bIsSelected != PreviousSelection[BoneIndex])
+				{
+					PreviousSelection[BoneIndex] = DataflowElement->bIsSelected;
+					bSelectionChanged = true;
+				}
 				if(DataflowElement->bIsSelected)
 				{
 					SelectedBones.Add(BoneIndex);
+					BoneNames.Add(FName(DataflowElement->ElementName));
 				}
 				// Populate the proxies with the ones stored on the elements. 
 				if(DataflowElement->IsA(FDataflowProxyElement::StaticType()))
@@ -122,6 +130,11 @@ void FDataflowDebugDrawSkeletonObject::DrawDataflowElements(FPrimitiveDrawInterf
 			}
 		}
 	}
+
+	if(bSelectionChanged)
+	{
+		OnBoneSelectionChanged.Broadcast(BoneNames);
+	}
 	
 	FSkelDebugDrawConfig DrawConfig;
 	DrawConfig.bUseMultiColorAsDefaultColor = true;
diff --git a/Engine/Plugins/Experimental/Dataflow/Source/DataflowNodes/Public/Dataflow/DataflowCollectionEditSkinWeightsNode.h b/Engine/Plugins/Experimental/Dataflow/Source/DataflowNodes/Public/Dataflow/DataflowCollectionEditSkinWeightsNode.h
new file mode 100644
--- /dev/null
+++ b/Engine/Plugins/Experimental/Dataflow/Source/DataflowNodes/Public/Dataflow/DataflowCollectionEditSkinWeightsNode.h
@@ -0,0 +1,107 @@
+﻿// Copyright Epic Games, Inc. All Rights Reserved.
+
+#pragma once 
+
+#include "DataflowCollectionAddScalarVertexPropertyNode.h"
+#include "DataflowPrimitiveNode.h"
+#include "DataflowSkeletalMeshNodes.h"
+#include "Dataflow/DataflowNode.h"
+#include "GeometryCollection/GeometryCollection.h"
+#include "GeometryCollection/ManagedArrayCollection.h"
+#include "Dataflow/DataflowCollectionAttributeKeyNodes.h"
+
+#include "DataflowCollectionEditSkinWeightsNode.generated.h"
+
+/** Dataflow skin weights data */
+USTRUCT()
+struct DATAFLOWNODES_API FDataflowSkinWeightData
+{
+	GENERATED_BODY()
+
+	UPROPERTY()
+	TArray<float> BoneWeights;
+
+	UPROPERTY()
+	TArray<int32> BoneIndices;
+};
+
+/** Edit skin weights vertex properties. */
+USTRUCT(Meta = (DataflowCollection))
+struct DATAFLOWNODES_API FDataflowCollectionEditSkinWeightsNode : public FDataflowPrimitiveNode
+{
+	GENERATED_BODY()
+	DATAFLOW_NODE_DEFINE_INTERNAL(FDataflowCollectionEditSkinWeightsNode, "EditSkinWeights", "Collection", "Edit skin weights and save it to collection")
+
+public:
+
+	FDataflowCollectionEditSkinWeightsNode(const UE::Dataflow::FNodeParameters& InParam, FGuid InGuid = FGuid::NewGuid());
+
+	UPROPERTY(Meta = (DataflowInput, DataflowOutput, DataflowPassthrough = "Collection"))
+	FManagedArrayCollection Collection;
+
+	/** The name to be set as a weight map attribute. */
+	UPROPERTY(EditAnywhere, Category = "Vertex Attributes", meta = (DisplayName = "Bone Indices"))
+	FString BoneIndicesName;
+
+	/** The name to be set as a weight map attribute. */
+	UPROPERTY(EditAnywhere, Category = "Vertex Attributes", meta = (DisplayName = "Bone Weights"))
+	FString BoneWeightsName;
+
+	/** Target group in which the attributes are stored */
+	UPROPERTY(EditAnywhere, Category = "Vertex Attributes", meta = (DisplayName = "Vertex Group"))
+	FScalarVertexPropertyGroup VertexGroup;
+
+	/** Bone indices key to be used in other nodes if necessary */
+	UPROPERTY(meta = (DisplayName = "Bone Indices Key", DataflowInput, DataflowOutput, DataflowPassthrough = "BoneIndicesKey"))
+	FCollectionAttributeKey BoneIndicesKey;
+
+	/** Bone weights key to be used in other nodes if necessary */
+	UPROPERTY(meta = (DisplayName = "Bone Weights Key", DataflowInput, DataflowOutput, DataflowPassthrough = "BoneWeightsKey"))
+	FCollectionAttributeKey BoneWeightsKey;
+
+	/** Skeletal mesh to extract the skeleton from for the skinning */
+	UPROPERTY(EditAnywhere, Category = "Skeleton Binding")
+	TObjectPtr<USkeleton> ObjectSkeleton = nullptr;
+
+	/** List of skin weights */
+	UPROPERTY()
+	TArray<FDataflowSkinWeightData> SkinWeights;
+	
+	/** Delegate to transfer the bone selection to the tool*/
+	FDataflowBoneSelectionChangedNotifyDelegate OnBoneSelectionChanged;
+
+	/** Report the vertex weights back onto the property ones */
+	void ReportVertexWeights(const TArray<TArray<int32>>& SetupIndices, const TArray<TArray<float>>& SetupWeights, const TArray<TArray<int32>>& FinalIndices, const TArray<TArray<float>>& FinalWeights);
+
+	/** Extract the vertex weights back from the property ones */
+	void ExtractVertexWeights(const TArray<TArray<int32>>& SetupIndices, const TArray<TArray<float>>& SetupWeights, TArrayView<TArray<int32>> FinalIndices, TArrayView<TArray<float>> FinalWeights) const;
+
+	/** Fill the weight attribute values from the collection */
+	bool FillAttributeWeights(const FManagedArrayCollection& SelectedCollection,
+		const FCollectionAttributeKey& IndicesAttributeKey, const FCollectionAttributeKey& WeightsAttributeKey, TArray<TArray<int32>>& AttributeIndices, TArray<TArray<float>>& AttributeWeights) const;
+
+	/** Get the weights attribute key to retrieve/set the weight values*/
+	FCollectionAttributeKey GetBoneIndicesKey(UE::Dataflow::FContext& Context) const;
+
+	/** Get the weights attribute key to retrieve/set the weight values*/
+    FCollectionAttributeKey GetBoneWeightsKey(UE::Dataflow::FContext& Context) const;
+private:
+
+	//~ Begin FDataflowPrimitiveNode interface
+	virtual TArray<UE::Dataflow::FRenderingParameter> GetRenderParametersImpl() const override;
+	virtual void Evaluate(UE::Dataflow::FContext& Context, const FDataflowOutput* Out) const override;
+	virtual void AddPrimitiveComponents(const TSharedPtr<const FManagedArrayCollection> RenderCollection,
+		TObjectPtr<UObject> NodeOwner, TObjectPtr<AActor> RootActor, TArray<TObjectPtr<UPrimitiveComponent>>& PrimitiveComponents) override;
+#if WITH_EDITOR
+	virtual bool CanDebugDraw() const override { return true; }
+	virtual bool CanDebugDrawViewMode(const FName& ViewModeName) const override;
+	virtual void DebugDraw(UE::Dataflow::FContext& Context,
+		IDataflowDebugDrawInterface& DataflowRenderingInterface,
+		const FDebugDrawParameters& DebugDrawParameters) const override;
+#endif
+	//~ End FDataflowPrimitiveNode interface
+
+	/** Transient skeletal mesh built from dataflow render collection */
+	UPROPERTY(Transient)
+	TObjectPtr<USkeletalMesh> SkeletalMesh = nullptr;
+};
diff --git a/Engine/Plugins/Experimental/Dataflow/Source/DataflowNodes/Public/Dataflow/DataflowSkeletalMeshNodes.h b/Engine/Plugins/Experimental/Dataflow/Source/DataflowNodes/Public/Dataflow/DataflowSkeletalMeshNodes.h
--- a/Engine/Plugins/Experimental/Dataflow/Source/DataflowNodes/Public/Dataflow/DataflowSkeletalMeshNodes.h
+++ b/Engine/Plugins/Experimental/Dataflow/Source/DataflowNodes/Public/Dataflow/DataflowSkeletalMeshNodes.h
@@ -17,6 +17,9 @@ DEFINE_LOG_CATEGORY_STATIC(LogDataflowSkeletalMeshNodes, Log, All);
 
 class USkeletalMesh;
 
+// A delegate for monitoring dataflow skeleton selection changes.
+DECLARE_MULTICAST_DELEGATE_OneParam(FDataflowBoneSelectionChangedNotifyDelegate, const TArray<FName>& /*InBoneNames*/);
+
 /** Debug draw skeleton object to debug skeletal meshes */
 struct FDataflowDebugDrawSkeletonObject : public FDataflowDebugDrawBaseObject
 {
@@ -28,14 +31,20 @@ struct FDataflowDebugDrawSkeletonObject : public FDataflowDebugDrawBaseObject
 	UE_API virtual void PopulateDataflowElements() override;
 
 	/** Debug draw dataflow element */
-	UE_API virtual void DrawDataflowElements(FPrimitiveDrawInterface* PDI) const override;
+	UE_API virtual void DrawDataflowElements(FPrimitiveDrawInterface* PDI) override;
 
 	/** Compute the dataflow elements bounding box */
 	UE_API virtual FBox ComputeBoundingBox() const override;
 
+	/** Delegate to broadcast bones selection changes */
+	FDataflowBoneSelectionChangedNotifyDelegate OnBoneSelectionChanged;
+
 private :
 	/** Skeletal mesh to use to populate/draw scene elements */
 	const FReferenceSkeleton& ReferenceSkeleton;
+	
+	/** Previous element selection */
+	TArray<bool> PreviousSelection;
 };
 
 USTRUCT()
diff --git a/Engine/Plugins/Runtime/MeshModelingToolset/Source/MeshModelingToolsEditorOnly/Public/SkeletalMesh/SkinWeightsPaintTool.h b/Engine/Plugins/Runtime/MeshModelingToolset/Source/MeshModelingToolsEditorOnly/Public/SkeletalMesh/SkinWeightsPaintTool.h
--- a/Engine/Plugins/Runtime/MeshModelingToolset/Source/MeshModelingToolsEditorOnly/Public/SkeletalMesh/SkinWeightsPaintTool.h
+++ b/Engine/Plugins/Runtime/MeshModelingToolset/Source/MeshModelingToolsEditorOnly/Public/SkeletalMesh/SkinWeightsPaintTool.h
@@ -812,7 +812,7 @@ public:
 	UWeightToolTransferManager* GetWeightTransferManager() const { return TransferManager; };
 
 	// get the viewport this tool is operating in
-	FEditorViewportClient* GetViewportClient() const;
+	virtual FEditorViewportClient* GetViewportClient() const;
 	
 	// get access to the selection isolation system
 	UWeightToolSelectionIsolator* GetSelectionIsolator() const { return SelectionIsolator; };
diff --git a/Engine/Source/Runtime/Experimental/Dataflow/Engine/Public/Dataflow/DataflowDebugDrawObject.h b/Engine/Source/Runtime/Experimental/Dataflow/Engine/Public/Dataflow/DataflowDebugDrawObject.h
--- a/Engine/Source/Runtime/Experimental/Dataflow/Engine/Public/Dataflow/DataflowDebugDrawObject.h
+++ b/Engine/Source/Runtime/Experimental/Dataflow/Engine/Public/Dataflow/DataflowDebugDrawObject.h
@@ -28,7 +28,7 @@ struct FDataflowDebugDrawBaseObject : public IDataflowDebugDrawObject
 	virtual void PopulateDataflowElements() = 0;
 
 	/** Debug draw dataflow element */
-	virtual void DrawDataflowElements(FPrimitiveDrawInterface* PDI) const = 0;
+	virtual void DrawDataflowElements(FPrimitiveDrawInterface* PDI) = 0;
 
 	/** Compute teh dataflow elements bounding box */
 	virtual FBox ComputeBoundingBox() const = 0;
