Fix CIS in dataflow nodes dependency on a dev module

diff --git a/Engine/Plugins/Experimental/Dataflow/Source/DataflowNodes/DataflowNodes.Build.cs b/Engine/Plugins/Experimental/Dataflow/Source/DataflowNodes/DataflowNodes.Build.cs
--- a/Engine/Plugins/Experimental/Dataflow/Source/DataflowNodes/DataflowNodes.Build.cs
+++ b/Engine/Plugins/Experimental/Dataflow/Source/DataflowNodes/DataflowNodes.Build.cs
@@ -4,7 +4,7 @@ namespace UnrealBuildTool.Rules
 {
 	public class DataflowNodes : ModuleRules
 	{
-        public DataflowNodes(ReadOnlyTargetRules Target) : base(Target)
+		public DataflowNodes(ReadOnlyTargetRules Target) : base(Target)
 		{
 			PublicDependencyModuleNames.AddRange(
 				new string[]
@@ -25,12 +25,15 @@ namespace UnrealBuildTool.Rules
 				}
 			);
 
-			PrivateDependencyModuleNames.AddRange(
-				new string[]
-				{
-					"SkeletalMeshUtilitiesCommon"
-				}
-			);
+			if (Target.Type == TargetType.Editor)
+			{
+				PrivateDependencyModuleNames.AddRange(
+					new string[]
+					{
+						"SkeletalMeshUtilitiesCommon"
+					}
+				);
+			}
 		}
 	}
 }
diff --git a/Engine/Plugins/Experimental/Dataflow/Source/DataflowNodes/Private/Dataflow/DataflowCollectionEditSkinWeightsNode.cpp b/Engine/Plugins/Experimental/Dataflow/Source/DataflowNodes/Private/Dataflow/DataflowCollectionEditSkinWeightsNode.cpp
--- a/Engine/Plugins/Experimental/Dataflow/Source/DataflowNodes/Private/Dataflow/DataflowCollectionEditSkinWeightsNode.cpp
+++ b/Engine/Plugins/Experimental/Dataflow/Source/DataflowNodes/Private/Dataflow/DataflowCollectionEditSkinWeightsNode.cpp
@@ -9,7 +9,6 @@
 #include "DynamicMeshToMeshDescription.h"
 #include "MeshConversionOptions.h"
 #include "SkeletalMeshAttributes.h"
-#include "StaticToSkeletalMeshConverter.h"
 #include "DynamicMesh/DynamicMesh3.h"
 #include "GeometryCollection/Facades/CollectionRenderingFacade.h"
 #include "Rendering/SkeletalMeshModel.h"
@@ -17,6 +16,10 @@
 #include "Animation/Skeleton.h"
 #include "Components/SkeletalMeshComponent.h"
 
+#if WITH_EDITOR
+#include "StaticToSkeletalMeshConverter.h"
+#endif
+
 #include UE_INLINE_GENERATED_CPP_BY_NAME(DataflowCollectionEditSkinWeightsNode)
 
 #define LOCTEXT_NAMESPACE "DataflowCollectionEditSkinWeights"
@@ -191,7 +194,7 @@ void FDataflowCollectionEditSkinWeightsNode::AddPrimitiveComponents(const TShare
 			SkeletalMeshName = MakeUniqueObjectName(NodeOwner, USkeletalMesh::StaticClass(), *SkeletalMeshName, EUniqueObjectNameOptions::GloballyUnique).ToString();
 			SkeletalMesh = NewObject<USkeletalMesh>(NodeOwner, FName(*SkeletalMeshName), RF_Transient);
 		}
-		USkeletalMeshComponent* SkeletalMeshComponent = SkeletalMeshComponent = NewObject<USkeletalMeshComponent>(RootActor, FName("Dataflow_SkeletalMesh"));
+		USkeletalMeshComponent* SkeletalMeshComponent = NewObject<USkeletalMeshComponent>(RootActor, FName("Dataflow_SkeletalMesh"));
 		if(UE::Dataflow::Private::BuildSkeletalMesh(SkeletalMesh, RenderCollection, ObjectSkeleton))
 		{
 			SkeletalMeshComponent->SetSkeletalMesh(SkeletalMesh);
