dataflow : make convex display the num vertices and num planes information in the collection spreadsheet

diff --git a/Engine/Plugins/Experimental/Dataflow/Source/DataflowEditor/Private/Dataflow/DataflowCollectionSpreadSheetHelpers.cpp b/Engine/Plugins/Experimental/Dataflow/Source/DataflowEditor/Private/Dataflow/DataflowCollectionSpreadSheetHelpers.cpp
--- a/Engine/Plugins/Experimental/Dataflow/Source/DataflowEditor/Private/Dataflow/DataflowCollectionSpreadSheetHelpers.cpp
+++ b/Engine/Plugins/Experimental/Dataflow/Source/DataflowEditor/Private/Dataflow/DataflowCollectionSpreadSheetHelpers.cpp
@@ -145,6 +145,17 @@ namespace UE::Dataflow::CollectionSpreadSheetHelpers
 		return Value.ToString();
 	}
 
+	FString AttributeValueToString(Chaos::FConvexPtr Value)
+	{
+		if (Value)
+		{
+			const int32 NumVertices= Value->NumVertices();
+			const int32 NumPlanes = Value->NumPlanes();
+			return FString::Printf(TEXT("Vertices:(%d), Planes:(%d)"), NumVertices, NumPlanes);
+		}
+		return FString(TEXT("(null)"));
+	}
+
 	FString AttributeValueToString(const FManagedArrayCollection& InCollection, const FName& InAttributeName, const FName& InGroupName, int32 InIdxColumn)
 	{
 		const FManagedArrayCollection::EArrayType ArrayType = InCollection.GetAttributeType(InAttributeName, InGroupName);
@@ -225,6 +236,10 @@ namespace UE::Dataflow::CollectionSpreadSheetHelpers
 			ValueAsString = AttributeValueToString<FGuid>(InCollection, InAttributeName, InGroupName, InIdxColumn);
 			break;
 
+		case FManagedArrayCollection::EArrayType::FFConvexRefCountedPtrType:
+			ValueAsString = AttributeValueToString<Chaos::FConvexPtr>(InCollection, InAttributeName, InGroupName, InIdxColumn);
+			break;
+
 		default:
 			//ensure(false);
 			ValueAsString = "<Unknown Data Type>";
diff --git a/Engine/Plugins/Experimental/Dataflow/Source/DataflowEditor/Public/Dataflow/DataflowCollectionSpreadSheetHelpers.h b/Engine/Plugins/Experimental/Dataflow/Source/DataflowEditor/Public/Dataflow/DataflowCollectionSpreadSheetHelpers.h
--- a/Engine/Plugins/Experimental/Dataflow/Source/DataflowEditor/Public/Dataflow/DataflowCollectionSpreadSheetHelpers.h
+++ b/Engine/Plugins/Experimental/Dataflow/Source/DataflowEditor/Public/Dataflow/DataflowCollectionSpreadSheetHelpers.h
@@ -62,6 +62,8 @@ namespace UE::Dataflow::CollectionSpreadSheetHelpers
 
 	FString AttributeValueToString(FGuid Value);
 
+	FString AttributeValueToString(Chaos::FConvexPtr Value);
+
 	FString AttributeValueToString(const FManagedArrayCollection& InCollection, const FName& InAttributeName, const FName& InGroupName, int32 InIdxColumn);
 
 	inline FName GetArrayTypeString(FManagedArrayCollection::EArrayType ArrayType)
