Improve dataflow painting tools (perf for skin weights / multi mesh support for painting)

diff --git a/Engine/Plugins/Experimental/Dataflow/Source/DataflowEditor/Private/DataflowEditorTools/DataflowEditorSkinWeightsPaintTool.cpp b/Engine/Plugins/Experimental/Dataflow/Source/DataflowEditor/Private/DataflowEditorTools/DataflowEditorSkinWeightsPaintTool.cpp
--- a/Engine/Plugins/Experimental/Dataflow/Source/DataflowEditor/Private/DataflowEditorTools/DataflowEditorSkinWeightsPaintTool.cpp
+++ b/Engine/Plugins/Experimental/Dataflow/Source/DataflowEditor/Private/DataflowEditorTools/DataflowEditorSkinWeightsPaintTool.cpp
@@ -268,6 +268,9 @@ void UDataflowEditorSkinWeightsPaintTool::OnShutdown(EToolShutdownType ShutdownT
 			}
 			SkinWeightNode->ReportVertexWeights(SetupIndices, SetupWeights, CurrentIndices, CurrentWeights);
 			SkinWeightNode->Invalidate();
+
+			// Avoid rebuilding the skeletal mesh after updating the skin weights
+			SkinWeightNode->ValidateSkeletalMesh();
 		}
 	}
 }
diff --git a/Engine/Plugins/Experimental/Dataflow/Source/DataflowEditor/Private/DataflowEditorTools/DataflowEditorWeightMapPaintTool.cpp b/Engine/Plugins/Experimental/Dataflow/Source/DataflowEditor/Private/DataflowEditorTools/DataflowEditorWeightMapPaintTool.cpp
--- a/Engine/Plugins/Experimental/Dataflow/Source/DataflowEditor/Private/DataflowEditorTools/DataflowEditorWeightMapPaintTool.cpp
+++ b/Engine/Plugins/Experimental/Dataflow/Source/DataflowEditor/Private/DataflowEditorTools/DataflowEditorWeightMapPaintTool.cpp
@@ -1790,17 +1790,11 @@ void UDataflowEditorWeightMapPaintTool::UpdateSelectedNode()
 
 	if (bHaveDynamicMeshToWeightConversion)
 	{
-		TArray<float> FinalWeights;
-		FinalWeights.Init(0.f, WeightToDynamicMesh.Num());
-		for (int32 DynamicMeshIdx = 0; DynamicMeshIdx < CurrentWeights.Num(); ++DynamicMeshIdx)
-		{
-			FinalWeights[DynamicMeshToWeight[DynamicMeshIdx]] = CurrentWeights[DynamicMeshIdx];
-		}
-		WeightMapNodeToUpdate->ReportVertexWeights(SetupWeights, FinalWeights);
+		WeightMapNodeToUpdate->ReportVertexWeights(SetupWeights, CurrentWeights, DynamicMeshToWeight);
 	}
 	else
 	{
-		WeightMapNodeToUpdate->ReportVertexWeights(SetupWeights, CurrentWeights);
+		WeightMapNodeToUpdate->ReportVertexWeights(SetupWeights, CurrentWeights, {});
 	}
 	
 	WeightMapNodeToUpdate->Name = UpdateWeightMapProperties->Name;
diff --git a/Engine/Plugins/Experimental/Dataflow/Source/DataflowNodes/Private/Dataflow/DataflowCollectionAddScalarVertexPropertyNode.cpp b/Engine/Plugins/Experimental/Dataflow/Source/DataflowNodes/Private/Dataflow/DataflowCollectionAddScalarVertexPropertyNode.cpp
--- a/Engine/Plugins/Experimental/Dataflow/Source/DataflowNodes/Private/Dataflow/DataflowCollectionAddScalarVertexPropertyNode.cpp
+++ b/Engine/Plugins/Experimental/Dataflow/Source/DataflowNodes/Private/Dataflow/DataflowCollectionAddScalarVertexPropertyNode.cpp
@@ -228,22 +228,24 @@ FCollectionAttributeKey FDataflowCollectionAddScalarVertexPropertyNode::GetWeigh
 	return Key;
 }
 
-void FDataflowCollectionAddScalarVertexPropertyNode::ReportVertexWeights(const TArray<float>& SetupWeights, const TArray<float>& FinalWeights)
+void FDataflowCollectionAddScalarVertexPropertyNode::ReportVertexWeights(const TArray<float>& SetupWeights, const TArray<float>& FinalWeights,
+		const TArray<int32>& WeightIndices)
 {
-	check(SetupWeights.Num() == FinalWeights.Num());
-	VertexWeights.SetNumZeroed(FinalWeights.Num());
+	check((WeightIndices.Num() == FinalWeights.Num()) || (WeightIndices.IsEmpty() && (SetupWeights.Num() == FinalWeights.Num())));
+	VertexWeights.SetNumZeroed(SetupWeights.Num());
 	for (int32 WeightIndex = 0, NumWeights = FinalWeights.Num(); WeightIndex <NumWeights; ++WeightIndex)
 	{
+		const int32 VertexIndex = WeightIndices.IsEmpty() ? WeightIndex : WeightIndices[WeightIndex];
 		switch (OverrideType)
 		{
 			case EDataflowWeightMapOverrideType::ReplaceAll:
-				VertexWeights[WeightIndex] = FinalWeights[WeightIndex];
+				VertexWeights[VertexIndex] = FinalWeights[WeightIndex];
 				break;
 			case EDataflowWeightMapOverrideType::ReplaceChanged:
-				VertexWeights[WeightIndex] = (SetupWeights[WeightIndex] == FinalWeights[WeightIndex]) ? ReplaceChangedPassthroughValue : FinalWeights[WeightIndex];
+				VertexWeights[VertexIndex] = (SetupWeights[VertexIndex] == FinalWeights[WeightIndex]) ? ReplaceChangedPassthroughValue : FinalWeights[WeightIndex];
 				break;
 			case EDataflowWeightMapOverrideType::AddDifference:
-				VertexWeights[WeightIndex] = FinalWeights[WeightIndex] - SetupWeights[WeightIndex];
+				VertexWeights[VertexIndex] = FinalWeights[WeightIndex] - SetupWeights[VertexIndex];
 				break;
 			default: unimplemented();
 		}
diff --git a/Engine/Plugins/Experimental/Dataflow/Source/DataflowNodes/Private/Dataflow/DataflowCollectionEditSkinWeightsNode.cpp b/Engine/Plugins/Experimental/Dataflow/Source/DataflowNodes/Private/Dataflow/DataflowCollectionEditSkinWeightsNode.cpp
--- a/Engine/Plugins/Experimental/Dataflow/Source/DataflowNodes/Private/Dataflow/DataflowCollectionEditSkinWeightsNode.cpp
+++ b/Engine/Plugins/Experimental/Dataflow/Source/DataflowNodes/Private/Dataflow/DataflowCollectionEditSkinWeightsNode.cpp
@@ -66,9 +66,13 @@ static bool CreateSkeletalMesh(USkeletalMesh* SkeletalMesh, const TArray<UE::Geo
 		Materials.Add(FSkeletalMaterial());
 		MaterialView = Materials;
 	}
+
+	static constexpr bool bRecomputeTangents = false;
+	static constexpr bool bRecomputeNormals = false;
+	static constexpr bool bCacheOptimize = false;
 	
 	if (!FStaticToSkeletalMeshConverter::InitializeSkeletalMeshFromMeshDescriptions(
-		SkeletalMesh, MeshDescriptions, MaterialView, ReferenceSkeleton,true, true))
+		SkeletalMesh, MeshDescriptions, MaterialView, ReferenceSkeleton, bRecomputeNormals, bRecomputeTangents, bCacheOptimize))
 	{
 		return false;
 	}
@@ -187,6 +191,7 @@ void FDataflowCollectionEditSkinWeightsNode::AddPrimitiveComponents(const TShare
 {
 	if(RootActor)
 	{
+		const bool bNeedsConstruction = (SkeletalMesh == nullptr) || !bValidSkeletalMesh;
 		if(!SkeletalMesh)
 		{
 			//Create transient skeletal mesh
@@ -194,12 +199,20 @@ void FDataflowCollectionEditSkinWeightsNode::AddPrimitiveComponents(const TShare
 			SkeletalMeshName = MakeUniqueObjectName(NodeOwner, USkeletalMesh::StaticClass(), *SkeletalMeshName, EUniqueObjectNameOptions::GloballyUnique).ToString();
 			SkeletalMesh = NewObject<USkeletalMesh>(NodeOwner, FName(*SkeletalMeshName), RF_Transient);
 		}
-		USkeletalMeshComponent* SkeletalMeshComponent = NewObject<USkeletalMeshComponent>(RootActor, FName("Dataflow_SkeletalMesh"));
-		if(UE::Dataflow::Private::BuildSkeletalMesh(SkeletalMesh, RenderCollection, ObjectSkeleton))
+		if(bNeedsConstruction)
+		{
+			bValidSkeletalMesh = UE::Dataflow::Private::BuildSkeletalMesh(SkeletalMesh, RenderCollection, ObjectSkeleton);
+			if(!bValidSkeletalMesh)
+			{
+				SkeletalMesh = nullptr;
+			}
+		}
+		if(SkeletalMesh)
 		{
+			USkeletalMeshComponent* SkeletalMeshComponent = NewObject<USkeletalMeshComponent>(RootActor, FName("Dataflow_SkeletalMesh"));
 			SkeletalMeshComponent->SetSkeletalMesh(SkeletalMesh);
+			PrimitiveComponents.Add(SkeletalMeshComponent);
 		}
-		PrimitiveComponents.Add(SkeletalMeshComponent);
 	}
 }
 
@@ -330,6 +343,11 @@ FCollectionAttributeKey FDataflowCollectionEditSkinWeightsNode::GetBoneWeightsKe
 	return Key;
 }
 
+void FDataflowCollectionEditSkinWeightsNode::OnInvalidate()
+{
+	bValidSkeletalMesh = false;
+}
+
 #if WITH_EDITOR
 
 void FDataflowCollectionEditSkinWeightsNode::DebugDraw(UE::Dataflow::FContext& Context,
diff --git a/Engine/Plugins/Experimental/Dataflow/Source/DataflowNodes/Public/Dataflow/DataflowCollectionAddScalarVertexPropertyNode.h b/Engine/Plugins/Experimental/Dataflow/Source/DataflowNodes/Public/Dataflow/DataflowCollectionAddScalarVertexPropertyNode.h
--- a/Engine/Plugins/Experimental/Dataflow/Source/DataflowNodes/Public/Dataflow/DataflowCollectionAddScalarVertexPropertyNode.h
+++ b/Engine/Plugins/Experimental/Dataflow/Source/DataflowNodes/Public/Dataflow/DataflowCollectionAddScalarVertexPropertyNode.h
@@ -115,7 +115,7 @@ public:
 	UE_API FDataflowCollectionAddScalarVertexPropertyNode(const UE::Dataflow::FNodeParameters& InParam, FGuid InGuid = FGuid::NewGuid());
 
 	/** Report the vertex weights back onto the property ones */
-	UE_API void ReportVertexWeights(const TArray<float>& SetupWeights, const TArray<float>& FinalWeights);
+	UE_API void ReportVertexWeights(const TArray<float>& SetupWeights, const TArray<float>& FinalWeights, const TArray<int32>& WeightIndices);
 
 	/** Extract the vertex weights back from the property ones */
 	UE_API void ExtractVertexWeights(const TArray<float>& SetupWeights, TArrayView<float> FinalWeights) const;
diff --git a/Engine/Plugins/Experimental/Dataflow/Source/DataflowNodes/Public/Dataflow/DataflowCollectionEditSkinWeightsNode.h b/Engine/Plugins/Experimental/Dataflow/Source/DataflowNodes/Public/Dataflow/DataflowCollectionEditSkinWeightsNode.h
--- a/Engine/Plugins/Experimental/Dataflow/Source/DataflowNodes/Public/Dataflow/DataflowCollectionEditSkinWeightsNode.h
+++ b/Engine/Plugins/Experimental/Dataflow/Source/DataflowNodes/Public/Dataflow/DataflowCollectionEditSkinWeightsNode.h
@@ -85,6 +85,9 @@ public:
 
 	/** Get the weights attribute key to retrieve/set the weight values*/
     FCollectionAttributeKey GetBoneWeightsKey(UE::Dataflow::FContext& Context) const;
+
+	/** Validate the skeletal mesh construction */
+	void ValidateSkeletalMesh() {bValidSkeletalMesh = true;}
 private:
 
 	//~ Begin FDataflowPrimitiveNode interface
@@ -92,6 +95,7 @@ private:
 	virtual void Evaluate(UE::Dataflow::FContext& Context, const FDataflowOutput* Out) const override;
 	virtual void AddPrimitiveComponents(const TSharedPtr<const FManagedArrayCollection> RenderCollection,
 		TObjectPtr<UObject> NodeOwner, TObjectPtr<AActor> RootActor, TArray<TObjectPtr<UPrimitiveComponent>>& PrimitiveComponents) override;
+	virtual void OnInvalidate() override;
 #if WITH_EDITOR
 	virtual bool CanDebugDraw() const override { return true; }
 	virtual bool CanDebugDrawViewMode(const FName& ViewModeName) const override;
@@ -104,4 +108,7 @@ private:
 	/** Transient skeletal mesh built from dataflow render collection */
 	UPROPERTY(Transient)
 	TObjectPtr<USkeletalMesh> SkeletalMesh = nullptr;
+
+	/** Valid skeletal mesh boolean to trigger the construction */
+	bool bValidSkeletalMesh = false;
 };
