Dataflow : Add asset util method to duplicate a graph between different graphs within the same asset

diff --git a/Engine/Plugins/Experimental/Dataflow/Source/DataflowEditor/Private/Dataflow/DataflowAssetEditUtils.cpp b/Engine/Plugins/Experimental/Dataflow/Source/DataflowEditor/Private/Dataflow/DataflowAssetEditUtils.cpp
--- a/Engine/Plugins/Experimental/Dataflow/Source/DataflowEditor/Private/Dataflow/DataflowAssetEditUtils.cpp
+++ b/Engine/Plugins/Experimental/Dataflow/Source/DataflowEditor/Private/Dataflow/DataflowAssetEditUtils.cpp
@@ -463,11 +463,18 @@ namespace UE::Dataflow
 
 	//////////////////////////////////////////////////////////////////////////////////////////////////////////
 	void FEditAssetUtils::DuplicateNodes(UEdGraph* EdGraph, const TArray<UEdGraphNode*>& EdNodesToDuplicate, const FVector2D& Location, TArray<UEdGraphNode*>& OutDuplicatedNodes)
+	{
+		TMap<FGuid, FGuid> NodeGuidMap;
+		DuplicateNodes(EdGraph, EdNodesToDuplicate, EdGraph, Location, OutDuplicatedNodes, NodeGuidMap);
+	}
+
+	void FEditAssetUtils::DuplicateNodes(UEdGraph* SourceEdGraph, const TArray<UEdGraphNode*>& EdNodesToDuplicate, UEdGraph* TargetEdGraph, const FVector2D& Location, TArray<UEdGraphNode*>& OutDuplicatedNodes, TMap<FGuid, FGuid>& OutNodeGuidMap)
 	{
 		using namespace EditAssetUtils::Private;
 
-		UDataflow* DataflowAsset = UDataflow::GetDataflowAssetFromEdGraph(EdGraph);
-		if (!DataflowAsset)
+		UDataflow* SourceDataflowAsset = UDataflow::GetDataflowAssetFromEdGraph(SourceEdGraph);
+		UDataflow* TargetDataflowAsset = UDataflow::GetDataflowAssetFromEdGraph(TargetEdGraph);
+		if (!SourceDataflowAsset || !TargetDataflowAsset)
 		{
 			// no graph to copy to
 			return;
@@ -488,8 +495,8 @@ namespace UE::Dataflow
 		};
 
 		auto DuplicateNodesInternal =
-			[&EdGraph, &Location, &RefLocation, &EdNodesToDuplicate, &OutDuplicatedNodes]
-			(UDataflow& DataflowAsset) -> EChangeResult
+			[&OutNodeGuidMap, &SourceEdGraph, TargetEdGraph, &Location, &RefLocation, &EdNodesToDuplicate, &OutDuplicatedNodes]
+			(UDataflow& TargetDataflowAsset) -> EChangeResult
 			{
 				TMap<FGuid, UDataflowEdNode*> EdNodeMap;
 				TMap<FGuid, FGuid> NodeGuidMap;
@@ -511,11 +518,11 @@ namespace UE::Dataflow
 							const FName NodeName = NodeToDuplicate->GetName();
 							const FName NodeTypeName = NodeToDuplicate->GetType();
 
-							if (TSharedPtr<FDataflowNode> DataflowNode = AddDataflowNode(&DataflowAsset, NodeName, NodeTypeName))
+							if (TSharedPtr<FDataflowNode> DataflowNode = AddDataflowNode(&TargetDataflowAsset, NodeName, NodeTypeName))
 							{
 								SDataflowEdNode::CopyDataflowNodeSettings(NodeToDuplicate, DataflowNode);
 	
-								if (UDataflowEdNode* EdNode = CreateDataflowEdNode(EdGraph, DataflowNode, NodeLocation, /*FromPin*/nullptr))
+								if (UDataflowEdNode* EdNode = CreateDataflowEdNode(TargetEdGraph, DataflowNode, NodeLocation, /*FromPin*/nullptr))
 								{
 									EdNodeMap.Add(DataflowEdNodeToDuplicate->DataflowNodeGuid, EdNode);
 									NodeGuidMap.Add(DataflowEdNodeToDuplicate->DataflowNodeGuid, EdNode->DataflowNodeGuid);
@@ -531,7 +538,7 @@ namespace UE::Dataflow
 							(double)CommentEdNodeToDuplicate->NodeWidth,
 							(double)CommentEdNodeToDuplicate->NodeHeight
 						};
-						if (UEdGraphNode_Comment* CommentEdNode = CreateCommentEdNode(EdGraph, NodeLocation, CommentEdNodeToDuplicate->NodeComment, CommentSize, CommentEdNodeToDuplicate->CommentColor, CommentEdNodeToDuplicate->FontSize))
+						if (UEdGraphNode_Comment* CommentEdNode = CreateCommentEdNode(TargetEdGraph, NodeLocation, CommentEdNodeToDuplicate->NodeComment, CommentSize, CommentEdNodeToDuplicate->CommentColor, CommentEdNodeToDuplicate->FontSize))
 						{
 							OutDuplicatedNodes.Add(CommentEdNode);
 						}
@@ -558,14 +565,14 @@ namespace UE::Dataflow
 									if (NodeGuidMap.Contains(DataflowNodeBGuid))
 									{
 										const FName InputputName = Connection->GetName();
-										if (TSharedPtr<FDataflowNode> DuplicatedDataflowNodeA = DataflowAsset.GetDataflow()->FindBaseNode(NodeGuidMap[DataflowNodeAGuid]))
+										if (TSharedPtr<FDataflowNode> DuplicatedDataflowNodeA = TargetDataflowAsset.GetDataflow()->FindBaseNode(NodeGuidMap[DataflowNodeAGuid]))
 										{
 											FDataflowOutput* OutputConnection = DuplicatedDataflowNodeA->FindOutput(OutputputName);
-											if (TSharedPtr<FDataflowNode> DuplicatedDataflowNodeB = DataflowAsset.GetDataflow()->FindBaseNode(NodeGuidMap[DataflowNodeBGuid]))
+											if (TSharedPtr<FDataflowNode> DuplicatedDataflowNodeB = TargetDataflowAsset.GetDataflow()->FindBaseNode(NodeGuidMap[DataflowNodeBGuid]))
 											{
 												FDataflowInput* InputConnection = DuplicatedDataflowNodeB->FindInput(InputputName);
 
-												DataflowAsset.GetDataflow()->Connect(OutputConnection, InputConnection);
+												TargetDataflowAsset.GetDataflow()->Connect(OutputConnection, InputConnection);
 
 												// Connect the UDataflowEdNode FPins as well
 												if (UEdGraphPin* OutputPin = FindPin(EdNodeMap[DataflowNodeAGuid], EEdGraphPinDirection::EGPD_Output, OutputputName))
@@ -583,10 +590,11 @@ namespace UE::Dataflow
 						}
 					}
 				}
+				OutNodeGuidMap = MoveTemp(NodeGuidMap);
 				return EChangeResult::Changed;
 			};
 
-		ChangeDataflowAssetWithTransaction(DataflowAsset, TransactionName, DuplicateNodesInternal, NAME_None);
+		ChangeDataflowAssetWithTransaction(TargetDataflowAsset, TransactionName, DuplicateNodesInternal, NAME_None);
 	}
 
 	//////////////////////////////////////////////////////////////////////////////////////////////////////////
diff --git a/Engine/Plugins/Experimental/Dataflow/Source/DataflowEditor/Private/Dataflow/DataflowAssetEditUtils.h b/Engine/Plugins/Experimental/Dataflow/Source/DataflowEditor/Private/Dataflow/DataflowAssetEditUtils.h
--- a/Engine/Plugins/Experimental/Dataflow/Source/DataflowEditor/Private/Dataflow/DataflowAssetEditUtils.h
+++ b/Engine/Plugins/Experimental/Dataflow/Source/DataflowEditor/Private/Dataflow/DataflowAssetEditUtils.h
@@ -67,14 +67,27 @@ namespace UE::Dataflow
 		static void DeleteNodes(UEdGraph* EdGraph, const TArray<UEdGraphNode*>& NodesToDelete);
 
 		/**
-		* Duplicate a node
+		* Duplicate nodes within the same graph
 		* @param EdGraph			Main Graph or Subgraph to duplicate the node in 
 		* @param NodeToDuplicate	List of nodes to duplicate
 		* @param Location			Location on the graph where to place the node 
+		* @param OutDuplicatedNodes	List of newly created nodes
 		* @return the newly created EdGraphNode
 		*/
 		static void DuplicateNodes(UEdGraph* EdGraph, const TArray<UEdGraphNode*>& NodesToDuplicate, const FVector2D& Location, TArray<UEdGraphNode*>& OutDuplicatedNodes);
 
+		/**
+		* Duplicate nodes from one graph to another ( assuming same dataflow asset ) 
+		* @param SourceEdGraph		Main Graph or Subgraph to duplicate the node from
+		* @param NodeToDuplicate	List of nodes to duplicate
+		* @param TargetEdGraph		Main Graph or Subgraph to duplicate the node to
+		* @param Location			Location on the graph where to place the node
+		* @param OutDuplicatedNodes	List of newly created nodes
+		* @param OutNodeGuidMap		Mapping of dup[licated -> new node guids 
+		* @return the newly created EdGraphNode
+		*/
+		static void DuplicateNodes(UEdGraph* SourceEdGraph, const TArray<UEdGraphNode*>& NodesToDuplicate, UEdGraph* TargetEdGraph, const FVector2D& Location, TArray<UEdGraphNode*>& OutDuplicatedNodes, TMap<FGuid, FGuid>& OutNodeGuidMap);
+
 		/**
 		* Copy a list of nodes and their connection to the clipboard
 		*/
