fix dataflow editor auto-focusing the camera too  frequently -- make it only auto-focus the first time it sees a valid mesh in 2D/3D construction views

diff --git a/Engine/Plugins/Experimental/Dataflow/Source/DataflowEditor/Private/Dataflow/DataflowEditorMode.cpp b/Engine/Plugins/Experimental/Dataflow/Source/DataflowEditor/Private/Dataflow/DataflowEditorMode.cpp
--- a/Engine/Plugins/Experimental/Dataflow/Source/DataflowEditor/Private/Dataflow/DataflowEditorMode.cpp
+++ b/Engine/Plugins/Experimental/Dataflow/Source/DataflowEditor/Private/Dataflow/DataflowEditorMode.cpp
@@ -101,6 +101,9 @@ void UDataflowEditorMode::Enter()
 		FEngineAnalytics::GetProvider().RecordEvent(TEXT("Editor.Usage.DataflowEditor.Enter"), EventAttributes);
 	}
 
+	// Reset tracking of whether we've seen a valid mesh in 2D or 3D viewports
+	bFirstValid2DMesh = true;
+	bFirstValid3DMesh = true;
 }
 
 void UDataflowEditorMode::SetDataflowEditor(UDataflowEditor* InDataflowEditor) 
@@ -627,7 +630,7 @@ void UDataflowEditorMode::ModeTick(float DeltaTime)
 		if (EditorContent->IsConstructionDirty())
 		{
 			ConstructionScene->UpdateConstructionScene();
-			RefocusConstructionViewportClient();
+			FirstTimeFocusConstructionViewport();
 		}
 	}
 
