Skip the node evaluation when changing properties in manual mode and in rendering callbacks

diff --git a/Engine/Plugins/Experimental/Dataflow/Source/DataflowEditor/Private/Dataflow/DataflowConstructionScene.cpp b/Engine/Plugins/Experimental/Dataflow/Source/DataflowEditor/Private/Dataflow/DataflowConstructionScene.cpp
--- a/Engine/Plugins/Experimental/Dataflow/Source/DataflowEditor/Private/Dataflow/DataflowConstructionScene.cpp
+++ b/Engine/Plugins/Experimental/Dataflow/Source/DataflowEditor/Private/Dataflow/DataflowConstructionScene.cpp
@@ -6,6 +6,7 @@
 #include "Dataflow/CollectionRenderingPatternUtility.h"
 #include "Dataflow/DataflowEditorCollectionComponent.h"
 #include "Dataflow/DataflowEditor.h"
+#include "Dataflow/DataflowEditorToolkit.h"
 #include "Dataflow/DataflowObject.h"
 #include "Dataflow/DataflowEditorStyle.h"
 #include "Dataflow/DataflowRenderingFactory.h"
@@ -13,6 +14,7 @@
 #include "Elements/Framework/EngineElementsLibrary.h"
 #include "Selection.h"
 #include "AssetViewerSettings.h"
+#include "DataflowEditorOptions.h"
 #include "Dataflow/DataflowPrimitiveNode.h"
 #include "GeometryCollection/Facades/CollectionRenderingFacade.h"
 #include "Selection/GeometrySelector.h"
@@ -166,6 +168,9 @@ void FDataflowConstructionScene::UpdateDynamicMeshComponents()
 	// The preview scene for the construction view will be
 	// cleared and rebuilt from scratch. This will generate a 
 	// list of UPrimitiveComponents for rendering.
+
+	const FDataflowEditorToolkit* EditorToolkit = DataflowEditor ? static_cast<FDataflowEditorToolkit*>(DataflowEditor->GetInstanceInterface()) : nullptr;
+	const bool bEvaluateOutputs = EditorToolkit ? (EditorToolkit->GetEvaluationMode() != EDataflowEditorEvaluationMode::Manual) : true;
 	
 	if (const TObjectPtr<UDataflowBaseContent>& EditorContent = GetEditorContent())
 	{
@@ -180,8 +185,8 @@ void FDataflowConstructionScene::UpdateDynamicMeshComponents()
 					TSharedPtr<FManagedArrayCollection> RenderCollection(new FManagedArrayCollection);
 					GeometryCollection::Facades::FRenderingFacade RenderingFacade(*RenderCollection);
 					RenderingFacade.DefineSchema();
-
-					const bool bHasPrimitves = UE::Dataflow::RenderNodeOutput(RenderingFacade, *Target, *EditorContent);
+					
+					const bool bHasPrimitves = UE::Dataflow::RenderNodeOutput(RenderingFacade, *Target, *EditorContent, bEvaluateOutputs);
 
 					if (Target == EditorContent->GetSelectedNode())
 					{
@@ -268,7 +273,7 @@ void FDataflowConstructionScene::UpdateDynamicMeshComponents()
 					GeometryCollection::Facades::FRenderingFacade RenderingFacade(*RenderCollection);
 					RenderingFacade.DefineSchema();
 
-					const bool bHasPrimitves = UE::Dataflow::RenderNodeOutput(RenderingFacade, *Target, *EditorContent);
+					const bool bHasPrimitves = UE::Dataflow::RenderNodeOutput(RenderingFacade, *Target, *EditorContent, bEvaluateOutputs);
 
 					if (Target == EditorContent->GetSelectedNode())
 
diff --git a/Engine/Plugins/Experimental/Dataflow/Source/DataflowEditor/Private/Dataflow/DataflowEditorToolkit.cpp b/Engine/Plugins/Experimental/Dataflow/Source/DataflowEditor/Private/Dataflow/DataflowEditorToolkit.cpp
--- a/Engine/Plugins/Experimental/Dataflow/Source/DataflowEditor/Private/Dataflow/DataflowEditorToolkit.cpp
+++ b/Engine/Plugins/Experimental/Dataflow/Source/DataflowEditor/Private/Dataflow/DataflowEditorToolkit.cpp
@@ -1241,8 +1241,10 @@ void FDataflowEditorToolkit::OnPropertyValueChanged(const FPropertyChangedEvent&
 			
 			FDataflowEditorCommands::OnPropertyValueChanged(DataflowAsset, DataflowContext, LastNodeTimestamp, PropertyChangedEvent, SelectedDataflowNodes);
 
-			EditorContent->SetDataflowContext(DataflowContext);
-			EditorContent->SetLastModifiedTimestamp(LastNodeTimestamp);
+			// For manual evaluation no need to invalidate the construction scene if we change any parameters
+			// since it will trigger the rendering callbacks with the same cached collection
+			const bool bMakesDirty = (EvaluationMode != EDataflowEditorEvaluationMode::Manual);
+			EditorContent->SetLastModifiedTimestamp(LastNodeTimestamp, bMakesDirty);
 
 			// Refresh graph display to display nodes with warning/error
 			GraphEditor->NotifyGraphChanged();
diff --git a/Engine/Plugins/Experimental/Dataflow/Source/DataflowEditor/Private/Dataflow/DataflowEditorUtil.cpp b/Engine/Plugins/Experimental/Dataflow/Source/DataflowEditor/Private/Dataflow/DataflowEditorUtil.cpp
--- a/Engine/Plugins/Experimental/Dataflow/Source/DataflowEditor/Private/Dataflow/DataflowEditorUtil.cpp
+++ b/Engine/Plugins/Experimental/Dataflow/Source/DataflowEditor/Private/Dataflow/DataflowEditorUtil.cpp
@@ -142,7 +142,7 @@ namespace UE::Dataflow
 		return false;
 	}
 
-	bool RenderNodeOutput(GeometryCollection::Facades::FRenderingFacade& RenderingFacade, const UDataflowEdNode& Node, const UDataflowBaseContent& EditorContent)
+	bool RenderNodeOutput(GeometryCollection::Facades::FRenderingFacade& RenderingFacade, const UDataflowEdNode& Node, const UDataflowBaseContent& EditorContent, const bool bEvaluateOutputs)
 	{
 		const TObjectPtr<UDataflow>& DataflowAsset = EditorContent.GetDataflowAsset();
 		const IDataflowConstructionViewMode* ConstructionViewMode = EditorContent.GetConstructionViewMode();
@@ -159,7 +159,7 @@ namespace UE::Dataflow
 					{
 						if (TSharedPtr<const FDataflowNode> NodeTarget = Graph->FindBaseNode(FName(Node.GetName())))
 						{
-							Factory->RenderNodeOutput(RenderingFacade, FGraphRenderingState{ Node.GetDataflowNodeGuid(), NodeTarget.Get(), Parameter, *DataflowContext, *ConstructionViewMode });
+							Factory->RenderNodeOutput(RenderingFacade, FGraphRenderingState{ Node.GetDataflowNodeGuid(), NodeTarget.Get(), Parameter, *DataflowContext, *ConstructionViewMode, bEvaluateOutputs });
 							bHasPrimitives = bHasPrimitives || NodeTarget->HasPrimitives();
 						}
 					}
diff --git a/Engine/Plugins/Experimental/Dataflow/Source/DataflowEditor/Private/Dataflow/DataflowEditorUtil.h b/Engine/Plugins/Experimental/Dataflow/Source/DataflowEditor/Private/Dataflow/DataflowEditorUtil.h
--- a/Engine/Plugins/Experimental/Dataflow/Source/DataflowEditor/Private/Dataflow/DataflowEditorUtil.h
+++ b/Engine/Plugins/Experimental/Dataflow/Source/DataflowEditor/Private/Dataflow/DataflowEditorUtil.h
@@ -55,6 +55,6 @@ namespace UE::Dataflow
 
 	/** This function fills the managed array collection linked to the rendering facade from the callbacks.
 	 * Returns true if the resulting collection should be directly rendered of false if the collection will be used to generate a primitive component */
-	bool RenderNodeOutput(GeometryCollection::Facades::FRenderingFacade& RenderingFacade, const UDataflowEdNode& Node, const UDataflowBaseContent& EditorContent);
+	bool RenderNodeOutput(GeometryCollection::Facades::FRenderingFacade& RenderingFacade, const UDataflowEdNode& Node, const UDataflowBaseContent& EditorContent, const bool bEvaluateOutputs);
 	void GetViewModesForNode(const UDataflowEdNode& EdNode, const UDataflowBaseContent& EditorContent, TArray<FName>& OutValidViewModeNames);
 }
\ No newline at end of file
diff --git a/Engine/Plugins/Experimental/Dataflow/Source/DataflowEditor/Public/Dataflow/DataflowEditorToolkit.h b/Engine/Plugins/Experimental/Dataflow/Source/DataflowEditor/Public/Dataflow/DataflowEditorToolkit.h
--- a/Engine/Plugins/Experimental/Dataflow/Source/DataflowEditor/Public/Dataflow/DataflowEditorToolkit.h
+++ b/Engine/Plugins/Experimental/Dataflow/Source/DataflowEditor/Public/Dataflow/DataflowEditorToolkit.h
@@ -105,6 +105,9 @@ public:
 	UDataflowSubGraph* GetSubGraph(FName SubGraphName) const;
 
 	const FString& GetDebugDrawOverlayString() const;
+	
+	/** Get the toolkit evaluation mode */
+	const EDataflowEditorEvaluationMode& GetEvaluationMode() const { return EvaluationMode; }
 
 protected:
 
diff --git a/Engine/Plugins/Experimental/Dataflow/Source/DataflowEditor/Public/Dataflow/DataflowRenderingFactory.h b/Engine/Plugins/Experimental/Dataflow/Source/DataflowEditor/Public/Dataflow/DataflowRenderingFactory.h
--- a/Engine/Plugins/Experimental/Dataflow/Source/DataflowEditor/Public/Dataflow/DataflowRenderingFactory.h
+++ b/Engine/Plugins/Experimental/Dataflow/Source/DataflowEditor/Public/Dataflow/DataflowRenderingFactory.h
@@ -21,7 +21,8 @@ namespace UE::Dataflow
 
 	struct FGraphRenderingState 
 	{
-		FGraphRenderingState(const FGuid InGuid, const FDataflowNode* InNode, const FRenderingParameter& InParameters, UE::Dataflow::FContext& InContext, const UE::Dataflow::IDataflowConstructionViewMode& ViewMode)
+		FGraphRenderingState(const FGuid InGuid, const FDataflowNode* InNode, const FRenderingParameter& InParameters, UE::Dataflow::FContext& InContext,
+			const UE::Dataflow::IDataflowConstructionViewMode& ViewMode, const bool bInEvaluateOutputs = true)
 			: NodeGuid(InGuid)
 			, Node(InNode)
 			, RenderName(InParameters.Name)
@@ -29,6 +30,7 @@ namespace UE::Dataflow
 			, RenderOutputs(InParameters.Outputs)
 			, Context(InContext)
 			, ViewMode(ViewMode)
+			, bEvaluateOutputs(bInEvaluateOutputs)
 		{}
 
 		const FGuid& GetGuid() const { return NodeGuid; }
@@ -43,7 +45,7 @@ namespace UE::Dataflow
 			{
 				if (const FDataflowOutput* Output = Node->FindOutput(OutputName))
 				{
-					if (Context.IsThreaded())
+					if (Context.IsThreaded() || !bEvaluateOutputs)
 					{
 						return Output->ReadValue<T>(Context, Default);
 					}
@@ -66,6 +68,9 @@ namespace UE::Dataflow
 		UE::Dataflow::FContext& Context;
 
 		const UE::Dataflow::IDataflowConstructionViewMode& ViewMode;
+
+		/** Boolean to check if we need to reevaluate the outputs or read the cached value */
+		bool bEvaluateOutputs = true;
 	};
 
 
