Dataflow : Handle change of the dataflow asset in the asset details
- Change the graph accordingly and deselect all nodes to make sure views are up to date wit hthe new dataflow
- Also handle no dataflow ( empty graph )

diff --git a/Engine/Plugins/Experimental/Dataflow/Source/DataflowEditor/Private/Dataflow/DataflowEditorToolkit.cpp b/Engine/Plugins/Experimental/Dataflow/Source/DataflowEditor/Private/Dataflow/DataflowEditorToolkit.cpp
--- a/Engine/Plugins/Experimental/Dataflow/Source/DataflowEditor/Private/Dataflow/DataflowEditorToolkit.cpp
+++ b/Engine/Plugins/Experimental/Dataflow/Source/DataflowEditor/Private/Dataflow/DataflowEditorToolkit.cpp
@@ -1258,6 +1258,56 @@ void FDataflowEditorToolkit::OnAssetPropertyValueChanged(const FPropertyChangedE
 {
 	if (const TObjectPtr<UDataflowBaseContent>& EditorContent = GetEditorContent())
 	{
+		if (IDataflowInstanceInterface* DataflowInstanceInterface = Cast<IDataflowInstanceInterface>(EditorContent->GetDataflowOwner()))
+		{
+			if (PropertyChangedEvent.GetPropertyName()== FDataflowInstance::GetDataflowAssetPropertyName())
+			{
+				// close all subgraphs tabs and selection to make sure no other widget refer to this graphs data
+				const TArray<TSharedPtr<SDockTab>> Tabs = DocumentManager->GetAllDocumentTabs();
+				for (TSharedPtr<SDockTab> Tab : Tabs)
+				{
+					if (Tab)
+					{
+						Tab->RequestCloseTab();
+					}
+				}
+				if (GraphEditor)
+				{
+					GraphEditor->ClearSelectionSet();
+				}
+				bViewsNeedRefresh = true;
+				
+				// change the asset and open a new graph in the tab 
+				if (UDataflow* NewDataflowAsset = DataflowInstanceInterface->GetDataflowInstance().GetDataflowAsset())
+				{
+					NewDataflowAsset->Schema = UDataflowSchema::StaticClass();
+					EditorContent->SetDataflowAsset(NewDataflowAsset);
+					GraphEditor = CreateGraphEditorWidget(NewDataflowAsset, NodeDetailsEditor);
+					if (GraphEditorTab)
+					{
+						if (GraphEditor)
+						{
+							GraphEditorTab->SetContent(GraphEditor.ToSharedRef());
+						}
+						else
+						{
+							GraphEditorTab->SetContent(SNew(SSpacer));
+						}
+					}
+				}
+				else
+				{
+					// Clear the GraphEditor area
+					// (Can't have a SDataflowGraphEditor with a null UDataflow, so just put down Spacers if we have no Dataflow)
+					GraphEditor.Reset();
+					if (GraphEditorTab)
+					{
+						GraphEditorTab->SetContent(SNew(SSpacer));
+					}
+				}
+			}
+		}
+
 		ensure(EditorContent);	FDataflowEditorCommands::OnAssetPropertyValueChanged(EditorContent, PropertyChangedEvent);
 	}
 }
@@ -2175,13 +2225,15 @@ TSharedRef<SDockTab> FDataflowEditorToolkit::SpawnTab_GraphCanvas(const FSpawnTa
 {
 	check(Args.GetTabId() == GraphCanvasTabId);
 
-	TSharedRef<SDockTab> SpawnedTab = SNew(SDockTab)
-		.Label(LOCTEXT("DataflowEditor_DataflowGraph_TabTitle", "Dataflow Graph"))
-		[
-			GraphEditor.ToSharedRef()
-		];
+	SAssignNew(GraphEditorTab, SDockTab)
+		.Label(LOCTEXT("DataflowEditor_DataflowGraph_TabTitle", "Dataflow Graph"));
 
-	return SpawnedTab;
+	if (GraphEditor)
+	{
+		GraphEditorTab.Get()->SetContent(GraphEditor.ToSharedRef());
+	}
+
+	return GraphEditorTab.ToSharedRef();
 }
 
 TSharedRef<SDockTab> FDataflowEditorToolkit::SpawnTab_SubGraphTab(const FSpawnTabArgs& Args)
diff --git a/Engine/Plugins/Experimental/Dataflow/Source/DataflowEditor/Public/Dataflow/DataflowEditorToolkit.h b/Engine/Plugins/Experimental/Dataflow/Source/DataflowEditor/Public/Dataflow/DataflowEditorToolkit.h
--- a/Engine/Plugins/Experimental/Dataflow/Source/DataflowEditor/Public/Dataflow/DataflowEditorToolkit.h
+++ b/Engine/Plugins/Experimental/Dataflow/Source/DataflowEditor/Public/Dataflow/DataflowEditorToolkit.h
@@ -208,6 +208,7 @@ private:
 	TSharedPtr<SDataflowConstructionViewport> DataflowConstructionViewport;
 	TSharedPtr<SDataflowSimulationViewport> DataflowSimulationViewport;
 	TSharedPtr<SDataflowGraphEditor> GraphEditor;
+	TSharedPtr<SDockTab> GraphEditorTab;
 	TSharedPtr<IStructureDetailsView> NodeDetailsEditor;
 	TSharedPtr<SDataflowMembersWidget> MembersWidget;
 	TSharedPtr<UE::Dataflow::FDataflowNodeDetailExtensionHandler> NodeDetailsExtensionHandler;
