Disallow really long columns in DataflowCollectionSpreadSheet.

diff --git a/Engine/Plugins/Experimental/Dataflow/Source/DataflowEditor/Private/Dataflow/DataflowCollectionSpreadSheetHelpers.cpp b/Engine/Plugins/Experimental/Dataflow/Source/DataflowEditor/Private/Dataflow/DataflowCollectionSpreadSheetHelpers.cpp
--- a/Engine/Plugins/Experimental/Dataflow/Source/DataflowEditor/Private/Dataflow/DataflowCollectionSpreadSheetHelpers.cpp
+++ b/Engine/Plugins/Experimental/Dataflow/Source/DataflowEditor/Private/Dataflow/DataflowCollectionSpreadSheetHelpers.cpp
@@ -245,6 +245,13 @@ namespace UE::Dataflow::CollectionSpreadSheetHelpers
 			ValueAsString = "<Unknown Data Type>";
 		}
 
+		// Clip really long strings so slate doesn't die.
+		constexpr int32 MaxStringLength = 10000;
+		if (ValueAsString.Len() > MaxStringLength)
+		{
+			ValueAsString.LeftInline(MaxStringLength);
+			ValueAsString += "...";
+		}
 		return ValueAsString;
 	}
 
