Dataflow : fix freeze customization potential crash when cutting at least two nodes from the graph
- Changed the lambdas to use a weak pointer on the node instead of retrieving the node form the property handle
- this allow to validate that the node is still available at the time of evaluating the lambdas

diff --git a/Engine/Plugins/Experimental/Dataflow/Source/DataflowEditor/Private/Dataflow/DataflowFreezeActionsCustomization.cpp b/Engine/Plugins/Experimental/Dataflow/Source/DataflowEditor/Private/Dataflow/DataflowFreezeActionsCustomization.cpp
--- a/Engine/Plugins/Experimental/Dataflow/Source/DataflowEditor/Private/Dataflow/DataflowFreezeActionsCustomization.cpp
+++ b/Engine/Plugins/Experimental/Dataflow/Source/DataflowEditor/Private/Dataflow/DataflowFreezeActionsCustomization.cpp
@@ -30,7 +30,10 @@ namespace UE::Dataflow
 						{
 							if (Struct->IsChildOf<FDataflowNode>())
 							{
-								return reinterpret_cast<FDataflowNode*>(StructOnScope->GetStructMemory());
+								if (StructOnScope->GetStructMemory())
+								{
+									return reinterpret_cast<FDataflowNode*>(StructOnScope->GetStructMemory());
+								}
 							}
 						}
 					}
@@ -59,23 +62,30 @@ namespace UE::Dataflow
 			.UseAllottedSize(true)
 		];
 
+		
+		TWeakPtr<FDataflowNode> WeakDataflowNode;
+		if (FDataflowNode* DataflowNode = Private::GetDataflowNode(StructPropertyHandle))
+		{
+			WeakDataflowNode = DataflowNode->AsWeak();
+		}
+		
 		// Create button image + text
 		WrapBox->AddSlot()
 			.Padding(0.f, 0.f, 5.f, 3.f)
 			[
 				SNew(SButton)
 					.Text(LOCTEXT("Freeze", "Freeze"))
-					.ToolTipText_Lambda([StructPropertyHandle]() -> FText
+					.ToolTipText_Lambda([WeakDataflowNode]() -> FText
 						{
-							const FDataflowNode* const DataflowNode = Private::GetDataflowNode(StructPropertyHandle);
+							const TSharedPtr<const FDataflowNode> DataflowNode = WeakDataflowNode.Pin();
 							return DataflowNode && !DataflowNode->IsFrozen() ?
 								LOCTEXT("FreezeToolTip", "Freeze all outputs, and disable evaluation for this node.") :
 								LOCTEXT("UnfreezeToolTip", "UInfreeze all outputs, and re-enable evaluation for this node.");
 						})
 					.OnClicked_Lambda(
-						[StructPropertyHandle, DataflowGraphEditor]() -> FReply
+						[WeakDataflowNode, DataflowGraphEditor]() -> FReply
 						{
-							if (FDataflowNode* const DataflowNode = Private::GetDataflowNode(StructPropertyHandle))
+							if (const TSharedPtr<FDataflowNode> DataflowNode = WeakDataflowNode.Pin())
 							{
 								// Retrieve context if any
 								const TSharedPtr<const SDataflowGraphEditor> DataflowGraphEditorPtr = DataflowGraphEditor.Pin();
@@ -94,9 +104,9 @@ namespace UE::Dataflow
 							return FReply::Handled();
 						})
 					.ContentPadding(FMargin(0.f, 2.f))  // Too much horizontal padding otherwise (default is 4, 2)
-					.Visibility_Lambda([StructPropertyHandle]() -> EVisibility
+					.Visibility_Lambda([WeakDataflowNode]() -> EVisibility
 						{
-							const FDataflowNode* const DataflowNode = Private::GetDataflowNode(StructPropertyHandle);
+							const TSharedPtr<const FDataflowNode> DataflowNode = WeakDataflowNode.Pin();
 							return DataflowNode && !DataflowNode->AsType<const FDataflowTerminalNode>() ? EVisibility::Visible : EVisibility::Collapsed;
 						})
 					[
@@ -117,9 +127,9 @@ namespace UE::Dataflow
 							.VAlign(VAlign_Center)
 							[
 								SNew(STextBlock)
-									.Text_Lambda([StructPropertyHandle]() -> FText
+									.Text_Lambda([WeakDataflowNode]() -> FText
 										{
-											const FDataflowNode* const DataflowNode = Private::GetDataflowNode(StructPropertyHandle);
+											const TSharedPtr<const FDataflowNode> DataflowNode = WeakDataflowNode.Pin();
 											return DataflowNode && !DataflowNode->IsFrozen() ?
 												LOCTEXT("Freeze", "Freeze") :
 												LOCTEXT("Unfreeze", "Unfreeze");
@@ -135,9 +145,9 @@ namespace UE::Dataflow
 					.Text(LOCTEXT("Refreeze", "Refreeze"))
 					.ToolTipText(LOCTEXT("RefreezeToolTip", "Unfreeze all outputs, redo this node's evaluation, and freeze all updated outputs again."))
 					.OnClicked_Lambda(
-						[StructPropertyHandle, DataflowGraphEditor]() -> FReply
+						[WeakDataflowNode, DataflowGraphEditor]() -> FReply
 						{
-							if (FDataflowNode* const DataflowNode = Private::GetDataflowNode(StructPropertyHandle))
+							if (const TSharedPtr<FDataflowNode> DataflowNode = WeakDataflowNode.Pin())
 							{
 								if (DataflowNode->IsFrozen())
 								{
@@ -148,21 +158,21 @@ namespace UE::Dataflow
 										TSharedPtr<UE::Dataflow::FContext>();
 									// Refreeze
 									DataflowNode->Unfreeze(*Context);
-									Context->Evaluate(DataflowNode, nullptr);
+									Context->Evaluate(DataflowNode.Get(), nullptr);
 									DataflowNode->Freeze(*Context);
 								}
 							}
 							return FReply::Handled();
 						})
 					.ContentPadding(FMargin(0.f, 2.f))  // Too much horizontal padding otherwise (default is 4, 2)
-					.IsEnabled_Lambda([StructPropertyHandle]() -> bool
+					.IsEnabled_Lambda([WeakDataflowNode]() -> bool
 						{
-							const FDataflowNode* const DataflowNode = Private::GetDataflowNode(StructPropertyHandle);
+							const TSharedPtr<const FDataflowNode> DataflowNode = WeakDataflowNode.Pin();
 							return DataflowNode && DataflowNode->IsFrozen();
 						})
-					.Visibility_Lambda([StructPropertyHandle]() -> EVisibility
+					.Visibility_Lambda([WeakDataflowNode]() -> EVisibility
 						{
-							const FDataflowNode* const DataflowNode = Private::GetDataflowNode(StructPropertyHandle);
+							const TSharedPtr<const FDataflowNode> DataflowNode = WeakDataflowNode.Pin();
 							return DataflowNode && !DataflowNode->AsType<const FDataflowTerminalNode>() ? EVisibility::Visible : EVisibility::Collapsed;
 						})
 					[
