Dataflow editor :  make sure to have the variable and subgraph section expanded when the editor opens

diff --git a/Engine/Plugins/Experimental/Dataflow/Source/DataflowEditor/Private/Dataflow/DataflowMembersWidget.cpp b/Engine/Plugins/Experimental/Dataflow/Source/DataflowEditor/Private/Dataflow/DataflowMembersWidget.cpp
--- a/Engine/Plugins/Experimental/Dataflow/Source/DataflowEditor/Private/Dataflow/DataflowMembersWidget.cpp
+++ b/Engine/Plugins/Experimental/Dataflow/Source/DataflowEditor/Private/Dataflow/DataflowMembersWidget.cpp
@@ -310,6 +310,14 @@ void SDataflowMembersWidget::Construct(const FArguments& InArgs, TSharedPtr<FDat
 		]
 	];
 
+	// make sure th sections are expanded
+	TMap<int32, bool> ExpandedSections;
+	for (const TPair<int32, TSharedPtr<ISection>>& Section : SectionMap)
+	{
+		ExpandedSections.Add(Section.Key, true);
+	}
+	GraphActionMenu->SetSectionExpansion(ExpandedSections);
+
 	FCoreUObjectDelegates::OnObjectPropertyChanged.AddRaw(this, &SDataflowMembersWidget::OnObjectPropertyChanged);
 	FDataflowAssetDelegates::OnSubGraphsChanged.AddRaw(this, &SDataflowMembersWidget::OnSubGraphsChanged);
 }
