Dataflow : fix name collision when renaming a node from the graph
- Properly check that the name is unique within the dataflow UObject and its nodes

diff --git a/Engine/Plugins/Experimental/Dataflow/Source/DataflowEditor/Private/Dataflow/DataflowEditorCommands.cpp b/Engine/Plugins/Experimental/Dataflow/Source/DataflowEditor/Private/Dataflow/DataflowEditorCommands.cpp
--- a/Engine/Plugins/Experimental/Dataflow/Source/DataflowEditor/Private/Dataflow/DataflowEditorCommands.cpp
+++ b/Engine/Plugins/Experimental/Dataflow/Source/DataflowEditor/Private/Dataflow/DataflowEditorCommands.cpp
@@ -267,20 +267,36 @@ bool FDataflowEditorCommands::OnNodeVerifyTitleCommit(const FText& NewText, UEdG
 	}
 	if (GraphNode)
 	{
-		if (UDataflowEdNode* DataflowNode = Cast<UDataflowEdNode>(GraphNode))
+		// Comments are always valid because the text does not need to be unique
+		if (GraphNode->IsA<UEdGraphNode_Comment>())
 		{
-			if (TSharedPtr<UE::Dataflow::FGraph> Graph = DataflowNode->GetDataflowGraph())
+			return true;
+		}
+
+		// Normal node let's make sure that the name is unique
+		const FName NewNodeName(NewString);
+
+		// check that the name is unique within the asset itself
+		bool bIsUniqueSubObjectName = false;
+		if (UDataflow* DataflowAsset = UDataflow::GetDataflowAssetFromEdGraph(GraphNode->GetGraph()))
+		{
+			bIsUniqueSubObjectName = UE::Dataflow::FEditAssetUtils::IsUniqueDataflowSubObjectName(DataflowAsset, NewNodeName);
+		}
+
+		// check if the node already exists ( technically bIsUniqueSubObjectName should be enough, but maybe a node has been renamed outside of the normal code paths )
+		if (bIsUniqueSubObjectName)
+		{
+			if (UDataflowEdNode* DataflowNode = Cast<UDataflowEdNode>(GraphNode))
 			{
-				if( Graph->FindBaseNode(FName(NewString)).Get()==nullptr )
+				if (TSharedPtr<UE::Dataflow::FGraph> Graph = DataflowNode->GetDataflowGraph())
 				{
-					return true;
+					if (Graph->FindBaseNode(FName(NewString)).Get() == nullptr)
+					{
+						return true;
+					}
 				}
 			}
 		}
-		else if( Cast<UEdGraphNode_Comment>(GraphNode))
-		{
-			return true;
-		}
 	}
 	OutErrorMessage = FText::FromString(FString::Printf(TEXT("Non-unique name for graph node (%s)"), *NewString));
 	return false;
