Dataflow node to correct skin weights based on a selection weight map

diff --git a/Engine/Plugins/Experimental/Dataflow/Source/DataflowEditor/Private/Dataflow/DataflowEditorModule.cpp b/Engine/Plugins/Experimental/Dataflow/Source/DataflowEditor/Private/Dataflow/DataflowEditorModule.cpp
--- a/Engine/Plugins/Experimental/Dataflow/Source/DataflowEditor/Private/Dataflow/DataflowEditorModule.cpp
+++ b/Engine/Plugins/Experimental/Dataflow/Source/DataflowEditor/Private/Dataflow/DataflowEditorModule.cpp
@@ -18,12 +18,14 @@
 #include "Dataflow/DataflowToolRegistry.h"
 #include "DataflowEditorTools/DataflowEditorWeightMapPaintTool.h"
 #include "Dataflow/DataflowCollectionAddScalarVertexPropertyNode.h"
+
 #include "Dataflow/DataflowInstanceDetails.h"
 
 #include "PropertyEditorModule.h"
 #include "EditorModeRegistry.h"
 #include "Dataflow/DataflowCollectionEditSkinWeightsNode.h"
 #include "DataflowEditorTools/DataflowEditorSkinWeightsPaintTool.h"
+#include "DataflowEditorTools/DataflowEditorCorrectSkinWeightsNode.h"
 
 #define LOCTEXT_NAMESPACE "DataflowEditor"
 
@@ -131,6 +133,8 @@ void FDataflowEditorModule::StartupModule()
 	ToolRegistry.AddNodeToToolMapping(FDataflowCollectionEditSkinWeightsNode::StaticType(), 
 		NewObject<UDataflowEditorSkinWeightsPaintToolBuilder>(), Actions,
 		FSlateIcon(FName("DataflowEditorStyle"), FName("Dataflow.EditSkinWeights")), LOCTEXT("AddSkinWeightNodeButtonText", "Edit Skin Weights"));
+		
+	DATAFLOW_NODE_REGISTER_CREATION_FACTORY(FDataflowCorrectSkinWeightsNode);
 	
 	UE::Dataflow::FDataflowConstructionVisualizationRegistry& ConstructionVisualizationRegistry = UE::Dataflow::FDataflowConstructionVisualizationRegistry::GetInstance();
 	ConstructionVisualizationRegistry.RegisterVisualization(MakeUnique<UE::Dataflow::FMeshStatsConstructionVisualization>());
diff --git a/Engine/Plugins/Experimental/Dataflow/Source/DataflowEditor/Private/DataflowEditorTools/DataflowEditorCorrectSkinWeightsNode.cpp b/Engine/Plugins/Experimental/Dataflow/Source/DataflowEditor/Private/DataflowEditorTools/DataflowEditorCorrectSkinWeightsNode.cpp
new file mode 100644
--- /dev/null
+++ b/Engine/Plugins/Experimental/Dataflow/Source/DataflowEditor/Private/DataflowEditorTools/DataflowEditorCorrectSkinWeightsNode.cpp
@@ -0,0 +1,453 @@
+﻿// Copyright Epic Games, Inc. All Rights Reserved.
+
+#include "DataflowEditorTools/DataflowEditorCorrectSkinWeightsNode.h"
+
+#include "Chaos/Deformable/MuscleActivationConstraints.h"
+#include "Dataflow/DataflowInputOutput.h"
+#include "Dataflow/DataflowConnectionTypes.h"
+#include "Dataflow/DataflowNode.h"
+#include "Dataflow/DataflowRenderingFactory.h"
+#include "Dataflow/CollectionRenderingPatternUtility.h"
+#include "DynamicMesh/DynamicMesh3.h"
+#include "DynamicMesh/DynamicMeshAttributeSet.h"
+#include "DynamicMesh/DynamicVertexSkinWeightsAttribute.h"
+#include "GeometryCollection/Facades/CollectionRenderingFacade.h"
+#include "Operations/SmoothBoneWeights.h"
+#include "SkeletalMeshAttributes.h"
+#include "Parameterization/MeshDijkstra.h"
+
+#include UE_INLINE_GENERATED_CPP_BY_NAME(DataflowEditorCorrectSkinWeightsNode)
+
+#define LOCTEXT_NAMESPACE "DataflowEditorCorrectSkinWeights"
+
+
+namespace UE::Dataflow::Private
+{
+
+static void BuildDynamicMeshes(const FDataflowNode* DataflowNode, UE::Dataflow::FContext& Context, TArray<Geometry::FDynamicMesh3>& DynamicMeshes)
+{
+	const TArray<UE::Dataflow::FRenderingParameter> RenderParameters = DataflowNode->GetRenderParameters();
+	if(RenderParameters.Num() == 1)
+	{
+		TSharedPtr<FManagedArrayCollection> RenderCollection(new FManagedArrayCollection);
+		GeometryCollection::Facades::FRenderingFacade RenderingFacade(*RenderCollection);
+		RenderingFacade.DefineSchema();
+	
+		UE::Dataflow::FDataflowConstruction3DViewMode ViewMode;
+		FGraphRenderingState RenderingState(FGuid::NewGuid(), DataflowNode, RenderParameters.Last(), Context, ViewMode, false);
+		Dataflow::FRenderingFactory::GetInstance()->RenderNodeOutput(RenderingFacade, RenderingState);
+
+		const int32 NumGeometry = RenderingFacade.NumGeometry();
+		for (int32 MeshIndex = 0; MeshIndex < NumGeometry; ++MeshIndex)
+		{
+			UE::Geometry::FDynamicMesh3& DynamicMesh = DynamicMeshes.AddDefaulted_GetRef();
+			UE::Dataflow::Conversion::RenderingFacadeToDynamicMesh(RenderingFacade, MeshIndex, DynamicMesh);
+
+			DynamicMesh.Attributes()->AttachSkinWeightsAttribute(FName("Default"),
+				new UE::Geometry::FDynamicMeshVertexSkinWeightsAttribute(&DynamicMesh));
+		}
+	}
+}
+
+static void PruneSkinWeights(const float PruningThreshold, const TArray<TArray<int32>>& SetupIndices, const TArray<TArray<float>>& SetupWeights,
+	const TArray<float>& SelectionMap, TArray<TArray<int32>>& FinalIndices, TArray<TArray<float>>& FinalWeights)
+{
+	for(int32 VertexIndex = 0, NumVertices = FinalIndices.Num(); VertexIndex < NumVertices; ++VertexIndex)
+	{
+		if(SelectionMap[VertexIndex] > 0)
+		{
+			float TotalWeight = 0.0f;
+			for(int32 WeightIndex = 0, NumWeights = SetupWeights[VertexIndex].Num(); WeightIndex < NumWeights; ++WeightIndex)
+			{
+				if(SetupWeights[VertexIndex][WeightIndex] >= PruningThreshold)
+				{
+					FinalIndices[VertexIndex].Add(SetupIndices[VertexIndex][WeightIndex]);
+					FinalWeights[VertexIndex].Add(SetupWeights[VertexIndex][WeightIndex]);
+					TotalWeight += SetupWeights[VertexIndex][WeightIndex];
+				}
+			}
+			if(TotalWeight != 0.0f)
+			{
+				for(int32 WeightIndex = 0, NumWeights = FinalWeights[VertexIndex].Num(); WeightIndex < NumWeights; ++WeightIndex)
+				{
+					FinalWeights[VertexIndex][WeightIndex] /= TotalWeight;
+				}
+			}
+		}
+		else
+		{
+			FinalIndices[VertexIndex] = SetupIndices[VertexIndex];
+			FinalWeights[VertexIndex] = SetupWeights[VertexIndex];
+		}
+	}
+}
+	
+static void NormalizeSkinWeights(const TArray<TArray<int32>>& SetupIndices, const TArray<TArray<float>>& SetupWeights,
+			const TArray<float>& SelectionMap, TArray<TArray<int32>>& FinalIndices, TArray<TArray<float>>& FinalWeights)
+{
+	for(int32 VertexIndex = 0, NumVertices = FinalIndices.Num(); VertexIndex < NumVertices; ++VertexIndex)
+	{
+		if(SelectionMap[VertexIndex] > 0)
+		{
+			float TotalWeight = 0.0f;
+			for(int32 WeightIndex = 0, NumWeights = SetupWeights[VertexIndex].Num(); WeightIndex < NumWeights; ++WeightIndex)
+			{
+				FinalIndices[VertexIndex].Add(SetupIndices[VertexIndex][WeightIndex]);
+				FinalWeights[VertexIndex].Add(SetupWeights[VertexIndex][WeightIndex]);
+				TotalWeight += SetupWeights[VertexIndex][WeightIndex];
+			}
+			if(TotalWeight != 0.0f)
+			{
+				for(int32 WeightIndex = 0, NumWeights = FinalWeights[VertexIndex].Num(); WeightIndex < NumWeights; ++WeightIndex)
+				{
+					FinalWeights[VertexIndex][WeightIndex] /= TotalWeight;
+				}
+			}
+		}
+		else
+		{
+			FinalIndices[VertexIndex] = SetupIndices[VertexIndex];
+			FinalWeights[VertexIndex] = SetupWeights[VertexIndex];
+		}
+	}
+}
+
+static void ClampSkinWeights(const int32 ClampingNumber, const TArray<TArray<int32>>& SetupIndices, const TArray<TArray<float>>& SetupWeights,
+		const TArray<float>& SelectionMap, TArray<TArray<int32>>& FinalIndices, TArray<TArray<float>>& FinalWeights)
+{
+	for(int32 VertexIndex = 0, NumVertices = FinalIndices.Num(); VertexIndex < NumVertices; ++VertexIndex)
+	{
+		if(SelectionMap[VertexIndex] > 0)
+		{
+			TArray<TPair<int32,float>> SortedWeights;
+			SortedWeights.Reserve(SetupWeights[VertexIndex].Num());
+			for(int32 WeightIndex = 0, NumWeights = SetupWeights[VertexIndex].Num(); WeightIndex < NumWeights; ++WeightIndex)
+			{
+				SortedWeights.Add({SetupIndices[VertexIndex][WeightIndex], SetupWeights[VertexIndex][WeightIndex]});
+			}
+			// sort in descending order by weight
+			Algo::Sort(SortedWeights, [](const TPair<int32,float>& A, const TPair<int32,float>& B)
+			{
+				return A.Value > B.Value; 
+			});
+			float TotalWeight = 0.0f;
+			for(int32 WeightIndex = 0, NumWeights = FMath::Min(ClampingNumber,SetupWeights[VertexIndex].Num()); WeightIndex < NumWeights; ++WeightIndex)
+			{
+				FinalIndices[VertexIndex].Add(SortedWeights[WeightIndex].Key);
+				FinalWeights[VertexIndex].Add(SortedWeights[WeightIndex].Value);
+				TotalWeight += SortedWeights[WeightIndex].Value;
+			}
+			if(TotalWeight != 0.0f)
+			{
+				for(int32 WeightIndex = 0, NumWeights = FinalWeights[VertexIndex].Num(); WeightIndex < NumWeights; ++WeightIndex)
+				{
+					FinalWeights[VertexIndex][WeightIndex] /= TotalWeight;
+				}
+			}
+		}
+		else
+		{
+			FinalIndices[VertexIndex] = SetupIndices[VertexIndex];
+			FinalWeights[VertexIndex] = SetupWeights[VertexIndex];
+		}
+	}
+}
+
+static void RelaxSkinWeights(const UE::Geometry::FDynamicMesh3& DynamicMesh, const int32 VertexOffset, const float SmoothStrength, const int32 NumIterations,
+		const TArray<TArray<int32>>& SetupIndices, const TArray<TArray<float>>& SetupWeights,
+		const TArray<float>& SelectionMap, TArray<TArray<int32>>& FinalIndices, TArray<TArray<float>>& FinalWeights)
+{
+	if(UE::Geometry::FDynamicMeshVertexSkinWeightsAttribute* SkinWeights = DynamicMesh.Attributes()->GetSkinWeightsAttribute(FName("Default")))
+	{
+		SkinWeights->Initialize();
+		const int32 NumVertices = DynamicMesh.VertexCount();
+		TArray<int32> SelectedVertices;
+		
+		TArray<UE::AnimationCore::FBoneWeight> BoneWeightsBuffer;
+		for (int32 VertexIndex = 0; VertexIndex < NumVertices; ++VertexIndex)
+		{
+			const int32 GlobalVertex = VertexOffset+VertexIndex;
+			BoneWeightsBuffer.SetNumUninitialized(SetupIndices[GlobalVertex].Num());
+
+			for(int32 BoneIndex = 0; BoneIndex < SetupIndices[GlobalVertex].Num(); ++BoneIndex)
+			{
+				BoneWeightsBuffer[BoneIndex].SetBoneIndex(SetupIndices[VertexIndex+VertexOffset][BoneIndex]);
+				BoneWeightsBuffer[BoneIndex].SetWeight(SetupWeights[VertexIndex+VertexOffset][BoneIndex]);
+			}
+			SkinWeights->SetValue(VertexIndex, UE::AnimationCore::FBoneWeights::Create(BoneWeightsBuffer));
+			
+			if(SelectionMap[GlobalVertex] > 0.0f)
+			{
+				SelectedVertices.Add(VertexIndex);
+			}
+			else
+			{
+				FinalIndices[GlobalVertex] = SetupIndices[GlobalVertex];
+				FinalWeights[GlobalVertex] = SetupWeights[GlobalVertex];
+			}
+		}
+		
+		static constexpr float PercentPerIteration = 0.95f;
+		UE::Geometry::FSmoothDynamicMeshVertexSkinWeights SmoothBoneWeights(&DynamicMesh, FSkeletalMeshAttributes::DefaultSkinWeightProfileName);
+		if (SmoothBoneWeights.Validate() == UE::Geometry::EOperationValidationResult::Ok)
+		{
+			for (int32 IterationIndex = 0; IterationIndex < NumIterations; ++IterationIndex)
+			{
+				for (const int32 SelectedVertex : SelectedVertices)
+				{
+					SmoothBoneWeights.SmoothWeightsAtVertex(SelectedVertex, SmoothStrength * PercentPerIteration);
+				}
+			}
+		}
+
+		for (const int32 SelectedVertex : SelectedVertices)
+		{
+			const int32 GlobalVertex = VertexOffset+SelectedVertex;
+			SkinWeights->GetValue(SelectedVertex, FinalIndices[GlobalVertex], FinalWeights[GlobalVertex]);
+		}
+	}
+}
+
+static void HammerSkinWeights(const UE::Geometry::FDynamicMesh3& DynamicMesh, const int32 VertexOffset, const TArray<TArray<int32>>& SetupIndices, const TArray<TArray<float>>& SetupWeights,
+	const TArray<float>& SelectionMap, TArray<TArray<int32>>& FinalIndices, TArray<TArray<float>>& FinalWeights)
+{
+	const int32 NumVertices = DynamicMesh.VertexCount();
+	
+	TSet<int32> NeighborVertices;
+	TArray<int32> SelectedVertices;
+	
+	for (int32 VertexIndex = 0; VertexIndex < NumVertices; ++VertexIndex)
+	{
+		const int32 GlobalVertex = VertexOffset+VertexIndex;
+		if(SelectionMap[GlobalVertex] > 0.0f)
+		{
+			SelectedVertices.Add(VertexIndex);
+			for (const int32 NeighborIndex : DynamicMesh.VtxVerticesItr(VertexIndex))
+			{
+				if (SelectionMap[VertexOffset+NeighborIndex] == 0.0f)
+				{
+					NeighborVertices.Add(NeighborIndex);
+				}
+			}
+		}
+		else
+		{
+			FinalIndices[GlobalVertex] = SetupIndices[GlobalVertex];
+			FinalWeights[GlobalVertex] = SetupWeights[GlobalVertex];
+		}
+	}
+	UE::Geometry::TMeshDijkstra<UE::Geometry::FDynamicMesh3> PathFinder(&DynamicMesh);
+	TArray<UE::Geometry::TMeshDijkstra<UE::Geometry::FDynamicMesh3>::FSeedPoint> SeedPoints;
+	for (const int32 NeighborVertex : NeighborVertices)
+	{
+		SeedPoints.Add({ NeighborVertex, NeighborVertex, 0 });
+	}
+	PathFinder.ComputeToMaxDistance(SeedPoints, TNumericLimits<double>::Max());
+
+	// for each selected vertex, find the nearest surrounding vertex and copy it's weights
+	TArray<int32> VertexPath;
+	for (const int32 SelectedVertex : SelectedVertices)
+	{
+		// find the closest surrounding vertex to this selected vertex
+		if (!PathFinder.FindPathToNearestSeed(SelectedVertex, VertexPath))
+		{
+			continue;
+		}
+		const int32 ClosestVertex = VertexPath.Last();
+		const int32 GlobalVertex = VertexOffset+SelectedVertex;
+		const int32 GlobalClosest = VertexOffset+ClosestVertex;
+		
+		FinalIndices[GlobalVertex] = SetupIndices[GlobalClosest];
+		FinalWeights[GlobalVertex] = SetupWeights[GlobalClosest];
+	}
+}
+	
+template<typename ArrayType>
+static bool SetAttributeValues(FManagedArrayCollection& SelectedCollection,
+	const FCollectionAttributeKey& AttributeKey, const TArray<ArrayType>& AttributeValues)
+{
+	if (!AttributeValues.IsEmpty() && !AttributeKey.Attribute.IsEmpty() && !AttributeKey.Group.IsEmpty())
+	{
+		const FName AttributeName(AttributeKey.Attribute);
+		const FName AttributeGroup(AttributeKey.Group);
+		
+		if(TManagedArray<ArrayType>* AttributeArray = SelectedCollection.FindAttributeTyped<ArrayType>(AttributeName, AttributeGroup))
+		{
+			if(AttributeArray->Num() == AttributeValues.Num())
+			{
+				for(int32 VertexIndex = 0, NumVertices = AttributeArray->Num(); VertexIndex < NumVertices; ++VertexIndex)
+				{
+					AttributeArray->GetData()[VertexIndex] = AttributeValues[VertexIndex];
+				}
+			}
+			return true;
+		}
+	}
+	return false;
+}
+
+template<typename ArrayType>
+static bool GetAttributeValues(FManagedArrayCollection& SelectedCollection,
+	const FCollectionAttributeKey& AttributeKey, TArray<ArrayType>& AttributeValues)
+{
+	if (!AttributeKey.Attribute.IsEmpty() && !AttributeKey.Group.IsEmpty())
+	{
+		const FName AttributeName(AttributeKey.Attribute);
+		const FName AttributeGroup(AttributeKey.Group);
+		
+		TManagedArray<ArrayType>& AttributeArray = SelectedCollection.AddAttribute<ArrayType>(AttributeName, AttributeGroup);
+		AttributeValues = AttributeArray.GetConstArray();
+		return true;
+	}
+	return false;
+}
+
+}
+
+//
+// FDataflowCorrectSkinWeightsNode
+//
+
+FDataflowCorrectSkinWeightsNode::FDataflowCorrectSkinWeightsNode(const UE::Dataflow::FNodeParameters& InParam, FGuid InGuid)
+	: FDataflowNode(InParam, InGuid)
+{
+	RegisterInputConnection(&Collection);
+	RegisterInputConnection(&BoneIndicesKey);
+	RegisterInputConnection(&BoneWeightsKey);
+	RegisterInputConnection(&SelectionMapKey);
+	RegisterOutputConnection(&Collection, &Collection);
+	RegisterOutputConnection(&BoneIndicesKey, &BoneIndicesKey);
+	RegisterOutputConnection(&BoneWeightsKey, &BoneWeightsKey);
+}
+
+TArray<UE::Dataflow::FRenderingParameter> FDataflowCorrectSkinWeightsNode::GetRenderParametersImpl() const
+{
+	return FDataflowAddScalarVertexPropertyCallbackRegistry::Get().GetRenderingParameters(VertexGroup.Name);
+}
+
+void FDataflowCorrectSkinWeightsNode::Evaluate(UE::Dataflow::FContext& Context, const FDataflowOutput* Out) const
+{
+	using namespace UE::Dataflow;
+	
+	// Get the pin value if plugged
+	FCollectionAttributeKey BoneIndicesKeyValue = GetBoneIndicesKey(Context);
+	FCollectionAttributeKey BoneWeightsKeyValue = GetBoneWeightsKey(Context);
+	FCollectionAttributeKey SelectionMapKeyValue = GetSelectionMapKey(Context);
+
+	if (Out->IsA<FManagedArrayCollection>(&Collection))
+	{
+		// Evaluate in collection
+		FManagedArrayCollection InCollection = GetValue<FManagedArrayCollection>(Context, &Collection);
+
+		// Set temporary collection output
+		FManagedArrayCollection OutCollection = InCollection;
+		SetValue(Context, MoveTemp(OutCollection), &Collection);
+
+		if (!BoneIndicesKeyValue.Attribute.IsEmpty() && !BoneWeightsKeyValue.Attribute.IsEmpty() && !SelectionMapKeyValue.Attribute.IsEmpty())
+		{
+			TArray<TArray<float>> SetupWeights, FinalWeights;
+			TArray<TArray<int32>> SetupIndices, FinalIndices;
+			
+			UE::Dataflow::Private::GetAttributeValues<TArray<int32>>(InCollection, BoneIndicesKeyValue, SetupIndices);
+		    UE::Dataflow::Private::GetAttributeValues<TArray<float>>(InCollection, BoneWeightsKeyValue, SetupWeights);
+
+			TArray<float> SelectionMap;
+			UE::Dataflow::Private::GetAttributeValues<float>(InCollection, SelectionMapKeyValue, SelectionMap);
+
+			FinalIndices.SetNum(SetupIndices.Num());
+			FinalWeights.SetNum(SetupWeights.Num());
+			
+			if(CorrectionType == ESkinWeightsCorrectionType::Prune)
+			{
+				UE::Dataflow::Private::PruneSkinWeights(PruningThreshold, SetupIndices, SetupWeights, SelectionMap, FinalIndices, FinalWeights);
+			}
+			else if(CorrectionType == ESkinWeightsCorrectionType::Clamp)
+			{
+				UE::Dataflow::Private::ClampSkinWeights(ClampingNumber, SetupIndices, SetupWeights, SelectionMap, FinalIndices, FinalWeights);
+			}
+			else if(CorrectionType == ESkinWeightsCorrectionType::Normalize)
+			{
+				UE::Dataflow::Private::NormalizeSkinWeights(SetupIndices, SetupWeights, SelectionMap, FinalIndices, FinalWeights);
+			}
+			else
+			{
+				// Do the work here
+				TArray<UE::Geometry::FDynamicMesh3> DynamicMeshes;
+				UE::Dataflow::Private::BuildDynamicMeshes(this, Context, DynamicMeshes);
+
+				int32 VertexOffset = 0;
+				for(UE::Geometry::FDynamicMesh3& DynamicMesh : DynamicMeshes)
+				{
+					if(CorrectionType == ESkinWeightsCorrectionType::Relax)
+					{
+						UE::Dataflow::Private::RelaxSkinWeights(DynamicMesh, VertexOffset, SmoothingFactor, SmoothingIterations,
+							SetupIndices, SetupWeights, SelectionMap, FinalIndices, FinalWeights);
+					}
+					else if (CorrectionType == ESkinWeightsCorrectionType::Hammer)
+					{
+						UE::Dataflow::Private::HammerSkinWeights(DynamicMesh, VertexOffset,
+							SetupIndices, SetupWeights, SelectionMap, FinalIndices, FinalWeights);
+					}
+					VertexOffset += DynamicMesh.VertexCount();
+				}
+			}
+
+			UE::Dataflow::Private::SetAttributeValues<TArray<int32>>(InCollection, BoneIndicesKeyValue, FinalIndices);
+			UE::Dataflow::Private::SetAttributeValues<TArray<float>>(InCollection, BoneWeightsKeyValue, FinalWeights);
+		}
+
+		SetValue(Context, MoveTemp(InCollection), &Collection);
+	}
+	else if (Out->IsA<FCollectionAttributeKey>(&BoneIndicesKey))
+	{
+		SetValue(Context, MoveTemp(BoneIndicesKeyValue), &BoneIndicesKey);
+	}
+	else if (Out->IsA<FCollectionAttributeKey>(&BoneWeightsKey))
+	{
+		SetValue(Context, MoveTemp(BoneWeightsKeyValue), &BoneWeightsKey);
+	}
+}
+
+FCollectionAttributeKey FDataflowCorrectSkinWeightsNode::GetBoneIndicesKey(UE::Dataflow::FContext& Context) const
+{
+	// Get the pin value if plugged
+	FCollectionAttributeKey Key = GetValue(Context, &BoneIndicesKey, BoneIndicesKey);
+
+	// If nothing set used the local value
+	if(Key.Attribute.IsEmpty() && Key.Group.IsEmpty())
+	{
+		Key.Group = VertexGroup.Name.ToString();
+		Key.Attribute = BoneIndicesName;
+	}
+	return Key;
+}
+
+FCollectionAttributeKey FDataflowCorrectSkinWeightsNode::GetBoneWeightsKey(UE::Dataflow::FContext& Context) const
+{
+	// Get the pin value if plugged
+	FCollectionAttributeKey Key = GetValue(Context, &BoneWeightsKey, BoneWeightsKey);
+
+	// If nothing set used the local value
+	if(Key.Attribute.IsEmpty() && Key.Group.IsEmpty())
+	{
+		Key.Group = VertexGroup.Name.ToString();
+		Key.Attribute = BoneWeightsName;
+	}
+	return Key;
+}
+
+FCollectionAttributeKey FDataflowCorrectSkinWeightsNode::GetSelectionMapKey(UE::Dataflow::FContext& Context) const
+{
+	// Get the pin value if plugged
+	FCollectionAttributeKey Key = GetValue(Context, &SelectionMapKey, SelectionMapKey);
+
+	// If nothing set used the local value
+	if(Key.Attribute.IsEmpty() && Key.Group.IsEmpty())
+	{
+		Key.Group = VertexGroup.Name.ToString();
+		Key.Attribute = SelectionMapName;
+	}
+	return Key;
+}
+
+#undef LOCTEXT_NAMESPACE
diff --git a/Engine/Plugins/Experimental/Dataflow/Source/DataflowEditor/Public/DataflowEditorTools/DataflowEditorCorrectSkinWeightsNode.h b/Engine/Plugins/Experimental/Dataflow/Source/DataflowEditor/Public/DataflowEditorTools/DataflowEditorCorrectSkinWeightsNode.h
new file mode 100644
--- /dev/null
+++ b/Engine/Plugins/Experimental/Dataflow/Source/DataflowEditor/Public/DataflowEditorTools/DataflowEditorCorrectSkinWeightsNode.h
@@ -0,0 +1,101 @@
+// Copyright Epic Games, Inc. All Rights Reserved.
+
+#pragma once 
+
+
+#include "Dataflow/DataflowCollectionAddScalarVertexPropertyNode.h"
+#include "Dataflow/DataflowNode.h"
+#include "GeometryCollection/ManagedArrayCollection.h"
+#include "Dataflow/DataflowCollectionAttributeKeyNodes.h"
+
+#include "DataflowEditorCorrectSkinWeightsNode.generated.h"
+
+/** Skin weights correction type */
+UENUM()
+enum class ESkinWeightsCorrectionType : uint8
+{
+	Relax,
+	Prune,
+	Hammer,
+	Clamp,
+	Normalize
+};
+
+/** Correct skin weights vertex properties. */
+USTRUCT(Meta = (Experimental, DataflowCollection))
+struct DATAFLOWEDITOR_API FDataflowCorrectSkinWeightsNode : public FDataflowNode
+{
+	GENERATED_BODY()
+	DATAFLOW_NODE_DEFINE_INTERNAL(FDataflowCorrectSkinWeightsNode, "CorrectSkinWeights", "Collection", "Correct skin weights and save it to collection")
+
+public:
+
+	FDataflowCorrectSkinWeightsNode(const UE::Dataflow::FNodeParameters& InParam, FGuid InGuid = FGuid::NewGuid());
+
+	UPROPERTY(Meta = (DataflowInput, DataflowOutput, DataflowPassthrough = "Collection"))
+	FManagedArrayCollection Collection;
+
+	/** The name to be set for the bone indices. */
+	UPROPERTY(EditAnywhere, Category = "Vertex Attributes", meta = (DisplayName = "Bone Indices"))
+	FString BoneIndicesName;
+
+	/** The name to be set for the bone weights. */
+	UPROPERTY(EditAnywhere, Category = "Vertex Attributes", meta = (DisplayName = "Bone Weights"))
+	FString BoneWeightsName;
+
+	/** Map name to be used to select vertices to correct */
+	UPROPERTY(EditAnywhere, Category = "Vertex Attributes", meta = (DisplayName = "Selection Map"))
+	FString SelectionMapName;
+
+	/** Target group in which the attributes are stored */
+	UPROPERTY(EditAnywhere, Category = "Vertex Attributes", meta = (DisplayName = "Vertex Group"))
+	FScalarVertexPropertyGroup VertexGroup;
+
+	/** Bone indices key to be used in other nodes if necessary */
+	UPROPERTY(meta = (DisplayName = "Bone Indices Key", DataflowInput, DataflowOutput, DataflowPassthrough = "BoneIndicesKey"))
+	FCollectionAttributeKey BoneIndicesKey;
+
+	/** Bone weights key to be used in other nodes if necessary */
+	UPROPERTY(meta = (DisplayName = "Bone Weights Key", DataflowInput, DataflowOutput, DataflowPassthrough = "BoneWeightsKey"))
+	FCollectionAttributeKey BoneWeightsKey;
+
+	/** Selection map key to be used in other nodes if necessary */
+	UPROPERTY(meta = (DisplayName = "Selection Map Key", DataflowInput))
+	FCollectionAttributeKey SelectionMapKey;
+	
+	/** Selection map key to be used in other nodes if necessary */
+	UPROPERTY(EditAnywhere, Category = "Skin Weights", meta = (DisplayName = "Correction Type"))
+	ESkinWeightsCorrectionType CorrectionType = ESkinWeightsCorrectionType::Relax;
+
+	/** Number of iteration required for the smoothing */
+	UPROPERTY(EditAnywhere, Category = "Skin Weights", meta = (DisplayName = "Smoothing Iterations", EditCondition = "CorrectionType == ESkinWeightsCorrectionType::Relax", EditConditionHides))
+	int32 SmoothingIterations = 5;
+
+	/** Lerp value in between the current and the average weight values */
+	UPROPERTY(EditAnywhere, Category = "Skin Weights", meta = (DisplayName = "Smooting Factor", EditCondition = "CorrectionType == ESkinWeightsCorrectionType::Relax", EditConditionHides))
+	float SmoothingFactor = 0.5;
+
+	/** All weights below this threshold will be pruned */
+	UPROPERTY(EditAnywhere, Category = "Skin Weights", meta = (DisplayName = "Pruning Threshold", EditCondition = "CorrectionType == ESkinWeightsCorrectionType::Prune", EditConditionHides))
+	float PruningThreshold = 0.01;
+	
+	/** Max number of bones to consider for the skin weights */
+    UPROPERTY(EditAnywhere, Category = "Skin Weights", meta = (DisplayName = "Clamping Number", EditCondition = "CorrectionType == ESkinWeightsCorrectionType::Clamp", EditConditionHides))
+    int32 ClampingNumber = 8;
+	
+	/** Get the weights attribute key to retrieve/set the bone indices values*/
+	FCollectionAttributeKey GetBoneIndicesKey(UE::Dataflow::FContext& Context) const;
+
+	/** Get the weights attribute key to retrieve/set the bone weights values*/
+    FCollectionAttributeKey GetBoneWeightsKey(UE::Dataflow::FContext& Context) const;
+
+	/** Get the selection map key to get the vertex selection on which the correction will be applied */
+	FCollectionAttributeKey GetSelectionMapKey(UE::Dataflow::FContext& Context) const;
+	
+private:
+
+	//~ Begin FDataflowPrimitiveNode interface
+	virtual TArray<UE::Dataflow::FRenderingParameter> GetRenderParametersImpl() const override;
+	virtual void Evaluate(UE::Dataflow::FContext& Context, const FDataflowOutput* Out) const override;
+	//~ End FDataflowPrimitiveNode interface
+};
diff --git a/Engine/Plugins/Experimental/Dataflow/Source/DataflowNodes/Public/Dataflow/DataflowCollectionEditSkinWeightsNode.h b/Engine/Plugins/Experimental/Dataflow/Source/DataflowNodes/Public/Dataflow/DataflowCollectionEditSkinWeightsNode.h
--- a/Engine/Plugins/Experimental/Dataflow/Source/DataflowNodes/Public/Dataflow/DataflowCollectionEditSkinWeightsNode.h
+++ b/Engine/Plugins/Experimental/Dataflow/Source/DataflowNodes/Public/Dataflow/DataflowCollectionEditSkinWeightsNode.h
@@ -28,7 +28,7 @@ struct FDataflowSkinWeightData
 };
 
 /** Edit skin weights vertex properties. */
-USTRUCT(Meta = (DataflowCollection))
+USTRUCT(Meta = (Experimental,DataflowCollection))
 struct FDataflowCollectionEditSkinWeightsNode : public FDataflowPrimitiveNode
 {
 	GENERATED_BODY()
