Fix dataflow panel simulation locking + synced dataflow caching deltatime

diff --git a/Engine/Plugins/Experimental/Dataflow/Source/DataflowEditor/Private/Dataflow/DataflowSimulationPanel.cpp b/Engine/Plugins/Experimental/Dataflow/Source/DataflowEditor/Private/Dataflow/DataflowSimulationPanel.cpp
--- a/Engine/Plugins/Experimental/Dataflow/Source/DataflowEditor/Private/Dataflow/DataflowSimulationPanel.cpp
+++ b/Engine/Plugins/Experimental/Dataflow/Source/DataflowEditor/Private/Dataflow/DataflowSimulationPanel.cpp
@@ -235,13 +235,12 @@ void SDataflowSimulationPanel::OnTickPlayback(double InCurrentTime, float InDelt
 	if (const TSharedPtr<FDataflowSimulationScene> PreviewScene = SimulationScene.Pin())
 	{
 		const float SequenceLength = GetSequenceLength();
-		
 		const float ScrubValue = (PlaybackMode == EPlaybackMode::PlayingForward) ?
 			GetScrubValue() + InDeltaTime : GetScrubValue() - InDeltaTime;
 		
 		if(PreviewPlaybackMode == EDataflowPlaybackMode::Looping)
 		{
-			UpdateSimulationTimeFromScrubValue(ScrubValue - SequenceLength * FMath::Floor(ScrubValue / SequenceLength));
+			UpdateSimulationTimeFromScrubValue(ScrubValue - SequenceLength * FMath::Floor(ScrubValue / SequenceLength), false);
 		}
 		else
 		{
@@ -257,7 +256,7 @@ void SDataflowSimulationPanel::OnTickPlayback(double InCurrentTime, float InDelt
 					PlaybackMode = EPlaybackMode::PlayingForward;
 				}
 			}
-			UpdateSimulationTimeFromScrubValue(FMath::Clamp(ScrubValue, 0.f, SequenceLength));
+			UpdateSimulationTimeFromScrubValue(FMath::Clamp(ScrubValue, 0.f, SequenceLength), false);
 		}
 	}
 }
@@ -267,18 +266,19 @@ EPlaybackMode::Type SDataflowSimulationPanel::GetPlaybackMode() const
 	return PlaybackMode;
 }
 
-void SDataflowSimulationPanel::UpdateSimulationTimeFromScrubValue(float ScrubValue)
+void SDataflowSimulationPanel::UpdateSimulationTimeFromScrubValue(float ScrubValue, const bool bRoundedFrame)
 {
 	if (const TSharedPtr<FDataflowSimulationScene> PreviewScene = SimulationScene.Pin())
 	{
 		const int32 TotalFrameRate = PreviewScene->GetFrameRate() * PreviewScene->GetSubframeRate();
-		const float RoundedFrameTime = FMath::RoundToFloat(ScrubValue * TotalFrameRate) / TotalFrameRate;
+		const float RoundedFrameTime = bRoundedFrame ? FMath::RoundToFloat(ScrubValue * TotalFrameRate) / TotalFrameRate : ScrubValue;
+
 		PreviewScene->SimulationTime = RoundedFrameTime + PreviewScene->GetTimeRange()[0];
 		FrameIndexWidget->SetText(FText::AsNumber(RoundedFrameTime * PreviewScene->GetFrameRate()));
 	}
 }
 
-void SDataflowSimulationPanel::OnValueChanged(float NewValue)
+void SDataflowSimulationPanel::OnValueChanged(float NewValue) 
 {
 	UpdateSimulationTimeFromScrubValue(NewValue);
 }
diff --git a/Engine/Plugins/Experimental/Dataflow/Source/DataflowEditor/Private/Dataflow/DataflowSimulationScene.cpp b/Engine/Plugins/Experimental/Dataflow/Source/DataflowEditor/Private/Dataflow/DataflowSimulationScene.cpp
--- a/Engine/Plugins/Experimental/Dataflow/Source/DataflowEditor/Private/Dataflow/DataflowSimulationScene.cpp
+++ b/Engine/Plugins/Experimental/Dataflow/Source/DataflowEditor/Private/Dataflow/DataflowSimulationScene.cpp
@@ -225,6 +225,7 @@ void FDataflowSimulationScene::CreateSimulationScene()
 
 		TimeRange = SceneDescription->CacheParams.TimeRange;
 		NumFrames = (TimeRange[1] > TimeRange[0]) ? FMath::Floor((TimeRange[1] - TimeRange[0] + UE_SMALL_NUMBER) * SceneDescription->CacheParams.FrameRate) + 1 : 1;
+		
 		DeltaTime = (SceneDescription->CacheParams.FrameRate > 0) ? 1.f / SceneDescription->CacheParams.FrameRate: 0.f;
 		PreviewActor = UE::Dataflow::SpawnSimulatedActor(SceneDescription->BlueprintClass, Cast<AChaosCacheManager>(RootSceneActor),
 			SceneDescription->CacheAsset, false, GetEditorContent(), SceneDescription->BlueprintTransform);
@@ -276,7 +277,7 @@ void FDataflowSimulationScene::RecordSimulationCache()
 	}
 }
 
-void FDataflowSimulationScene::UpdateSimulationCache(float DeltaSeconds)
+void FDataflowSimulationScene::UpdateSimulationCache(float& ElapsedTime)
 {
 	if(RootSceneActor)
 	{
@@ -292,6 +293,7 @@ void FDataflowSimulationScene::UpdateSimulationCache(float DeltaSeconds)
 		}
 		else
 		{
+			ElapsedTime = DeltaTime;
 			SimulationTime += DeltaTime;
 
 			if(SimulationTime >= TimeRange[1]+UE_SMALL_NUMBER)
@@ -380,6 +382,7 @@ FBox FDataflowSimulationScene::GetBoundingBox() const
 
 void FDataflowSimulationScene::TickDataflowScene(const float DeltaSeconds)
 {
+	float ElapsedTime = DeltaSeconds;
 	if(const TObjectPtr<UDataflowBaseContent>& EditorContent = GetEditorContent())
 	{
 		if (const TObjectPtr<UDataflow> DataflowGraph = EditorContent->GetDataflowAsset())
@@ -410,13 +413,13 @@ void FDataflowSimulationScene::TickDataflowScene(const float DeltaSeconds)
 		if(SceneDescription->CacheAsset)
 		{
 			// Update the simulation cache (loading/recording)
-			UpdateSimulationCache(DeltaSeconds);
+			UpdateSimulationCache(ElapsedTime);
 			
 			// Update all the skelmesh animations at the simulation time
 			UE::Dataflow::UpdateSkeletonAnimation(PreviewActor, SimulationTime);
 		}
 	}
-	GetWorld()->Tick(ELevelTick::LEVELTICK_All, DeltaSeconds);
+	GetWorld()->Tick(ELevelTick::LEVELTICK_All, ElapsedTime);
 }
 
 void FDataflowSimulationScene::AddReferencedObjects(FReferenceCollector& Collector)
diff --git a/Engine/Plugins/Experimental/Dataflow/Source/DataflowEditor/Public/Dataflow/DataflowSimulationPanel.h b/Engine/Plugins/Experimental/Dataflow/Source/DataflowEditor/Public/Dataflow/DataflowSimulationPanel.h
--- a/Engine/Plugins/Experimental/Dataflow/Source/DataflowEditor/Public/Dataflow/DataflowSimulationPanel.h
+++ b/Engine/Plugins/Experimental/Dataflow/Source/DataflowEditor/Public/Dataflow/DataflowSimulationPanel.h
@@ -48,7 +48,7 @@ private:
 
 	/** Main simulation delegates */
 	void OnTickPlayback(double InCurrentTime, float InDeltaTime);
-	void UpdateSimulationTimeFromScrubValue(float ScrubValue);
+	void UpdateSimulationTimeFromScrubValue(float ScrubValue, const bool bRoundedFrame = true);
 	void OnValueChanged(float NewValue);
 	void SetFrameIndex(const FText& InNewText, ETextCommit::Type InCommitType);
 	void OnBeginSliderMovement();
diff --git a/Engine/Plugins/Experimental/Dataflow/Source/DataflowEditor/Public/Dataflow/DataflowSimulationScene.h b/Engine/Plugins/Experimental/Dataflow/Source/DataflowEditor/Public/Dataflow/DataflowSimulationScene.h
--- a/Engine/Plugins/Experimental/Dataflow/Source/DataflowEditor/Public/Dataflow/DataflowSimulationScene.h
+++ b/Engine/Plugins/Experimental/Dataflow/Source/DataflowEditor/Public/Dataflow/DataflowSimulationScene.h
@@ -173,7 +173,7 @@ public:
 private:
 
 	/** Update the simulation cache at some point in time (loading/recording)*/
-	void UpdateSimulationCache(float DeltaSeconds);
+	void UpdateSimulationCache(float& DeltaSeconds);
 
 	/** Bind the scene selection to the components */
 	void BindSceneSelection();
