Update resizing graph to support skin weight correction + nodes to retrieve the vertex selection for skin weights correction

#virtualized

diff --git a/Engine/Plugins/Experimental/Dataflow/Source/DataflowEditor/Private/Dataflow/DataflowEditorModule.cpp b/Engine/Plugins/Experimental/Dataflow/Source/DataflowEditor/Private/Dataflow/DataflowEditorModule.cpp
--- a/Engine/Plugins/Experimental/Dataflow/Source/DataflowEditor/Private/Dataflow/DataflowEditorModule.cpp
+++ b/Engine/Plugins/Experimental/Dataflow/Source/DataflowEditor/Private/Dataflow/DataflowEditorModule.cpp
@@ -135,6 +135,8 @@ void FDataflowEditorModule::StartupModule()
 		FSlateIcon(FName("DataflowEditorStyle"), FName("Dataflow.EditSkinWeights")), LOCTEXT("AddSkinWeightNodeButtonText", "Edit Skin Weights"));
 		
 	DATAFLOW_NODE_REGISTER_CREATION_FACTORY(FDataflowCorrectSkinWeightsNode);
+	DATAFLOW_NODE_REGISTER_CREATION_FACTORY(FDataflowGetSkinningSelectionNode);
+	DATAFLOW_NODE_REGISTER_CREATION_FACTORY(FDataflowSetSkinningSelectionNode);
 	
 	UE::Dataflow::FDataflowConstructionVisualizationRegistry& ConstructionVisualizationRegistry = UE::Dataflow::FDataflowConstructionVisualizationRegistry::GetInstance();
 	ConstructionVisualizationRegistry.RegisterVisualization(MakeUnique<UE::Dataflow::FMeshStatsConstructionVisualization>());
diff --git a/Engine/Plugins/Experimental/Dataflow/Source/DataflowEditor/Private/DataflowEditorTools/DataflowEditorCorrectSkinWeightsNode.cpp b/Engine/Plugins/Experimental/Dataflow/Source/DataflowEditor/Private/DataflowEditorTools/DataflowEditorCorrectSkinWeightsNode.cpp
--- a/Engine/Plugins/Experimental/Dataflow/Source/DataflowEditor/Private/DataflowEditorTools/DataflowEditorCorrectSkinWeightsNode.cpp
+++ b/Engine/Plugins/Experimental/Dataflow/Source/DataflowEditor/Private/DataflowEditorTools/DataflowEditorCorrectSkinWeightsNode.cpp
@@ -1,4 +1,4 @@
-﻿// Copyright Epic Games, Inc. All Rights Reserved.
+// Copyright Epic Games, Inc. All Rights Reserved.
 
 #include "DataflowEditorTools/DataflowEditorCorrectSkinWeightsNode.h"
 
@@ -14,6 +14,7 @@
 #include "GeometryCollection/Facades/CollectionRenderingFacade.h"
 #include "Operations/SmoothBoneWeights.h"
 #include "SkeletalMeshAttributes.h"
+#include "Chaos/CollectionPropertyFacade.h"
 #include "Parameterization/MeshDijkstra.h"
 
 #include UE_INLINE_GENERATED_CPP_BY_NAME(DataflowEditorCorrectSkinWeightsNode)
@@ -308,6 +309,12 @@ static bool GetAttributeValues(FManagedArrayCollection& SelectedCollection,
 // FDataflowCorrectSkinWeightsNode
 //
 
+const FName FDataflowCorrectSkinWeightsNode::PruneSkinWeightsSelectionName = "PruneSkinWeightsSelection";
+const FName FDataflowCorrectSkinWeightsNode::HammerSkinWeightsSelectionName = "HammerSkinWeightsSelection";
+const FName FDataflowCorrectSkinWeightsNode::RelaxSkinWeightsSelectionName = "RelaxSkinWeightsSelection";
+const FName FDataflowCorrectSkinWeightsNode::ClampSkinWeightsSelectionName = "ClampSkinWeightsSelection";
+const FName FDataflowCorrectSkinWeightsNode::NormalizeSkinWeightsSelectionName = "NormalizeSkinWeightsSelection";
+
 FDataflowCorrectSkinWeightsNode::FDataflowCorrectSkinWeightsNode(const UE::Dataflow::FNodeParameters& InParam, FGuid InGuid)
 	: FDataflowNode(InParam, InGuid)
 {
@@ -450,4 +457,133 @@ FCollectionAttributeKey FDataflowCorrectSkinWeightsNode::GetSelectionMapKey(UE::
 	return Key;
 }
 
+FDataflowSetSkinningSelectionNode::FDataflowSetSkinningSelectionNode(const UE::Dataflow::FNodeParameters& InParam, FGuid InGuid)
+	: FDataflowNode(InParam, InGuid)
+{
+	RegisterInputConnection(&Collection);
+	RegisterInputConnection(&SelectionMapKey);
+	RegisterOutputConnection(&Collection, &Collection);
+}
+
+TArray<UE::Dataflow::FRenderingParameter> FDataflowSetSkinningSelectionNode::GetRenderParametersImpl() const
+{
+	return FDataflowAddScalarVertexPropertyCallbackRegistry::Get().GetRenderingParameters(VertexGroup.Name);
+}
+
+void FDataflowSetSkinningSelectionNode::Evaluate(UE::Dataflow::FContext& Context, const FDataflowOutput* Out) const
+{
+	using namespace UE::Dataflow;
+	
+	// Get the pin value if plugged
+	FCollectionAttributeKey SelectionMapKeyValue = GetSelectionMapKey(Context);
+
+	if (Out->IsA<FManagedArrayCollection>(&Collection))
+	{
+		// Evaluate in collection
+		FManagedArrayCollection InCollection = GetValue<FManagedArrayCollection>(Context, &Collection);
+		const TSharedRef<FManagedArrayCollection> OutCollection = MakeShared<FManagedArrayCollection>(MoveTemp(InCollection));
+		
+		Chaos::Softs::FCollectionPropertyMutableFacade PropertyFacade(OutCollection);
+		PropertyFacade.DefineSchema();
+
+		if(CorrectionType == ESkinWeightsCorrectionType::Prune)
+		{
+			PropertyFacade.AddStringValue(SelectionMapKeyValue.Group + FString("::") + FDataflowCorrectSkinWeightsNode::PruneSkinWeightsSelectionName.ToString(), SelectionMapKeyValue.Attribute);
+		}
+		else if(CorrectionType == ESkinWeightsCorrectionType::Relax)
+		{
+			PropertyFacade.AddStringValue(SelectionMapKeyValue.Group + FString("::") + FDataflowCorrectSkinWeightsNode::RelaxSkinWeightsSelectionName.ToString(), SelectionMapKeyValue.Attribute);
+		}
+		else if(CorrectionType == ESkinWeightsCorrectionType::Hammer)
+		{
+			PropertyFacade.AddStringValue(SelectionMapKeyValue.Group + FString("::") + FDataflowCorrectSkinWeightsNode::HammerSkinWeightsSelectionName.ToString(), SelectionMapKeyValue.Attribute);
+		}
+		else if(CorrectionType == ESkinWeightsCorrectionType::Clamp)
+		{
+			PropertyFacade.AddStringValue(SelectionMapKeyValue.Group + FString("::") + FDataflowCorrectSkinWeightsNode::ClampSkinWeightsSelectionName.ToString(), SelectionMapKeyValue.Attribute);
+		}
+		else if(CorrectionType == ESkinWeightsCorrectionType::Normalize)
+		{
+			PropertyFacade.AddStringValue(SelectionMapKeyValue.Group + FString("::") + FDataflowCorrectSkinWeightsNode::NormalizeSkinWeightsSelectionName.ToString(), SelectionMapKeyValue.Attribute);
+		}
+
+		SetValue(Context, MoveTemp(*OutCollection), &Collection);
+	}
+}
+
+FCollectionAttributeKey FDataflowSetSkinningSelectionNode::GetSelectionMapKey(UE::Dataflow::FContext& Context) const
+{
+	// Get the pin value if plugged
+	FCollectionAttributeKey Key = GetValue(Context, &SelectionMapKey, SelectionMapKey);
+
+	// If nothing set used the local value
+	if(Key.Attribute.IsEmpty() && Key.Group.IsEmpty())
+	{
+		Key.Group = VertexGroup.Name.ToString();
+		Key.Attribute = SelectionMapName;
+	}
+	return Key;
+}
+
+
+FDataflowGetSkinningSelectionNode::FDataflowGetSkinningSelectionNode(const UE::Dataflow::FNodeParameters& InParam, FGuid InGuid)
+	: FDataflowNode(InParam, InGuid)
+{
+	RegisterInputConnection(&Collection);
+	RegisterOutputConnection(&Collection, &Collection);
+	RegisterOutputConnection(&SelectionMapKey);
+}
+
+TArray<UE::Dataflow::FRenderingParameter> FDataflowGetSkinningSelectionNode::GetRenderParametersImpl() const
+{
+	return FDataflowAddScalarVertexPropertyCallbackRegistry::Get().GetRenderingParameters(VertexGroup.Name);
+}
+
+void FDataflowGetSkinningSelectionNode::Evaluate(UE::Dataflow::FContext& Context, const FDataflowOutput* Out) const
+{
+	using namespace UE::Dataflow;
+	
+	if (Out->IsA<FManagedArrayCollection>(&Collection) || Out->IsA<FCollectionAttributeKey>(&SelectionMapKey))
+	{
+		// Evaluate in collection
+		FManagedArrayCollection InCollection = GetValue<FManagedArrayCollection>(Context, &Collection);
+		const TSharedRef<FManagedArrayCollection> OutCollection = MakeShared<FManagedArrayCollection>(MoveTemp(InCollection));
+
+		if(Out->IsA<FCollectionAttributeKey>(&SelectionMapKey))
+		{
+			// Get the pin value if plugged
+			FCollectionAttributeKey SelectionMapKeyValue;
+			SelectionMapKeyValue.Group = VertexGroup.Name.ToString();
+			
+			Chaos::Softs::FCollectionPropertyConstFacade PropertyFacade(OutCollection);  
+
+			if(CorrectionType == ESkinWeightsCorrectionType::Prune)
+			{
+				SelectionMapKeyValue.Attribute = PropertyFacade.GetStringValue(SelectionMapKeyValue.Group + FString("::") + FDataflowCorrectSkinWeightsNode::PruneSkinWeightsSelectionName.ToString());
+			}
+			else if(CorrectionType == ESkinWeightsCorrectionType::Relax)
+			{
+				SelectionMapKeyValue.Attribute = PropertyFacade.GetStringValue(SelectionMapKeyValue.Group + FString("::") + FDataflowCorrectSkinWeightsNode::RelaxSkinWeightsSelectionName.ToString());
+			}
+			else if(CorrectionType == ESkinWeightsCorrectionType::Hammer)
+			{
+				SelectionMapKeyValue.Attribute = PropertyFacade.GetStringValue(SelectionMapKeyValue.Group + FString("::") + FDataflowCorrectSkinWeightsNode::HammerSkinWeightsSelectionName.ToString());
+			}
+			else if(CorrectionType == ESkinWeightsCorrectionType::Clamp)
+			{
+				SelectionMapKeyValue.Attribute = PropertyFacade.GetStringValue(SelectionMapKeyValue.Group + FString("::") + FDataflowCorrectSkinWeightsNode::ClampSkinWeightsSelectionName.ToString());
+			}
+			else if(CorrectionType == ESkinWeightsCorrectionType::Normalize)
+			{
+				SelectionMapKeyValue.Attribute = PropertyFacade.GetStringValue(SelectionMapKeyValue.Group + FString("::") + FDataflowCorrectSkinWeightsNode::NormalizeSkinWeightsSelectionName.ToString());
+			}
+			SetValue(Context, MoveTemp(SelectionMapKeyValue), &SelectionMapKey);
+		}
+		if(Out->IsA<FManagedArrayCollection>(&Collection))
+		{
+			SetValue(Context, MoveTemp(*OutCollection), &Collection);
+		}
+	}
+}
+
 #undef LOCTEXT_NAMESPACE
diff --git a/Engine/Plugins/Experimental/Dataflow/Source/DataflowEditor/Public/DataflowEditorTools/DataflowEditorCorrectSkinWeightsNode.h b/Engine/Plugins/Experimental/Dataflow/Source/DataflowEditor/Public/DataflowEditorTools/DataflowEditorCorrectSkinWeightsNode.h
--- a/Engine/Plugins/Experimental/Dataflow/Source/DataflowEditor/Public/DataflowEditorTools/DataflowEditorCorrectSkinWeightsNode.h
+++ b/Engine/Plugins/Experimental/Dataflow/Source/DataflowEditor/Public/DataflowEditorTools/DataflowEditorCorrectSkinWeightsNode.h
@@ -30,6 +30,12 @@ struct DATAFLOWEDITOR_API FDataflowCorrectSkinWeightsNode : public FDataflowNode
 
 public:
 
+	static const FName PruneSkinWeightsSelectionName;
+	static const FName HammerSkinWeightsSelectionName;
+	static const FName RelaxSkinWeightsSelectionName;
+	static const FName ClampSkinWeightsSelectionName;
+	static const FName NormalizeSkinWeightsSelectionName;
+
 	FDataflowCorrectSkinWeightsNode(const UE::Dataflow::FNodeParameters& InParam, FGuid InGuid = FGuid::NewGuid());
 
 	UPROPERTY(Meta = (DataflowInput, DataflowOutput, DataflowPassthrough = "Collection"))
@@ -99,3 +105,82 @@ private:
 	virtual void Evaluate(UE::Dataflow::FContext& Context, const FDataflowOutput* Out) const override;
 	//~ End FDataflowPrimitiveNode interface
 };
+
+/** Set skin weights selection attributes. */
+USTRUCT(Meta = (Experimental, DataflowCollection))
+struct DATAFLOWEDITOR_API FDataflowSetSkinningSelectionNode : public FDataflowNode
+{
+	GENERATED_BODY()
+	DATAFLOW_NODE_DEFINE_INTERNAL(FDataflowSetSkinningSelectionNode, "SetSkinningSelection", "Collection", "Set the skin weights selection for a future correction")
+
+public:
+
+	FDataflowSetSkinningSelectionNode(const UE::Dataflow::FNodeParameters& InParam, FGuid InGuid = FGuid::NewGuid());
+
+	UPROPERTY(Meta = (DataflowInput, DataflowOutput, DataflowPassthrough = "Collection"))
+	FManagedArrayCollection Collection;
+
+	/** Map name to be used to select vertices to correct */
+	UPROPERTY(EditAnywhere, Category = "Vertex Attributes", meta = (DisplayName = "Selection Map"))
+	FString SelectionMapName;
+
+	/** Target group in which the attributes are stored */
+	UPROPERTY(EditAnywhere, Category = "Vertex Attributes", meta = (DisplayName = "Vertex Group"))
+	FScalarVertexPropertyGroup VertexGroup;
+
+	/** Selection map key to be used in other nodes if necessary */
+	UPROPERTY(meta = (DisplayName = "Selection Map Key", DataflowInput))
+	FCollectionAttributeKey SelectionMapKey;
+	
+	/** Selection map key to be used in other nodes if necessary */
+	UPROPERTY(EditAnywhere, Category = "Skin Weights", meta = (DisplayName = "Correction Type"))
+	ESkinWeightsCorrectionType CorrectionType = ESkinWeightsCorrectionType::Relax;
+
+	/** Get the selection map key to get the vertex selection on which the correction will be applied */
+	FCollectionAttributeKey GetSelectionMapKey(UE::Dataflow::FContext& Context) const;
+	
+private:
+
+	//~ Begin FDataflowPrimitiveNode interface
+	virtual TArray<UE::Dataflow::FRenderingParameter> GetRenderParametersImpl() const override;
+	virtual void Evaluate(UE::Dataflow::FContext& Context, const FDataflowOutput* Out) const override;
+	//~ End FDataflowPrimitiveNode interface
+};
+
+
+
+/** Get skin weights selection attributes. */
+USTRUCT(Meta = (Experimental, DataflowCollection))
+struct DATAFLOWEDITOR_API FDataflowGetSkinningSelectionNode : public FDataflowNode
+{
+	GENERATED_BODY()
+	DATAFLOW_NODE_DEFINE_INTERNAL(FDataflowGetSkinningSelectionNode, "GetSkinningSelection", "Collection", "Get the skin weights selection for a future correction")
+
+public:
+
+	FDataflowGetSkinningSelectionNode(const UE::Dataflow::FNodeParameters& InParam, FGuid InGuid = FGuid::NewGuid());
+
+	UPROPERTY(Meta = (DataflowInput, DataflowOutput, DataflowPassthrough = "Collection"))
+	FManagedArrayCollection Collection;
+
+	/** Target group in which the attributes are stored */
+	UPROPERTY(EditAnywhere, Category = "Vertex Attributes", meta = (DisplayName = "Vertex Group"))
+	FScalarVertexPropertyGroup VertexGroup;
+
+	/** Selection map key to be used in other nodes if necessary */
+	UPROPERTY(meta = (DisplayName = "Selection Map Key", DataflowOutput))
+	FCollectionAttributeKey SelectionMapKey;
+	
+	/** Selection map key to be used in other nodes if necessary */
+	UPROPERTY(EditAnywhere, Category = "Skin Weights", meta = (DisplayName = "Correction Type"))
+	ESkinWeightsCorrectionType CorrectionType = ESkinWeightsCorrectionType::Relax;
+	
+private:
+
+	//~ Begin FDataflowPrimitiveNode interface
+	virtual TArray<UE::Dataflow::FRenderingParameter> GetRenderParametersImpl() const override;
+	virtual void Evaluate(UE::Dataflow::FContext& Context, const FDataflowOutput* Out) const override;
+	//~ End FDataflowPrimitiveNode interface
+};
+
+
