Dataflow : make suyre to evaluate the graph when switching form manual to automatic mode

diff --git a/Engine/Plugins/Experimental/Dataflow/Source/DataflowEditor/Private/Dataflow/DataflowEditorToolkit.cpp b/Engine/Plugins/Experimental/Dataflow/Source/DataflowEditor/Private/Dataflow/DataflowEditorToolkit.cpp
--- a/Engine/Plugins/Experimental/Dataflow/Source/DataflowEditor/Private/Dataflow/DataflowEditorToolkit.cpp
+++ b/Engine/Plugins/Experimental/Dataflow/Source/DataflowEditor/Private/Dataflow/DataflowEditorToolkit.cpp
@@ -756,10 +756,18 @@ void FDataflowEditorToolkit::PostInitAssetEditor()
 
 void FDataflowEditorToolkit::SetEvaluateGraphMode(EDataflowEditorEvaluationMode Mode)
 {
-	EvaluationMode = Mode;
-	if (UDataflowEditorOptions* const Options = UDataflowEditorOptions::StaticClass()->GetDefaultObject<UDataflowEditorOptions>())
+	if (Mode != EvaluationMode)
 	{
-		Options->EditorEvaluationMode = EvaluationMode;
+		EvaluationMode = Mode;
+		if (UDataflowEditorOptions* const Options = UDataflowEditorOptions::StaticClass()->GetDefaultObject<UDataflowEditorOptions>())
+		{
+			Options->EditorEvaluationMode = EvaluationMode;
+		}
+		// when going back to automatic, make sure the graph is up to date
+		if (EvaluationMode == EDataflowEditorEvaluationMode::Automatic)
+		{
+			EvaluateGraph();
+		}
 	}
 }
 
