add limit to point rendering to avoid hangs on vis if large-ish number of points are requested

diff --git a/Engine/Plugins/Experimental/Dataflow/Source/DataflowEditor/Private/Dataflow/DataflowEngineRendering.cpp b/Engine/Plugins/Experimental/Dataflow/Source/DataflowEditor/Private/Dataflow/DataflowEngineRendering.cpp
--- a/Engine/Plugins/Experimental/Dataflow/Source/DataflowEditor/Private/Dataflow/DataflowEngineRendering.cpp
+++ b/Engine/Plugins/Experimental/Dataflow/Source/DataflowEditor/Private/Dataflow/DataflowEngineRendering.cpp
@@ -23,6 +23,11 @@ namespace UE::Dataflow
 {
 	namespace Private
 	{
+		int32 DataflowPointRenderLimit = 100000;
+		FAutoConsoleVariableRef CVARDataflowPointRenderLimit(TEXT("p.Dataflow.PointRenderLimit"), DataflowPointRenderLimit,
+			TEXT("Limit for the number of points rendered in a dataflow node visualization. Set to -1 to remove limit."));
+
+
 		void RenderDynamicMesh(GeometryCollection::Facades::FRenderingFacade& RenderCollection, 
 			const FString& GeometryGroupName, 
 			const UE::Geometry::FDynamicMesh3& DynamicMesh, 
@@ -874,8 +879,16 @@ namespace UE::Dataflow
 				TArray<FVector> Default;
 				const TArray<FVector>& Points = State.GetValue<TArray<FVector>>(PrimaryOutput, Default);
 
-				TArray<FVector3f> PointsArr; PointsArr.AddUninitialized(Points.Num());
-				for (int32 Idx = 0; Idx < Points.Num(); ++Idx)
+				int32 NumPointsToRender = Points.Num();
+				if (UE::Dataflow::Private::DataflowPointRenderLimit > -1 && NumPointsToRender > UE::Dataflow::Private::DataflowPointRenderLimit)
+				{
+					UE_LOG(LogChaosDataflow, Warning, TEXT("Limited the number of points rendered from %d to %d; to see all points, adjust CVAR: p.Dataflow.PointRenderLimit"),
+						NumPointsToRender, UE::Dataflow::Private::DataflowPointRenderLimit);
+					NumPointsToRender = UE::Dataflow::Private::DataflowPointRenderLimit;
+				}
+
+				TArray<FVector3f> PointsArr; PointsArr.AddUninitialized(NumPointsToRender);
+				for (int32 Idx = 0; Idx < NumPointsToRender; ++Idx)
 				{
 					PointsArr[Idx] = (FVector3f)Points[Idx];
 				}
