Fix crash when correcting skin weights with multiple geometry per collection

diff --git a/Engine/Plugins/ChaosClothAsset/Source/ChaosClothAsset/Private/ChaosClothAsset/CollectionClothFacade.cpp b/Engine/Plugins/ChaosClothAsset/Source/ChaosClothAsset/Private/ChaosClothAsset/CollectionClothFacade.cpp
--- a/Engine/Plugins/ChaosClothAsset/Source/ChaosClothAsset/Private/ChaosClothAsset/CollectionClothFacade.cpp
+++ b/Engine/Plugins/ChaosClothAsset/Source/ChaosClothAsset/Private/ChaosClothAsset/CollectionClothFacade.cpp
@@ -811,7 +811,8 @@ namespace UE::Chaos::ClothAsset
 		static const TArray<FName> ClothCollectionGroups =
 		{
 			ClothCollectionGroup::SimFaces,  // Face int maps (self collision layers)
-			ClothCollectionGroup::RenderPatterns  // RecomputeTangents
+			ClothCollectionGroup::RenderPatterns,  // RecomputeTangents
+			ClothCollectionGroup::RenderVertices  
 		};
 
 		for (const FName& Group : ClothCollectionGroups)
diff --git a/Engine/Plugins/ChaosClothAssetEditor/Source/ChaosClothAssetDataflowNodes/Private/ChaosClothAsset/ChaosClothAssetDataflowNodesModule.cpp b/Engine/Plugins/ChaosClothAssetEditor/Source/ChaosClothAssetDataflowNodes/Private/ChaosClothAsset/ChaosClothAssetDataflowNodesModule.cpp
--- a/Engine/Plugins/ChaosClothAssetEditor/Source/ChaosClothAssetDataflowNodes/Private/ChaosClothAsset/ChaosClothAssetDataflowNodesModule.cpp
+++ b/Engine/Plugins/ChaosClothAssetEditor/Source/ChaosClothAssetDataflowNodes/Private/ChaosClothAsset/ChaosClothAssetDataflowNodesModule.cpp
@@ -362,7 +362,9 @@ PRAGMA_ENABLE_DEPRECATION_WARNINGS
 
 			virtual TArray<UE::Dataflow::FRenderingParameter> GetRenderingParameters() const override
 			{
-				return { {TEXT("SurfaceRender"), FName("FClothCollection"), {TEXT("Collection")} } };
+				return { {TEXT("SurfaceRender"), FName("FClothCollection"), {TEXT("Collection")}, FCloth2DSimViewMode::Name},
+						 {TEXT("SurfaceRender"), FName("FClothCollection"), {TEXT("Collection")}, FCloth3DSimViewMode::Name},
+						 {TEXT("SurfaceRender"), FName("FClothCollection"), {TEXT("Collection")}, FClothRenderViewMode::Name} };
 			}
 		};
 
diff --git a/Engine/Plugins/Experimental/Dataflow/Source/DataflowAssetTools/Private/Dataflow/CollectionRenderingPatternUtility.cpp b/Engine/Plugins/Experimental/Dataflow/Source/DataflowAssetTools/Private/Dataflow/CollectionRenderingPatternUtility.cpp
--- a/Engine/Plugins/Experimental/Dataflow/Source/DataflowAssetTools/Private/Dataflow/CollectionRenderingPatternUtility.cpp
+++ b/Engine/Plugins/Experimental/Dataflow/Source/DataflowAssetTools/Private/Dataflow/CollectionRenderingPatternUtility.cpp
@@ -15,7 +15,7 @@ namespace UE::Dataflow
 	namespace Conversion
 	{
 		// Convert a rendering facade to a dynamic mesh
-		void RenderingFacadeToDynamicMesh(const GeometryCollection::Facades::FRenderingFacade& Facade, int32 InMeshIndex, FDynamicMesh3& DynamicMesh)
+		void RenderingFacadeToDynamicMesh(const GeometryCollection::Facades::FRenderingFacade& Facade, int32 InMeshIndex, FDynamicMesh3& DynamicMesh, const bool bBuildRemapping)
 		{
 			if (Facade.CanRenderSurface())
 			{
@@ -66,7 +66,7 @@ namespace UE::Dataflow
 				DynamicMesh.EnableAttributes();
 
 				// Build Remmaping indices back into the colleciton. 
-				if (Remapping.Num() < Facade.NumVertices())
+				if (bBuildRemapping && (Remapping.Num() < Facade.NumVertices()))
 				{
 					UE::Geometry::FNonManifoldMappingSupport::AttachNonManifoldVertexMappingData(Remapping, DynamicMesh);
 				}
diff --git a/Engine/Plugins/Experimental/Dataflow/Source/DataflowAssetTools/Public/Dataflow/CollectionRenderingPatternUtility.h b/Engine/Plugins/Experimental/Dataflow/Source/DataflowAssetTools/Public/Dataflow/CollectionRenderingPatternUtility.h
--- a/Engine/Plugins/Experimental/Dataflow/Source/DataflowAssetTools/Public/Dataflow/CollectionRenderingPatternUtility.h
+++ b/Engine/Plugins/Experimental/Dataflow/Source/DataflowAssetTools/Public/Dataflow/CollectionRenderingPatternUtility.h
@@ -13,7 +13,7 @@ namespace UE::Dataflow
 	namespace Conversion
 	{
 		// Convert a rendering facade to a dynamic mesh. if InMeshIndex==INDEX_NONE then all Geometry is converted. 
-		void DATAFLOWASSETTOOLS_API RenderingFacadeToDynamicMesh(const GeometryCollection::Facades::FRenderingFacade& Facade, int32 InMeshIndex, UE::Geometry::FDynamicMesh3& DynamicMesh);
+		void DATAFLOWASSETTOOLS_API RenderingFacadeToDynamicMesh(const GeometryCollection::Facades::FRenderingFacade& Facade, int32 InMeshIndex, UE::Geometry::FDynamicMesh3& DynamicMesh, const bool bBuildRemapping = true);
 
 		// Convert a dynamic mesh to a rendering facade
 		void DATAFLOWASSETTOOLS_API DynamicMeshToRenderingFacade(const UE::Geometry::FDynamicMesh3& DynamicMesh, GeometryCollection::Facades::FRenderingFacade& Facade);
diff --git a/Engine/Plugins/Experimental/Dataflow/Source/DataflowEditor/Private/DataflowEditorTools/DataflowEditorCorrectSkinWeightsNode.cpp b/Engine/Plugins/Experimental/Dataflow/Source/DataflowEditor/Private/DataflowEditorTools/DataflowEditorCorrectSkinWeightsNode.cpp
--- a/Engine/Plugins/Experimental/Dataflow/Source/DataflowEditor/Private/DataflowEditorTools/DataflowEditorCorrectSkinWeightsNode.cpp
+++ b/Engine/Plugins/Experimental/Dataflow/Source/DataflowEditor/Private/DataflowEditorTools/DataflowEditorCorrectSkinWeightsNode.cpp
@@ -2,7 +2,7 @@
 
 #include "DataflowEditorTools/DataflowEditorCorrectSkinWeightsNode.h"
 
-#include "Chaos/Deformable/MuscleActivationConstraints.h"
+#include "Dataflow/DataflowCollectionEditSkinWeightsNode.h"
 #include "Dataflow/DataflowInputOutput.h"
 #include "Dataflow/DataflowConnectionTypes.h"
 #include "Dataflow/DataflowNode.h"
@@ -33,19 +33,22 @@ static void BuildDynamicMeshes(const FDataflowNode* DataflowNode, UE::Dataflow::
 		TSharedPtr<FManagedArrayCollection> RenderCollection(new FManagedArrayCollection);
 		GeometryCollection::Facades::FRenderingFacade RenderingFacade(*RenderCollection);
 		RenderingFacade.DefineSchema();
-	
-		UE::Dataflow::FDataflowConstruction3DViewMode ViewMode;
-		FGraphRenderingState RenderingState(FGuid::NewGuid(), DataflowNode, RenderParameters.Last(), Context, ViewMode, false);
-		Dataflow::FRenderingFactory::GetInstance()->RenderNodeOutput(RenderingFacade, RenderingState);
 
-		const int32 NumGeometry = RenderingFacade.NumGeometry();
-		for (int32 MeshIndex = 0; MeshIndex < NumGeometry; ++MeshIndex)
+		if(const IDataflowConstructionViewMode* ViewMode =
+			UE::Dataflow::FRenderingViewModeFactory::GetInstance().GetViewMode(RenderParameters.Last().ViewMode))
 		{
-			UE::Geometry::FDynamicMesh3& DynamicMesh = DynamicMeshes.AddDefaulted_GetRef();
-			UE::Dataflow::Conversion::RenderingFacadeToDynamicMesh(RenderingFacade, MeshIndex, DynamicMesh);
+			FGraphRenderingState RenderingState(FGuid::NewGuid(), DataflowNode, RenderParameters.Last(), Context, *ViewMode, false);
+			Dataflow::FRenderingFactory::GetInstance()->RenderNodeOutput(RenderingFacade, RenderingState);
+
+			const int32 NumGeometry = RenderingFacade.NumGeometry();
+			for (int32 MeshIndex = 0; MeshIndex < NumGeometry; ++MeshIndex)
+			{
+				UE::Geometry::FDynamicMesh3& DynamicMesh = DynamicMeshes.AddDefaulted_GetRef();
+				UE::Dataflow::Conversion::RenderingFacadeToDynamicMesh(RenderingFacade, MeshIndex, DynamicMesh, false);
 
-			DynamicMesh.Attributes()->AttachSkinWeightsAttribute(FName("Default"),
-				new UE::Geometry::FDynamicMeshVertexSkinWeightsAttribute(&DynamicMesh));
+				DynamicMesh.Attributes()->AttachSkinWeightsAttribute(FName("Default"),
+					new UE::Geometry::FDynamicMeshVertexSkinWeightsAttribute(&DynamicMesh));
+			}
 		}
 	}
 }
@@ -286,17 +289,16 @@ static bool SetAttributeValues(FManagedArrayCollection& SelectedCollection,
 	}
 	return false;
 }
-
-template<typename ArrayType>
-static bool GetAttributeValues(FManagedArrayCollection& SelectedCollection,
-	const FCollectionAttributeKey& AttributeKey, TArray<ArrayType>& AttributeValues)
+	
+static bool GetSkinningSelection(FManagedArrayCollection& SelectedCollection,
+	const FCollectionAttributeKey& AttributeKey, TArray<float>& AttributeValues)
 {
 	if (!AttributeKey.Attribute.IsEmpty() && !AttributeKey.Group.IsEmpty())
 	{
 		const FName AttributeName(AttributeKey.Attribute);
 		const FName AttributeGroup(AttributeKey.Group);
 		
-		TManagedArray<ArrayType>& AttributeArray = SelectedCollection.AddAttribute<ArrayType>(AttributeName, AttributeGroup);
+		TManagedArray<float>& AttributeArray = SelectedCollection.AddAttribute<float>(AttributeName, AttributeGroup);
 		AttributeValues = AttributeArray.GetConstArray();
 		return true;
 	}
@@ -354,53 +356,53 @@ void FDataflowCorrectSkinWeightsNode::Evaluate(UE::Dataflow::FContext& Context,
 		{
 			TArray<TArray<float>> SetupWeights, FinalWeights;
 			TArray<TArray<int32>> SetupIndices, FinalIndices;
-			
-			UE::Dataflow::Private::GetAttributeValues<TArray<int32>>(InCollection, BoneIndicesKeyValue, SetupIndices);
-		    UE::Dataflow::Private::GetAttributeValues<TArray<float>>(InCollection, BoneWeightsKeyValue, SetupWeights);
-
-			TArray<float> SelectionMap;
-			UE::Dataflow::Private::GetAttributeValues<float>(InCollection, SelectionMapKeyValue, SelectionMap);
 
-			FinalIndices.SetNum(SetupIndices.Num());
-			FinalWeights.SetNum(SetupWeights.Num());
-			
-			if(CorrectionType == ESkinWeightsCorrectionType::Prune)
-			{
-				UE::Dataflow::Private::PruneSkinWeights(PruningThreshold, SetupIndices, SetupWeights, SelectionMap, FinalIndices, FinalWeights);
-			}
-			else if(CorrectionType == ESkinWeightsCorrectionType::Clamp)
+			if(FDataflowCollectionEditSkinWeightsNode::FillAttributeWeights(InCollection, BoneIndicesKeyValue, BoneWeightsKeyValue, SetupIndices, SetupWeights))
 			{
-				UE::Dataflow::Private::ClampSkinWeights(ClampingNumber, SetupIndices, SetupWeights, SelectionMap, FinalIndices, FinalWeights);
-			}
-			else if(CorrectionType == ESkinWeightsCorrectionType::Normalize)
-			{
-				UE::Dataflow::Private::NormalizeSkinWeights(SetupIndices, SetupWeights, SelectionMap, FinalIndices, FinalWeights);
-			}
-			else
-			{
-				// Do the work here
-				TArray<UE::Geometry::FDynamicMesh3> DynamicMeshes;
-				UE::Dataflow::Private::BuildDynamicMeshes(this, Context, DynamicMeshes);
-
-				int32 VertexOffset = 0;
-				for(UE::Geometry::FDynamicMesh3& DynamicMesh : DynamicMeshes)
+				TArray<float> SelectionMap;
+				if(UE::Dataflow::Private::GetSkinningSelection(InCollection, SelectionMapKeyValue, SelectionMap))
 				{
-					if(CorrectionType == ESkinWeightsCorrectionType::Relax)
+					FinalIndices.SetNum(SetupIndices.Num());
+					FinalWeights.SetNum(SetupWeights.Num());
+				
+					if(CorrectionType == ESkinWeightsCorrectionType::Prune)
 					{
-						UE::Dataflow::Private::RelaxSkinWeights(DynamicMesh, VertexOffset, SmoothingFactor, SmoothingIterations,
-							SetupIndices, SetupWeights, SelectionMap, FinalIndices, FinalWeights);
+						UE::Dataflow::Private::PruneSkinWeights(PruningThreshold, SetupIndices, SetupWeights, SelectionMap, FinalIndices, FinalWeights);
 					}
-					else if (CorrectionType == ESkinWeightsCorrectionType::Hammer)
+					else if(CorrectionType == ESkinWeightsCorrectionType::Clamp)
 					{
-						UE::Dataflow::Private::HammerSkinWeights(DynamicMesh, VertexOffset,
-							SetupIndices, SetupWeights, SelectionMap, FinalIndices, FinalWeights);
+						UE::Dataflow::Private::ClampSkinWeights(ClampingNumber, SetupIndices, SetupWeights, SelectionMap, FinalIndices, FinalWeights);
 					}
-					VertexOffset += DynamicMesh.VertexCount();
+					else if(CorrectionType == ESkinWeightsCorrectionType::Normalize)
+					{
+						UE::Dataflow::Private::NormalizeSkinWeights(SetupIndices, SetupWeights, SelectionMap, FinalIndices, FinalWeights);
+					}
+					else
+					{
+						// Do the work here
+						TArray<UE::Geometry::FDynamicMesh3> DynamicMeshes;
+						UE::Dataflow::Private::BuildDynamicMeshes(this, Context, DynamicMeshes);
+
+						int32 VertexOffset = 0;
+						for(UE::Geometry::FDynamicMesh3& DynamicMesh : DynamicMeshes)
+						{
+							if(CorrectionType == ESkinWeightsCorrectionType::Relax)
+							{
+								UE::Dataflow::Private::RelaxSkinWeights(DynamicMesh, VertexOffset, SmoothingFactor, SmoothingIterations,
+									SetupIndices, SetupWeights, SelectionMap, FinalIndices, FinalWeights);
+							}
+							else if (CorrectionType == ESkinWeightsCorrectionType::Hammer)
+							{
+								UE::Dataflow::Private::HammerSkinWeights(DynamicMesh, VertexOffset,
+									SetupIndices, SetupWeights, SelectionMap, FinalIndices, FinalWeights);
+							}
+							VertexOffset += DynamicMesh.VertexCount();
+						}
+					}
+
+					FDataflowCollectionEditSkinWeightsNode::SetAttributeWeights(InCollection, BoneIndicesKeyValue, BoneWeightsKeyValue, FinalIndices, FinalWeights);
 				}
 			}
-
-			UE::Dataflow::Private::SetAttributeValues<TArray<int32>>(InCollection, BoneIndicesKeyValue, FinalIndices);
-			UE::Dataflow::Private::SetAttributeValues<TArray<float>>(InCollection, BoneWeightsKeyValue, FinalWeights);
 		}
 
 		SetValue(Context, MoveTemp(InCollection), &Collection);
diff --git a/Engine/Plugins/Experimental/Dataflow/Source/DataflowEditor/Private/DataflowEditorTools/DataflowEditorSkinWeightsPaintTool.cpp b/Engine/Plugins/Experimental/Dataflow/Source/DataflowEditor/Private/DataflowEditorTools/DataflowEditorSkinWeightsPaintTool.cpp
--- a/Engine/Plugins/Experimental/Dataflow/Source/DataflowEditor/Private/DataflowEditorTools/DataflowEditorSkinWeightsPaintTool.cpp
+++ b/Engine/Plugins/Experimental/Dataflow/Source/DataflowEditor/Private/DataflowEditorTools/DataflowEditorSkinWeightsPaintTool.cpp
@@ -1,4 +1,4 @@
-﻿// Copyright Epic Games, Inc. All Rights Reserved.
+// Copyright Epic Games, Inc. All Rights Reserved.
 
 #include "DataflowEditorTools/DataflowEditorSkinWeightsPaintTool.h"
 
@@ -110,7 +110,8 @@ FMeshDescription* UDataflowEditorSkinWeightsPaintTool::GetCurrentDescription(con
 	return nullptr;
 }
 
-bool UDataflowEditorSkinWeightsPaintTool::FillVertexOffsets(TArray<int32>& VertexOffsets) const 
+
+int32 UDataflowEditorSkinWeightsPaintTool::GetVertexOffset() const
 {
 	if(const IPrimitiveComponentBackedTarget* TargetComponent = Cast<IPrimitiveComponentBackedTarget>(Target))
 	{
@@ -118,32 +119,16 @@ bool UDataflowEditorSkinWeightsPaintTool::FillVertexOffsets(TArray<int32>& Verte
 		{
 			if(USkeletalMesh* SkeletalMesh = SkeletalMeshComponent->GetSkeletalMeshAsset())
 			{
-				const int32 NumLODs = SkeletalMesh->GetLODNum();
-				VertexOffsets.Init(0, NumLODs+1);
-
-				int32 NumVertices = 0;
-				for(int32 LODIndex = 0; LODIndex < NumLODs; ++LODIndex)
-				{
-					VertexOffsets[LODIndex] = NumVertices;
-					if(FMeshDescription* MeshDescription = SkeletalMesh->GetMeshDescription(static_cast<int32>(LODIndex)))
-					{
-						NumVertices += MeshDescription->Vertices().Num();
-					}
-				}
-				VertexOffsets[NumLODs] = NumVertices;
-				return true;
+				return SkinWeightNode->GetSkeletalMeshOffset(SkeletalMesh);
 			}
 		}
 	}
-	return false;
+	return INDEX_NONE;
 }
 
-bool UDataflowEditorSkinWeightsPaintTool::ExtractSkinWeights(const int32 NumVertices, TArray<TArray<int32>>& CurrentIndices, TArray<TArray<float>>& CurrentWeights)
+bool UDataflowEditorSkinWeightsPaintTool::ExtractSkinWeights(TArray<TArray<int32>>& CurrentIndices, TArray<TArray<float>>& CurrentWeights)
 {
 	// Setup DynamicMeshToWeight conversion and get Input weight map (if it exists)
-	SetupWeights.SetNumZeroed(NumVertices);
-	SetupIndices.SetNumZeroed(NumVertices);
-
 	if(DataflowEditorContextObject && SkinWeightNode)
 	{
 		// Find the map if it exists.
@@ -153,11 +138,11 @@ bool UDataflowEditorSkinWeightsPaintTool::ExtractSkinWeights(const int32 NumVert
 			{
 				// Fill the attribute values
 				SkinWeightNode->FillAttributeWeights(*DataflowEditorContextObject->GetSelectedCollection().Get(),
-					SkinWeightNode->GetBoneIndicesKey(*DataflowContext), SkinWeightNode->GetBoneWeightsKey(*DataflowContext),
+					SkinWeightNode->GetBoneIndicesKey(*DataflowContext), SkinWeightNode->GetBoneWeightsKey(*DataflowContext), 
 					SetupIndices, SetupWeights);
 						
-				CurrentWeights.SetNumZeroed(NumVertices);
-				CurrentIndices.SetNumZeroed(NumVertices);
+				CurrentWeights.SetNumZeroed(SetupWeights.Num()); 
+				CurrentIndices.SetNumZeroed(SetupIndices.Num());
 						
 				SkinWeightNode->ExtractVertexWeights(SetupIndices, SetupWeights,
 					TArrayView<TArray<int32>>(CurrentIndices), TArrayView<TArray<float>>(CurrentWeights));
@@ -171,45 +156,41 @@ bool UDataflowEditorSkinWeightsPaintTool::ExtractSkinWeights(const int32 NumVert
 
 void UDataflowEditorSkinWeightsPaintTool::Setup()
 {
-	TArray<int32> VertexOffsets;
-	if(FillVertexOffsets(VertexOffsets))
+	const int32 VertexOffset = GetVertexOffset();
+	if(VertexOffset != INDEX_NONE)
 	{
-		TArray<TArray<float>> CurrentWeights;
-		TArray<TArray<int32>> CurrentIndices;
-		
-		if(ExtractSkinWeights(VertexOffsets.Last(), CurrentIndices, CurrentWeights))
+		if(FMeshDescription* MeshDescription = GetCurrentDescription(0))
 		{
-			for(int32 LODIndex = 0, NumLODs = VertexOffsets.Num()-1; LODIndex < NumLODs; ++LODIndex)
+			const int32 NumVertices = MeshDescription->Vertices().Num();
+			
+			TArray<TArray<float>> CurrentWeights;
+			TArray<TArray<int32>> CurrentIndices;
+	
+			if(ExtractSkinWeights(CurrentIndices, CurrentWeights))
 			{
-				if(FMeshDescription* MeshDescription = GetCurrentDescription(LODIndex))
+				FSkeletalMeshAttributes MeshAttribs(*MeshDescription);
+				FSkinWeightsVertexAttributesRef SkinWeights =
+					MeshAttribs.GetVertexSkinWeights(FSkeletalMeshAttributesShared::DefaultSkinWeightProfileName);
+			
+				UE::AnimationCore::FBoneWeightsSettings WeightsSettings;
+				WeightsSettings.SetNormalizeType(UE::AnimationCore::EBoneWeightNormalizeType::Always);
+
+				for(int32 VertexIndex = 0; VertexIndex < NumVertices; ++VertexIndex)
 				{
-					const int32 NumVertices = MeshDescription->Vertices().Num();
+					TArray<UE::AnimationCore::FBoneWeight> BoneWeights;
+					const int32 CurrentIndex = VertexOffset+VertexIndex;
 				
-					FSkeletalMeshAttributes MeshAttribs(*MeshDescription);
-					FSkinWeightsVertexAttributesRef SkinWeights =
-						MeshAttribs.GetVertexSkinWeights(FSkeletalMeshAttributesShared::DefaultSkinWeightProfileName);
-				
-					UE::AnimationCore::FBoneWeightsSettings WeightsSettings;
-					WeightsSettings.SetNormalizeType(UE::AnimationCore::EBoneWeightNormalizeType::Always);
-
-					const int32 VertexOffset = VertexOffsets[LODIndex];
-					for(int32 VertexIndex = 0; VertexIndex < NumVertices; ++VertexIndex)
+					if(CurrentIndex < CurrentWeights.Num() && CurrentIndex < CurrentIndices.Num() &&
+						CurrentWeights[CurrentIndex].Num() == CurrentIndices[CurrentIndex].Num())
 					{
-						TArray<UE::AnimationCore::FBoneWeight> BoneWeights;
-						const int32 CurrentIndex = VertexOffset+VertexIndex;
-					
-						if(CurrentIndex < CurrentWeights.Num() && CurrentIndex < CurrentIndices.Num() &&
-							CurrentWeights[CurrentIndex].Num() == CurrentIndices[CurrentIndex].Num())
+						BoneWeights.Reserve(CurrentWeights[CurrentIndex].Num());
+						for(int32 WeightIndex = 0, NumWeights = CurrentWeights[CurrentIndex].Num(); WeightIndex < NumWeights; ++WeightIndex)
 						{
-							BoneWeights.Reserve(CurrentWeights[CurrentIndex].Num());
-							for(int32 WeightIndex = 0, NumWeights = CurrentWeights[CurrentIndex].Num(); WeightIndex < NumWeights; ++WeightIndex)
-							{
-								BoneWeights.Add(UE::AnimationCore::FBoneWeight(CurrentIndices[CurrentIndex][WeightIndex],
-									CurrentWeights[CurrentIndex][WeightIndex]));
-							}
+							BoneWeights.Add(UE::AnimationCore::FBoneWeight(CurrentIndices[CurrentIndex][WeightIndex],
+								CurrentWeights[CurrentIndex][WeightIndex]));
 						}
-						SkinWeights.Set(FVertexID(VertexIndex), UE::AnimationCore::FBoneWeights::Create(BoneWeights, WeightsSettings));
 					}
+					SkinWeights.Set(FVertexID(VertexIndex), UE::AnimationCore::FBoneWeights::Create(BoneWeights, WeightsSettings));
 				}
 			}
 		}
@@ -223,19 +204,21 @@ void UDataflowEditorSkinWeightsPaintTool::OnShutdown(EToolShutdownType ShutdownT
 
 	if(Target && ShutdownType == EToolShutdownType::Accept)
 	{
-		TArray<int32> VertexOffsets;
-		if(FillVertexOffsets(VertexOffsets))
+		const int32 VertexOffset = GetVertexOffset();
+		if(VertexOffset != INDEX_NONE)
 		{
 			TArray<TArray<int32>> CurrentIndices;
 			TArray<TArray<float>> CurrentWeights;
-			
+
+			CurrentWeights.SetNumZeroed(SetupWeights.Num()); 
 			CurrentIndices.SetNumZeroed(SetupIndices.Num());
-			CurrentWeights.SetNumZeroed(SetupWeights.Num());
+						
+			SkinWeightNode->ExtractVertexWeights(SetupIndices, SetupWeights,
+				TArrayView<TArray<int32>>(CurrentIndices), TArrayView<TArray<float>>(CurrentWeights));
 			
-			for(TPair<EMeshLODIdentifier, FCleanedEditMesh>& EditedeMesh : EditedMeshes)
+			if(FCleanedEditMesh* EditedMesh = EditedMeshes.Find(EMeshLODIdentifier::LOD0))
 			{
-				const FMeshDescription& MeshDescription = EditedeMesh.Value.GetEditableMeshDescription();
-				const int32 LODIndex = static_cast<int32>(EditedeMesh.Key);
+				const FMeshDescription& MeshDescription = EditedMesh->GetEditableMeshDescription();
 			
 				// profile to edit
 				const FName ActiveProfile = WeightToolProperties->GetActiveSkinWeightProfile();
@@ -244,24 +227,33 @@ void UDataflowEditorSkinWeightsPaintTool::OnShutdown(EToolShutdownType ShutdownT
 				FSkinWeightsVertexAttributesConstRef SkinWeights = MeshAttribs.GetVertexSkinWeights(ActiveProfile);
 			
 				const int32 NumVertices = MeshDescription.Vertices().Num();
-				if(LODIndex < VertexOffsets.Num())
+				for(int32 VertexIndex = 0; VertexIndex < NumVertices; ++VertexIndex)
 				{
-					const int32 VertexOffset = VertexOffsets[LODIndex];
-					for(int32 VertexIndex = 0; VertexIndex < NumVertices; ++VertexIndex)
+					FVertexBoneWeightsConst BoneWeights = SkinWeights.Get(FVertexID(VertexIndex));
+					const int32 CurrentIndex = VertexOffset+VertexIndex;
+			
+					if(CurrentIndex < CurrentWeights.Num() && CurrentIndex < CurrentIndices.Num() &&
+						CurrentWeights[CurrentIndex].Num() == CurrentIndices[CurrentIndex].Num())
 					{
-						FVertexBoneWeightsConst BoneWeights = SkinWeights.Get(FVertexID(VertexIndex));
-						const int32 CurrentIndex = VertexOffset+VertexIndex;
-				
-						if(CurrentIndex < CurrentWeights.Num() && CurrentIndex < CurrentIndices.Num() &&
-							CurrentWeights[CurrentIndex].Num() == CurrentIndices[CurrentIndex].Num())
+						CurrentIndices[CurrentIndex].Reset(BoneWeights.Num());
+						CurrentWeights[CurrentIndex].Reset(BoneWeights.Num());
+
+						TArray<TPair<int32,float>> SortedWeights;
+						SortedWeights.Reserve(BoneWeights.Num());
+						
+						for(int32 WeightIndex = 0, NumWeights = BoneWeights.Num(); WeightIndex < NumWeights; ++WeightIndex)
+						{
+							SortedWeights.Add({BoneWeights[WeightIndex].GetBoneIndex(), BoneWeights[WeightIndex].GetWeight()});
+						}
+						Algo::Sort(SortedWeights, [](const TPair<int32,float>& A, const TPair<int32,float>& B)
+						{
+							return A.Value > B.Value; 
+						});
+
+						for(int32 WeightIndex = 0, NumWeights = BoneWeights.Num(); WeightIndex < NumWeights; ++WeightIndex)
 						{
-							CurrentIndices[CurrentIndex].Reserve(BoneWeights.Num());
-							CurrentWeights[CurrentIndex].Reserve(BoneWeights.Num());
-							for(int32 WeightIndex = 0, NumWeights = BoneWeights.Num(); WeightIndex < NumWeights; ++WeightIndex)
-							{
-								CurrentIndices[CurrentIndex].Add(BoneWeights[WeightIndex].GetBoneIndex());
-								CurrentWeights[CurrentIndex].Add(BoneWeights[WeightIndex].GetWeight());
-							}
+							CurrentIndices[CurrentIndex].Add(SortedWeights[WeightIndex].Key);
+							CurrentWeights[CurrentIndex].Add(SortedWeights[WeightIndex].Value);
 						}
 					}
 				}
@@ -270,7 +262,7 @@ void UDataflowEditorSkinWeightsPaintTool::OnShutdown(EToolShutdownType ShutdownT
 			SkinWeightNode->Invalidate();
 
 			// Avoid rebuilding the skeletal mesh after updating the skin weights
-			SkinWeightNode->ValidateSkeletalMesh();
+			SkinWeightNode->ValidateSkeletalMeshes();
 		}
 	}
 }
diff --git a/Engine/Plugins/Experimental/Dataflow/Source/DataflowEditor/Public/DataflowEditorTools/DataflowEditorCorrectSkinWeightsNode.h b/Engine/Plugins/Experimental/Dataflow/Source/DataflowEditor/Public/DataflowEditorTools/DataflowEditorCorrectSkinWeightsNode.h
--- a/Engine/Plugins/Experimental/Dataflow/Source/DataflowEditor/Public/DataflowEditorTools/DataflowEditorCorrectSkinWeightsNode.h
+++ b/Engine/Plugins/Experimental/Dataflow/Source/DataflowEditor/Public/DataflowEditorTools/DataflowEditorCorrectSkinWeightsNode.h
@@ -58,15 +58,15 @@ public:
 	FScalarVertexPropertyGroup VertexGroup;
 
 	/** Bone indices key to be used in other nodes if necessary */
-	UPROPERTY(meta = (DisplayName = "Bone Indices Key", DataflowInput, DataflowOutput, DataflowPassthrough = "BoneIndicesKey"))
+	UPROPERTY(meta = (DisplayName = "Bone Indices", DataflowInput, DataflowOutput, DataflowPassthrough = "BoneIndicesKey"))
 	FCollectionAttributeKey BoneIndicesKey;
 
 	/** Bone weights key to be used in other nodes if necessary */
-	UPROPERTY(meta = (DisplayName = "Bone Weights Key", DataflowInput, DataflowOutput, DataflowPassthrough = "BoneWeightsKey"))
+	UPROPERTY(meta = (DisplayName = "Bone Weights", DataflowInput, DataflowOutput, DataflowPassthrough = "BoneWeightsKey"))
 	FCollectionAttributeKey BoneWeightsKey;
 
 	/** Selection map key to be used in other nodes if necessary */
-	UPROPERTY(meta = (DisplayName = "Selection Map Key", DataflowInput))
+	UPROPERTY(meta = (DisplayName = "Selection Map", DataflowInput))
 	FCollectionAttributeKey SelectionMapKey;
 	
 	/** Selection map key to be used in other nodes if necessary */
@@ -129,7 +129,7 @@ public:
 	FScalarVertexPropertyGroup VertexGroup;
 
 	/** Selection map key to be used in other nodes if necessary */
-	UPROPERTY(meta = (DisplayName = "Selection Map Key", DataflowInput))
+	UPROPERTY(meta = (DisplayName = "Selection Map", DataflowInput))
 	FCollectionAttributeKey SelectionMapKey;
 	
 	/** Selection map key to be used in other nodes if necessary */
@@ -168,7 +168,7 @@ public:
 	FScalarVertexPropertyGroup VertexGroup;
 
 	/** Selection map key to be used in other nodes if necessary */
-	UPROPERTY(meta = (DisplayName = "Selection Map Key", DataflowOutput))
+	UPROPERTY(meta = (DisplayName = "Selection Map", DataflowOutput))
 	FCollectionAttributeKey SelectionMapKey;
 	
 	/** Selection map key to be used in other nodes if necessary */
diff --git a/Engine/Plugins/Experimental/Dataflow/Source/DataflowEditor/Public/DataflowEditorTools/DataflowEditorSkinWeightsPaintTool.h b/Engine/Plugins/Experimental/Dataflow/Source/DataflowEditor/Public/DataflowEditorTools/DataflowEditorSkinWeightsPaintTool.h
--- a/Engine/Plugins/Experimental/Dataflow/Source/DataflowEditor/Public/DataflowEditorTools/DataflowEditorSkinWeightsPaintTool.h
+++ b/Engine/Plugins/Experimental/Dataflow/Source/DataflowEditor/Public/DataflowEditorTools/DataflowEditorSkinWeightsPaintTool.h
@@ -61,10 +61,10 @@ private:
 	}
 
 	/** Extract the skin weights from the node */
-	bool ExtractSkinWeights(const int32 NumVertices, TArray<TArray<int32>>& CurrentIndices, TArray<TArray<float>>& CurrentWeights);
-
-	/** Fill the vertex offsets fromm all the LODs description */
-	bool FillVertexOffsets(TArray<int32>& VertexOffsets) const;
+	bool ExtractSkinWeights(TArray<TArray<int32>>& CurrentIndices, TArray<TArray<float>>& CurrentWeights);
+	
+	/** Get the skeletal mesh vertex offset in the node collection */
+	int32 GetVertexOffset() const;
 
 	/** Get the LOD description given a  LOD index */
 	FMeshDescription* GetCurrentDescription(const int32 LODIndex) const;
diff --git a/Engine/Plugins/Experimental/Dataflow/Source/DataflowNodes/Private/Dataflow/DataflowCollectionAddScalarVertexPropertyNode.cpp b/Engine/Plugins/Experimental/Dataflow/Source/DataflowNodes/Private/Dataflow/DataflowCollectionAddScalarVertexPropertyNode.cpp
--- a/Engine/Plugins/Experimental/Dataflow/Source/DataflowNodes/Private/Dataflow/DataflowCollectionAddScalarVertexPropertyNode.cpp
+++ b/Engine/Plugins/Experimental/Dataflow/Source/DataflowNodes/Private/Dataflow/DataflowCollectionAddScalarVertexPropertyNode.cpp
@@ -75,7 +75,7 @@ TArray<UE::Dataflow::FRenderingParameter> FDataflowAddScalarVertexPropertyCallba
 				}
 			}
 		}
-		else
+		else if(TargetGroups.Find(TargetGroup) != INDEX_NONE)
 		{
 			for (const UE::Dataflow::FRenderingParameter& RenderingParameter : CallbacksEntry.Value->GetRenderingParameters())
 			{
diff --git a/Engine/Plugins/Experimental/Dataflow/Source/DataflowNodes/Private/Dataflow/DataflowCollectionEditSkinWeightsNode.cpp b/Engine/Plugins/Experimental/Dataflow/Source/DataflowNodes/Private/Dataflow/DataflowCollectionEditSkinWeightsNode.cpp
--- a/Engine/Plugins/Experimental/Dataflow/Source/DataflowNodes/Private/Dataflow/DataflowCollectionEditSkinWeightsNode.cpp
+++ b/Engine/Plugins/Experimental/Dataflow/Source/DataflowNodes/Private/Dataflow/DataflowCollectionEditSkinWeightsNode.cpp
@@ -15,6 +15,8 @@
 #include "Engine/SkeletalMesh.h"
 #include "Animation/Skeleton.h"
 #include "Components/SkeletalMeshComponent.h"
+#include "Materials/Material.h"
+#include "MaterialDomain.h"
 
 #if WITH_EDITOR
 #include "StaticToSkeletalMeshConverter.h"
@@ -37,13 +39,13 @@ static bool CreateSkeletalMesh(USkeletalMesh* SkeletalMesh, const TArray<UE::Geo
 	SkeletalMesh->GetImportedModel()->LODModels.Empty();
 #endif
 	SkeletalMesh->ResetLODInfo();
-	
+
 	TArray<FMeshDescription> LocalDescriptions;
 	for (const UE::Geometry::FDynamicMesh3& DynamicMesh : DynamicMeshes)
 	{
 		// Create a mesh description
 		FMeshDescription& MeshDescription = LocalDescriptions.AddDefaulted_GetRef();
-		
+	
 		// Add skeletal mesh attributes to the mesh description
 		FSkeletalMeshAttributes Attributes(MeshDescription);
 		Attributes.Register();
@@ -63,14 +65,14 @@ static bool CreateSkeletalMesh(USkeletalMesh* SkeletalMesh, const TArray<UE::Geo
 	// ensure there is at least one material
 	if (MaterialView.IsEmpty())
 	{
-		Materials.Add(FSkeletalMaterial());
+		Materials.Add( UMaterial::GetDefaultMaterial(MD_Surface));
 		MaterialView = Materials;
 	}
 
 	static constexpr bool bRecomputeTangents = false;
 	static constexpr bool bRecomputeNormals = false;
 	static constexpr bool bCacheOptimize = false;
-	
+
 	if (!FStaticToSkeletalMeshConverter::InitializeSkeletalMeshFromMeshDescriptions(
 		SkeletalMesh, MeshDescriptions, MaterialView, ReferenceSkeleton, bRecomputeNormals, bRecomputeTangents, bCacheOptimize))
 	{
@@ -81,8 +83,8 @@ static bool CreateSkeletalMesh(USkeletalMesh* SkeletalMesh, const TArray<UE::Geo
 	return false;
 #endif
 }
-	
-static bool BuildSkeletalMesh(USkeletalMesh* SkeletalMesh, 
+
+static bool BuildSkeletalMeshes(TArray<TObjectPtr<USkeletalMesh>>& SkeletalMeshes, 
 	const TSharedPtr<const FManagedArrayCollection>& RenderCollection, const TObjectPtr<USkeleton>& ObjectSkeleton)
 {
 	if(RenderCollection.IsValid() && ObjectSkeleton)
@@ -90,32 +92,38 @@ static bool BuildSkeletalMesh(USkeletalMesh* SkeletalMesh,
 		GeometryCollection::Facades::FRenderingFacade Facade(*RenderCollection);
 		if(Facade.IsValid())
 		{
-			TArray<UE::Geometry::FDynamicMesh3> DynamicMeshes;
+			bool bValidSkeletalMeshes = true;
 			const int32 NumGeometry = Facade.NumGeometry();
 
-			for (int32 MeshIndex = 0; MeshIndex < NumGeometry; ++MeshIndex)
+			if(NumGeometry == SkeletalMeshes.Num())
 			{
-				UE::Geometry::FDynamicMesh3& DynamicMesh = DynamicMeshes.AddDefaulted_GetRef();
-				UE::Dataflow::Conversion::RenderingFacadeToDynamicMesh(Facade, MeshIndex, DynamicMesh);
+				for (int32 MeshIndex = 0; MeshIndex < NumGeometry; ++MeshIndex)
+				{
+					UE::Geometry::FDynamicMesh3 DynamicMesh;
+					UE::Dataflow::Conversion::RenderingFacadeToDynamicMesh(Facade, MeshIndex, DynamicMesh);
+
+					if(!CreateSkeletalMesh(SkeletalMeshes[MeshIndex], {DynamicMesh}, ObjectSkeleton->GetReferenceSkeleton()))
+					{
+						bValidSkeletalMeshes = false;
+					}
+				}
 			}
-	
-			return CreateSkeletalMesh(SkeletalMesh, DynamicMeshes, ObjectSkeleton->GetReferenceSkeleton());
+			return bValidSkeletalMeshes;
 		}
 	}
 	return false;
 }
-
 	
-template<typename ArrayType>
-static bool SetAttributeValues(FManagedArrayCollection& SelectedCollection,
-	const FCollectionAttributeKey& AttributeKey, const TArray<TArray<ArrayType>>& AttributeValues)
+template<typename ScalarType, typename VectorType, int32 NumComponents>
+bool SetAttributeValues(FManagedArrayCollection& SelectedCollection,
+	const FCollectionAttributeKey& AttributeKey, const TArray<TArray<ScalarType>>& AttributeValues, const ScalarType& DefaultValue, const bool bRenormalizeValues)
 {
 	if (!AttributeValues.IsEmpty() && !AttributeKey.Attribute.IsEmpty() && !AttributeKey.Group.IsEmpty())
 	{
 		const FName AttributeName(AttributeKey.Attribute);
 		const FName AttributeGroup(AttributeKey.Group);
 		
-		if(TManagedArray<TArray<ArrayType>>* AttributeArray = SelectedCollection.FindAttributeTyped<TArray<ArrayType>>(AttributeName, AttributeGroup))
+		if(TManagedArray<TArray<ScalarType>>* AttributeArray = SelectedCollection.FindAttributeTyped<TArray<ScalarType>>(AttributeName, AttributeGroup))
 		{
 			if(AttributeArray->Num() == AttributeValues.Num())
 			{
@@ -126,42 +134,119 @@ static bool SetAttributeValues(FManagedArrayCollection& SelectedCollection,
 			}
 			return true;
 		}
+		else if(TManagedArray<VectorType>* AttributeVector = SelectedCollection.FindAttributeTyped<VectorType>(AttributeName, AttributeGroup))
+		{
+			if(AttributeVector->Num() == AttributeValues.Num())
+			{
+				for(int32 VertexIndex = 0, NumVertices = AttributeVector->Num(); VertexIndex < NumVertices; ++VertexIndex)
+				{
+					VectorType& ElementVector = AttributeVector->GetData()[VertexIndex];
+					const int32 NumValidComponents = FMath::Min(NumComponents, AttributeValues[VertexIndex].Num());
+					
+					float TotalValue = 0.0f;
+					for(int32 ComponentIndex = 0; ComponentIndex < NumValidComponents; ++ComponentIndex)
+					{
+						ElementVector[ComponentIndex] = AttributeValues[VertexIndex][ComponentIndex];
+						TotalValue += ElementVector[ComponentIndex];
+					}
+
+					for(int32 ComponentIndex = 0; ComponentIndex < NumComponents; ++ComponentIndex)
+					{
+						ElementVector[ComponentIndex] = (ComponentIndex >= NumValidComponents) ? DefaultValue : (bRenormalizeValues && TotalValue != 0.0f)
+							?  ElementVector[ComponentIndex] / TotalValue : ElementVector[ComponentIndex];
+					}
+				}
+			}
+			return true;
+		}
 	}
 	return false;
 }
 
-template<typename ArrayType>
-static bool GetAttributeValues(FManagedArrayCollection& SelectedCollection,
-	const FCollectionAttributeKey& AttributeKey, TArray<TArray<ArrayType>>& AttributeValues)
+template<typename ScalarType, typename VectorType, int32 NumComponents>
+bool FillAttributeValues(const FManagedArrayCollection& SelectedCollection,
+	const FCollectionAttributeKey& AttributeKey, TArray<TArray<ScalarType>>& AttributeValues)
 {
 	if (!AttributeKey.Attribute.IsEmpty() && !AttributeKey.Group.IsEmpty())
 	{
 		const FName AttributeName(AttributeKey.Attribute);
 		const FName AttributeGroup(AttributeKey.Group);
-		
-		TManagedArray<TArray<ArrayType>>& AttributeArray = SelectedCollection.AddAttribute<TArray<ArrayType>>(AttributeName, AttributeGroup);
-		AttributeValues = AttributeArray.GetConstArray();
-		return true;
+
+		if(const TManagedArray<TArray<ScalarType>>* AttributeArray = SelectedCollection.FindAttributeTyped<TArray<ScalarType>>(AttributeName, AttributeGroup))
+		{
+			AttributeValues = AttributeArray->GetConstArray();
+			return true;
+		}
+		else if(const TManagedArray<VectorType>* AttributeVector = SelectedCollection.FindAttributeTyped<VectorType>(AttributeName, AttributeGroup))
+		{
+			AttributeValues.SetNum(AttributeVector->Num());
+			for(int32 VertexIndex = 0, NumVertices = AttributeVector->Num(); VertexIndex < NumVertices; ++VertexIndex)
+			{
+				const VectorType& ElementVector = AttributeVector->GetConstArray()[VertexIndex];
+				AttributeValues[VertexIndex].SetNum(NumComponents);
+				
+				for(int32 ComponentIndex = 0; ComponentIndex < NumComponents; ++ComponentIndex)
+				{
+					AttributeValues[VertexIndex][ComponentIndex] = ElementVector[ComponentIndex];
+				}
+			}
+			return true;
+		}
 	}
 	return false;
 }
 
-template<typename ArrayType>
-static bool FillAttributeValues(const FManagedArrayCollection& SelectedCollection,
-	const FCollectionAttributeKey& AttributeKey, TArray<TArray<ArrayType>>& AttributeValues)
+template<typename ScalarType, typename VectorType, int32 NumComponents>
+bool GetAttributeValues(FManagedArrayCollection& SelectedCollection,
+	const FCollectionAttributeKey& AttributeKey, TArray<TArray<ScalarType>>& AttributeValues, const bool bVectorValues)
 {
 	if (!AttributeKey.Attribute.IsEmpty() && !AttributeKey.Group.IsEmpty())
 	{
 		const FName AttributeName(AttributeKey.Attribute);
 		const FName AttributeGroup(AttributeKey.Group);
 
-		if(const TManagedArray<TArray<ArrayType>>* AttributeArray = SelectedCollection.FindAttributeTyped<TArray<ArrayType>>(AttributeName, AttributeGroup))
+		if(SelectedCollection.FindAttributeTyped<TArray<ScalarType>>(AttributeName, AttributeGroup) == nullptr &&
+			SelectedCollection.FindAttributeTyped<VectorType>(AttributeName, AttributeGroup) == nullptr)
 		{
-			AttributeValues = AttributeArray->GetConstArray();
-			return true;
+			if(bVectorValues)
+			{
+				SelectedCollection.AddAttribute<VectorType>(AttributeName, AttributeGroup);
+			}
+			else
+			{
+				SelectedCollection.AddAttribute<TArray<ScalarType>>(AttributeName, AttributeGroup);
+			}
+		}
+	}
+	return FillAttributeValues<ScalarType,VectorType, NumComponents>(SelectedCollection, AttributeKey, AttributeValues);
+}
+
+void CorrectSkinWeights(TArray<TArray<int32>>& BoneIndices, TArray<TArray<float>>& BoneWeights)
+{
+	check(BoneIndices.Num() == BoneWeights.Num());
+
+	TArray<int32> ValidIndices;
+	TArray<float> ValidWeights;
+	
+	for(int32 VertexIndex = 0, NumVertices = BoneWeights.Num(); VertexIndex < NumVertices; ++VertexIndex)
+	{
+		if(BoneWeights[VertexIndex].Num() == BoneIndices[VertexIndex].Num())
+		{
+			ValidIndices.Reset(BoneIndices.Num());
+			ValidWeights.Reset(BoneWeights.Num());
+			
+			for(int32 WeightIndex = 0, NumWeights = BoneWeights[VertexIndex].Num(); WeightIndex < NumWeights; ++WeightIndex)
+			{
+				if(BoneIndices[VertexIndex][WeightIndex] != INDEX_NONE)
+				{
+					ValidIndices.Add(BoneIndices[VertexIndex][WeightIndex]);
+					ValidWeights.Add(BoneWeights[VertexIndex][WeightIndex]);
+				}
+			}
+			BoneIndices[VertexIndex] = ValidIndices;
+			BoneWeights[VertexIndex] = ValidWeights;
 		}
 	}
-	return false;
 }
 
 }
@@ -191,27 +276,38 @@ void FDataflowCollectionEditSkinWeightsNode::AddPrimitiveComponents(const TShare
 {
 	if(RootActor)
 	{
-		const bool bNeedsConstruction = (SkeletalMesh == nullptr) || !bValidSkeletalMesh;
-		if(!SkeletalMesh)
+		GeometryCollection::Facades::FRenderingFacade RenderingFacade(*RenderCollection);
+		const int32 NumGeometry = RenderingFacade.IsValid() ? RenderingFacade.NumGeometry() : 0;
+		
+		const bool bNeedsConstruction = (SkeletalMeshes.Num() != NumGeometry) || !bValidSkeletalMeshes;
+		
+		if(SkeletalMeshes.Num() != NumGeometry)
 		{
-			//Create transient skeletal mesh
-			FString SkeletalMeshName("SK_DataflowSkeletalMesh");
-			SkeletalMeshName = MakeUniqueObjectName(NodeOwner, USkeletalMesh::StaticClass(), *SkeletalMeshName, EUniqueObjectNameOptions::GloballyUnique).ToString();
-			SkeletalMesh = NewObject<USkeletalMesh>(NodeOwner, FName(*SkeletalMeshName), RF_Transient);
+			SkeletalMeshes.Reset();
+			for(int32 GeometryIndex = 0; GeometryIndex < NumGeometry; ++GeometryIndex)
+			{
+				FString SkeletalMeshName = FString("SK_DataflowSkeletalMesh_") + FString::FromInt(GeometryIndex);
+				SkeletalMeshName = MakeUniqueObjectName(NodeOwner, USkeletalMesh::StaticClass(), *SkeletalMeshName, EUniqueObjectNameOptions::GloballyUnique).ToString();
+				SkeletalMeshes.Add(NewObject<USkeletalMesh>(NodeOwner, FName(*SkeletalMeshName), RF_Transient));
+			}
 		}
 		if(bNeedsConstruction)
 		{
-			bValidSkeletalMesh = UE::Dataflow::Private::BuildSkeletalMesh(SkeletalMesh, RenderCollection, ObjectSkeleton);
-			if(!bValidSkeletalMesh)
+			bValidSkeletalMeshes = UE::Dataflow::Private::BuildSkeletalMeshes(SkeletalMeshes, RenderCollection, ObjectSkeleton);
+			if(!bValidSkeletalMeshes)
 			{
-				SkeletalMesh = nullptr;
+				SkeletalMeshes.Reset();
 			}
 		}
-		if(SkeletalMesh)
+		if(!SkeletalMeshes.IsEmpty())
 		{
-			USkeletalMeshComponent* SkeletalMeshComponent = NewObject<USkeletalMeshComponent>(RootActor, FName("Dataflow_SkeletalMesh"));
-			SkeletalMeshComponent->SetSkeletalMesh(SkeletalMesh);
-			PrimitiveComponents.Add(SkeletalMeshComponent);
+			for(int32 GeometryIndex = 0; GeometryIndex < NumGeometry; ++GeometryIndex)
+			{
+				FName SkeletalMeshName(FString("Dataflow_SkeletalMesh") + FString::FromInt(GeometryIndex));
+				USkeletalMeshComponent* SkeletalMeshComponent = NewObject<USkeletalMeshComponent>(RootActor, SkeletalMeshName);
+				SkeletalMeshComponent->SetSkeletalMesh(SkeletalMeshes[GeometryIndex]);
+				PrimitiveComponents.Add(SkeletalMeshComponent);
+			}
 		}
 	}
 }
@@ -234,17 +330,16 @@ void FDataflowCollectionEditSkinWeightsNode::Evaluate(UE::Dataflow::FContext& Co
 			TArray<TArray<float>> SetupWeights, FinalWeights;
 			TArray<TArray<int32>> SetupIndices, FinalIndices;
 			
-			UE::Dataflow::Private::GetAttributeValues<int32>(InCollection, BoneIndicesKeyValue, SetupIndices);
-		    UE::Dataflow::Private::GetAttributeValues<float>(InCollection, BoneWeightsKeyValue, SetupWeights);
-			
-			FinalIndices.SetNum(SetupIndices.Num());
-			FinalWeights.SetNum(SetupWeights.Num());
-			
-			ExtractVertexWeights(SetupIndices, SetupWeights, TArrayView<TArray<int32>>(FinalIndices.GetData(), FinalIndices.Num()),
-			TArrayView<TArray<float>>(FinalWeights.GetData(), FinalWeights.Num()));
+			if(GetAttributeWeights(InCollection, BoneIndicesKeyValue, BoneWeightsKeyValue, SetupIndices, SetupWeights, bCompressSkinWeights))
+			{
+				FinalIndices.SetNum(SetupIndices.Num());
+				FinalWeights.SetNum(SetupWeights.Num());
+				
+				ExtractVertexWeights(SetupIndices, SetupWeights, TArrayView<TArray<int32>>(FinalIndices.GetData(), FinalIndices.Num()),
+				TArrayView<TArray<float>>(FinalWeights.GetData(), FinalWeights.Num()));
 
-			UE::Dataflow::Private::SetAttributeValues<int32>(InCollection, BoneIndicesKeyValue, FinalIndices);
-			UE::Dataflow::Private::SetAttributeValues<float>(InCollection, BoneWeightsKeyValue, FinalWeights);
+				SetAttributeWeights(InCollection, BoneIndicesKeyValue, BoneWeightsKeyValue, FinalIndices, FinalWeights);
+			}
 		}
 
 		SetValue(Context, MoveTemp(InCollection), &Collection);
@@ -307,12 +402,38 @@ void FDataflowCollectionEditSkinWeightsNode::ExtractVertexWeights(const TArray<T
 	}
 }
 
+bool FDataflowCollectionEditSkinWeightsNode::SetAttributeWeights(FManagedArrayCollection& SelectedCollection,
+	const FCollectionAttributeKey& InBoneIndicesKey, const FCollectionAttributeKey& InBoneWeightsKey,
+	const TArray<TArray<int32>>& AttributeIndices, const TArray<TArray<float>>& AttributeWeights)
+{
+	return UE::Dataflow::Private::SetAttributeValues<int32, FIntVector4, 4>(SelectedCollection, InBoneIndicesKey, AttributeIndices,  INDEX_NONE, false) && 
+		   UE::Dataflow::Private::SetAttributeValues<float, FVector4f, 4>(SelectedCollection, InBoneWeightsKey, AttributeWeights, 0.0f, true);
+}
+
+bool FDataflowCollectionEditSkinWeightsNode::GetAttributeWeights(FManagedArrayCollection& SelectedCollection,
+	const FCollectionAttributeKey& InBoneIndicesKey, const FCollectionAttributeKey& InBoneWeightsKey,
+	TArray<TArray<int32>>& AttributeIndices, TArray<TArray<float>>& AttributeWeights, const bool bCanCompressSkinWeights)
+{
+	const bool bValidAttributes =
+		UE::Dataflow::Private::GetAttributeValues<int32, FIntVector4, 4>(SelectedCollection, InBoneIndicesKey, AttributeIndices, bCanCompressSkinWeights) && 
+		UE::Dataflow::Private::GetAttributeValues<float, FVector4f, 4>(SelectedCollection, InBoneWeightsKey, AttributeWeights, bCanCompressSkinWeights);
+
+	UE::Dataflow::Private::CorrectSkinWeights(AttributeIndices, AttributeWeights);
+
+	return bValidAttributes;
+}
+
 bool FDataflowCollectionEditSkinWeightsNode::FillAttributeWeights(const FManagedArrayCollection& SelectedCollection,
 	const FCollectionAttributeKey& InBoneIndicesKey, const FCollectionAttributeKey& InBoneWeightsKey,
-	TArray<TArray<int32>>& AttributeIndices, TArray<TArray<float>>& AttributeWeights) const
+	TArray<TArray<int32>>& AttributeIndices, TArray<TArray<float>>& AttributeWeights)
 {
-	return UE::Dataflow::Private::FillAttributeValues<int32>(SelectedCollection, InBoneIndicesKey, AttributeIndices) &&
-		   UE::Dataflow::Private::FillAttributeValues<float>(SelectedCollection, InBoneWeightsKey, AttributeWeights);
+	const bool bValidAttributes =
+		UE::Dataflow::Private::FillAttributeValues<int32, FIntVector4, 4>(SelectedCollection, InBoneIndicesKey, AttributeIndices) &&
+		UE::Dataflow::Private::FillAttributeValues<float, FVector4f, 4>(SelectedCollection, InBoneWeightsKey, AttributeWeights);
+
+	UE::Dataflow::Private::CorrectSkinWeights(AttributeIndices, AttributeWeights);
+
+	return bValidAttributes;
 }
 
 FCollectionAttributeKey FDataflowCollectionEditSkinWeightsNode::GetBoneIndicesKey(UE::Dataflow::FContext& Context) const
@@ -345,7 +466,7 @@ FCollectionAttributeKey FDataflowCollectionEditSkinWeightsNode::GetBoneWeightsKe
 
 void FDataflowCollectionEditSkinWeightsNode::OnInvalidate()
 {
-	bValidSkeletalMesh = false;
+	bValidSkeletalMeshes = false;
 }
 
 #if WITH_EDITOR
@@ -375,5 +496,23 @@ bool FDataflowCollectionEditSkinWeightsNode::CanDebugDrawViewMode(const FName& V
 
 #endif
 
+int32 FDataflowCollectionEditSkinWeightsNode::GetSkeletalMeshOffset(const TObjectPtr<USkeletalMesh>& SkeletalMesh) const
+{
+#if WITH_EDITORONLY_DATA
+	int32 SkeletalMeshOffset = 0;
+	for(int32 SkeletalMeshIndex = 0, NumMeshes = SkeletalMeshes.Num(); SkeletalMeshIndex < NumMeshes; ++SkeletalMeshIndex)
+	{
+		if (SkeletalMeshes[SkeletalMeshIndex] == SkeletalMesh)
+		{
+			return SkeletalMeshOffset;
+		}
+		else if (FMeshDescription* MeshDescription = SkeletalMeshes[SkeletalMeshIndex]->GetMeshDescription(0))
+		{
+			SkeletalMeshOffset += MeshDescription->Vertices().Num();
+		}
+	}
+#endif
+	return INDEX_NONE;
+}
 
 #undef LOCTEXT_NAMESPACE
diff --git a/Engine/Plugins/Experimental/Dataflow/Source/DataflowNodes/Public/Dataflow/DataflowCollectionEditSkinWeightsNode.h b/Engine/Plugins/Experimental/Dataflow/Source/DataflowNodes/Public/Dataflow/DataflowCollectionEditSkinWeightsNode.h
--- a/Engine/Plugins/Experimental/Dataflow/Source/DataflowNodes/Public/Dataflow/DataflowCollectionEditSkinWeightsNode.h
+++ b/Engine/Plugins/Experimental/Dataflow/Source/DataflowNodes/Public/Dataflow/DataflowCollectionEditSkinWeightsNode.h
@@ -54,17 +54,21 @@ public:
 	FScalarVertexPropertyGroup VertexGroup;
 
 	/** Bone indices key to be used in other nodes if necessary */
-	UPROPERTY(meta = (DisplayName = "Bone Indices Key", DataflowInput, DataflowOutput, DataflowPassthrough = "BoneIndicesKey"))
+	UPROPERTY(meta = (DisplayName = "Bone Indices", DataflowInput, DataflowOutput, DataflowPassthrough = "BoneIndicesKey"))
 	FCollectionAttributeKey BoneIndicesKey;
 
 	/** Bone weights key to be used in other nodes if necessary */
-	UPROPERTY(meta = (DisplayName = "Bone Weights Key", DataflowInput, DataflowOutput, DataflowPassthrough = "BoneWeightsKey"))
+	UPROPERTY(meta = (DisplayName = "Bone Weights", DataflowInput, DataflowOutput, DataflowPassthrough = "BoneWeightsKey"))
 	FCollectionAttributeKey BoneWeightsKey;
 
 	/** Skeletal mesh to extract the skeleton from for the skinning */
 	UPROPERTY(EditAnywhere, Category = "Skeleton Binding")
 	TObjectPtr<USkeleton> ObjectSkeleton = nullptr;
 
+	/** Boolean to use a compressed format (FVector4f, FIntVector) to store the skin weights */
+	UPROPERTY(EditAnywhere, Category = "Skeleton Binding")
+	bool bCompressSkinWeights = false;
+
 	/** List of skin weights */
 	UPROPERTY()
 	TArray<FDataflowSkinWeightData> SkinWeights;
@@ -79,8 +83,18 @@ public:
 	UE_API void ExtractVertexWeights(const TArray<TArray<int32>>& SetupIndices, const TArray<TArray<float>>& SetupWeights, TArrayView<TArray<int32>> FinalIndices, TArrayView<TArray<float>> FinalWeights) const;
 
 	/** Fill the weight attribute values from the collection */
-	UE_API bool FillAttributeWeights(const FManagedArrayCollection& SelectedCollection,
-		const FCollectionAttributeKey& IndicesAttributeKey, const FCollectionAttributeKey& WeightsAttributeKey, TArray<TArray<int32>>& AttributeIndices, TArray<TArray<float>>& AttributeWeights) const;
+	UE_API static bool FillAttributeWeights(const FManagedArrayCollection& SelectedCollection,
+		const FCollectionAttributeKey& IndicesAttributeKey, const FCollectionAttributeKey& WeightsAttributeKey, TArray<TArray<int32>>& AttributeIndices, TArray<TArray<float>>& AttributeWeights);
+
+	/** Get the weight attribute values and add them if not there*/
+	UE_API static bool GetAttributeWeights(FManagedArrayCollection& SelectedCollection,
+			const FCollectionAttributeKey& InBoneIndicesKey, const FCollectionAttributeKey& InBoneWeightsKey,
+			TArray<TArray<int32>>& AttributeIndices, TArray<TArray<float>>& AttributeWeights, const bool bCanCompressSkinWeights);
+
+	/** Set back the attribute values into the collection */
+	UE_API static bool SetAttributeWeights(FManagedArrayCollection& SelectedCollection,
+			const FCollectionAttributeKey& InBoneIndicesKey, const FCollectionAttributeKey& InBoneWeightsKey,
+			const TArray<TArray<int32>>& AttributeIndices, const TArray<TArray<float>>& AttributeWeights);
 
 	/** Get the weights attribute key to retrieve/set the weight values*/
 	UE_API FCollectionAttributeKey GetBoneIndicesKey(UE::Dataflow::FContext& Context) const;
@@ -89,7 +103,11 @@ public:
     UE_API FCollectionAttributeKey GetBoneWeightsKey(UE::Dataflow::FContext& Context) const;
 
 	/** Validate the skeletal mesh construction */
-	void ValidateSkeletalMesh() {bValidSkeletalMesh = true;}
+	UE_API void ValidateSkeletalMeshes() {bValidSkeletalMeshes = true;}
+
+	/** Return the collection offset given a skeletal mesh*/
+	UE_API int32 GetSkeletalMeshOffset(const TObjectPtr<USkeletalMesh>& SkeletalMesh) const;
+	
 private:
 
 	//~ Begin FDataflowPrimitiveNode interface
@@ -109,10 +127,10 @@ private:
 
 	/** Transient skeletal mesh built from dataflow render collection */
 	UPROPERTY(Transient)
-	TObjectPtr<USkeletalMesh> SkeletalMesh = nullptr;
+	TArray<TObjectPtr<USkeletalMesh>> SkeletalMeshes;
 
 	/** Valid skeletal mesh boolean to trigger the construction */
-	bool bValidSkeletalMesh = false;
+	bool bValidSkeletalMeshes = false;
 };
 
 #undef UE_API
diff --git a/Engine/Source/Runtime/Experimental/Dataflow/Core/Public/Dataflow/DataflowNodeParameters.h b/Engine/Source/Runtime/Experimental/Dataflow/Core/Public/Dataflow/DataflowNodeParameters.h
--- a/Engine/Source/Runtime/Experimental/Dataflow/Core/Public/Dataflow/DataflowNodeParameters.h
+++ b/Engine/Source/Runtime/Experimental/Dataflow/Core/Public/Dataflow/DataflowNodeParameters.h
@@ -24,16 +24,17 @@ namespace UE::Dataflow
 	struct FRenderingParameter 
 	{
 		FRenderingParameter() {}
-		FRenderingParameter(FString InRenderName, FName InTypeName, const TArray<FName>& InOutputs)
-			: Name(InRenderName), Type(InTypeName), Outputs(InOutputs) {}
-		FRenderingParameter(FString InRenderName, FName InTypeName, TArray<FName>&& InOutputs)
-			: Name(InRenderName), Type(InTypeName), Outputs(InOutputs) {}
+		FRenderingParameter(FString InRenderName, FName InTypeName, const TArray<FName>& InOutputs, const FName& InViewMode = FName("3DView"))
+			: Name(InRenderName), Type(InTypeName), Outputs(InOutputs), ViewMode(InViewMode) {}
+		FRenderingParameter(FString InRenderName, FName InTypeName, TArray<FName>&& InOutputs, const FName& InViewMode = FName("3DView"))
+			: Name(InRenderName), Type(InTypeName), Outputs(InOutputs), ViewMode(InViewMode) {}
 		
 		bool operator==(const FRenderingParameter& Other) const = default;
 
 		FString Name = FString("");
 		FName Type = FName("");
 		TArray<FName> Outputs;
+		FName ViewMode = FName("");
 	};
 
 	//////////////////////////////////////////////////////////////////////////////////
