Support GetPhysicsAssetFromSkeletalMesh node.

diff --git a/Engine/Plugins/ChaosClothAssetEditor/Source/ChaosClothAssetDataflowNodes/Private/ChaosClothAsset/SetPhysicsAssetNode.cpp b/Engine/Plugins/ChaosClothAssetEditor/Source/ChaosClothAssetDataflowNodes/Private/ChaosClothAsset/SetPhysicsAssetNode.cpp
--- a/Engine/Plugins/ChaosClothAssetEditor/Source/ChaosClothAssetDataflowNodes/Private/ChaosClothAsset/SetPhysicsAssetNode.cpp
+++ b/Engine/Plugins/ChaosClothAssetEditor/Source/ChaosClothAssetDataflowNodes/Private/ChaosClothAsset/SetPhysicsAssetNode.cpp
@@ -4,6 +4,7 @@
 #include "ChaosClothAsset/ClothGeometryTools.h"
 #include "ChaosClothAsset/CollectionClothFacade.h"
 #include "Dataflow/DataflowInputOutput.h"
+#include "PhysicsEngine/PhysicsAsset.h"
 
 #include UE_INLINE_GENERATED_CPP_BY_NAME(SetPhysicsAssetNode)
 
@@ -13,6 +14,7 @@ FChaosClothAssetSetPhysicsAssetNode::FChaosClothAssetSetPhysicsAssetNode(const U
 	: FDataflowNode(InParam, InGuid)
 {
 	RegisterInputConnection(&Collection);
+	RegisterInputConnection(&PhysicsAsset);
 	RegisterOutputConnection(&Collection, &Collection);
 }
 
@@ -29,7 +31,8 @@ void FChaosClothAssetSetPhysicsAssetNode::Evaluate(UE::Dataflow::FContext& Conte
 		if (FCollectionClothFacade(ClothCollection).IsValid())  // Can only act on the collection if it is a valid cloth collection
 		{
 			FCollectionClothFacade CollectionClothFacade(ClothCollection);
-			CollectionClothFacade.SetPhysicsAssetPathName(PhysicsAsset ? PhysicsAsset.GetPathName() : FString());
+			TObjectPtr<const UPhysicsAsset> InPhysicsAsset = GetValue(Context, &PhysicsAsset);
+			CollectionClothFacade.SetPhysicsAssetPathName(InPhysicsAsset ? InPhysicsAsset.GetPathName() : FString());
 		}
 
 		SetValue(Context, MoveTemp(*ClothCollection), &Collection);
diff --git a/Engine/Plugins/ChaosClothAssetEditor/Source/ChaosClothAssetDataflowNodes/Public/ChaosClothAsset/SetPhysicsAssetNode.h b/Engine/Plugins/ChaosClothAssetEditor/Source/ChaosClothAssetDataflowNodes/Public/ChaosClothAsset/SetPhysicsAssetNode.h
--- a/Engine/Plugins/ChaosClothAssetEditor/Source/ChaosClothAssetDataflowNodes/Public/ChaosClothAsset/SetPhysicsAssetNode.h
+++ b/Engine/Plugins/ChaosClothAssetEditor/Source/ChaosClothAssetDataflowNodes/Public/ChaosClothAsset/SetPhysicsAssetNode.h
@@ -27,7 +27,7 @@ public:
 	FManagedArrayCollection Collection;
 
 	/** The physics asset to assign to the Cloth Collection. */
-	UPROPERTY(EditAnywhere, Category = "Set Physics Asset")
+	UPROPERTY(EditAnywhere, Category = "Set Physics Asset", Meta = (DataflowInput))
 	TObjectPtr<UPhysicsAsset> PhysicsAsset;
 
 	FChaosClothAssetSetPhysicsAssetNode(const UE::Dataflow::FNodeParameters& InParam, FGuid InGuid = FGuid::NewGuid());
diff --git a/Engine/Plugins/Experimental/Dataflow/Source/DataflowNodes/Private/Dataflow/DataflowSkeletalMeshNodes.cpp b/Engine/Plugins/Experimental/Dataflow/Source/DataflowNodes/Private/Dataflow/DataflowSkeletalMeshNodes.cpp
--- a/Engine/Plugins/Experimental/Dataflow/Source/DataflowNodes/Private/Dataflow/DataflowSkeletalMeshNodes.cpp
+++ b/Engine/Plugins/Experimental/Dataflow/Source/DataflowNodes/Private/Dataflow/DataflowSkeletalMeshNodes.cpp
@@ -12,6 +12,7 @@
 #include "SkeletalDebugRendering.h"
 #include "Dataflow/DataflowDebugDrawInterface.h"
 #include "Dataflow/DataflowElement.h"
+#include "PhysicsEngine/PhysicsAsset.h"
 
 //#if WITH_EDITOR
 //#include "Dataflow/DataflowRenderingViewMode.h"
@@ -27,7 +28,8 @@ namespace UE::Dataflow
 		DATAFLOW_NODE_REGISTER_CREATION_FACTORY(FGetSkeletonDataflowNode);
 		DATAFLOW_NODE_REGISTER_CREATION_FACTORY(FSkeletalMeshBoneDataflowNode);
 		DATAFLOW_NODE_REGISTER_CREATION_FACTORY(FSkeletalMeshReferenceTransformDataflowNode);
-
+		DATAFLOW_NODE_REGISTER_CREATION_FACTORY(FGetPhysicsAssetFromSkeletalMeshDataflowNode);
+		
 		DATAFLOW_NODE_REGISTER_GETTER_FOR_ASSET(USkeletalMesh, FGetSkeletalMeshDataflowNode);
 		DATAFLOW_NODE_REGISTER_GETTER_FOR_ASSET(USkeleton, FGetSkeletonDataflowNode);
 	}
@@ -322,7 +324,12 @@ void FSkeletalMeshReferenceTransformDataflowNode::Evaluate(UE::Dataflow::FContex
 	}
 }
 
-
-
-
-
+void FGetPhysicsAssetFromSkeletalMeshDataflowNode::Evaluate(UE::Dataflow::FContext& Context, const FDataflowOutput* Out) const
+{
+	if (Out->IsA(&PhysicsAsset))
+	{
+		TObjectPtr<const USkeletalMesh> InSkeletalMesh = GetValue(Context, &SkeletalMesh);
+		UPhysicsAsset* OutPhysicsAsset = InSkeletalMesh ? InSkeletalMesh->GetPhysicsAsset() : nullptr;
+		SetValue(Context, TObjectPtr<const UPhysicsAsset>(OutPhysicsAsset), &PhysicsAsset);
+	}
+}
\ No newline at end of file
diff --git a/Engine/Plugins/Experimental/Dataflow/Source/DataflowNodes/Public/Dataflow/DataflowSkeletalMeshNodes.h b/Engine/Plugins/Experimental/Dataflow/Source/DataflowNodes/Public/Dataflow/DataflowSkeletalMeshNodes.h
--- a/Engine/Plugins/Experimental/Dataflow/Source/DataflowNodes/Public/Dataflow/DataflowSkeletalMeshNodes.h
+++ b/Engine/Plugins/Experimental/Dataflow/Source/DataflowNodes/Public/Dataflow/DataflowSkeletalMeshNodes.h
@@ -16,6 +16,7 @@
 DEFINE_LOG_CATEGORY_STATIC(LogDataflowSkeletalMeshNodes, Log, All);
 
 class USkeletalMesh;
+class UPhysicsAsset;
 
 // A delegate for monitoring dataflow skeleton selection changes.
 DECLARE_MULTICAST_DELEGATE_OneParam(FDataflowBoneSelectionChangedNotifyDelegate, const TArray<FName>& /*InBoneNames*/);
@@ -178,6 +179,33 @@ public:
 	UE_API virtual void Evaluate(UE::Dataflow::FContext& Context, const FDataflowOutput* Out) const override;
 };
 
+/** Get the physics assets from input skeletal mesh. */
+USTRUCT()
+struct FGetPhysicsAssetFromSkeletalMeshDataflowNode : public FDataflowNode
+{
+	GENERATED_USTRUCT_BODY()
+	DATAFLOW_NODE_DEFINE_INTERNAL(FGetPhysicsAssetFromSkeletalMeshDataflowNode, "GetPhysicsAssetFromSkeletalMesh", "General", "Get Physics Asset Skeletal Mesh")
+
+public:
+	/** Input skeletal mesh */
+	UPROPERTY(EditAnywhere, Category = "Skeletal Mesh", Meta = (DataflowInput))
+	TObjectPtr<const USkeletalMesh> SkeletalMesh;
+
+	/** Output physics asset */
+	UPROPERTY(Meta = (DataflowOutput))
+	TObjectPtr<const UPhysicsAsset> PhysicsAsset;
+
+	FGetPhysicsAssetFromSkeletalMeshDataflowNode(const UE::Dataflow::FNodeParameters& InParam, FGuid InGuid = FGuid::NewGuid())
+		: FDataflowNode(InParam, InGuid)
+	{
+		RegisterInputConnection(&SkeletalMesh);
+		RegisterOutputConnection(&PhysicsAsset);
+	}
+
+private:
+	virtual void Evaluate(UE::Dataflow::FContext& Context, const FDataflowOutput* Out) const override;
+};
+
 namespace UE::Dataflow
 {
 	void RegisterSkeletalMeshNodes();
