Fix skin weights correction for duplicated vertices at the seam

diff --git a/Engine/Plugins/Experimental/Dataflow/Source/DataflowEditor/Private/DataflowEditorTools/DataflowEditorCorrectSkinWeightsNode.cpp b/Engine/Plugins/Experimental/Dataflow/Source/DataflowEditor/Private/DataflowEditorTools/DataflowEditorCorrectSkinWeightsNode.cpp
--- a/Engine/Plugins/Experimental/Dataflow/Source/DataflowEditor/Private/DataflowEditorTools/DataflowEditorCorrectSkinWeightsNode.cpp
+++ b/Engine/Plugins/Experimental/Dataflow/Source/DataflowEditor/Private/DataflowEditorTools/DataflowEditorCorrectSkinWeightsNode.cpp
@@ -53,6 +53,119 @@ static void BuildDynamicMeshes(const FDataflowNode* DataflowNode, UE::Dataflow::
 	}
 }
 
+static void AccumulateSkinWeights(const float InterpolationWeight, 
+		const TArray<int32>& BoneIndices, const TArray<float>& BoneWeights, TMap<int32,float>& SkinWeights, float& TotalWeight)
+{
+	for (int32 WeightIndex = 0, NumWeights = BoneIndices.Num(); WeightIndex < NumWeights; ++WeightIndex)
+	{
+		const float InterpolatedWeight = InterpolationWeight * BoneWeights[WeightIndex];
+
+		if (float* WeightValue = SkinWeights.Find(BoneIndices[WeightIndex]))
+		{
+			*WeightValue += InterpolatedWeight;
+		}
+		else
+		{
+			SkinWeights.Add(BoneIndices[WeightIndex], InterpolatedWeight);
+		}
+		TotalWeight += InterpolatedWeight;
+	}
+}
+
+static void ReportSkinWeights(TArray<int32>& BoneIndices, TArray<float>& BoneWeights, const TMap<int32, float>& SkinWeights, const float TotalWeight)
+{
+	if (TotalWeight > 0.0)
+	{
+		BoneIndices.Reset();
+		BoneWeights.Reset();
+
+		for (const TPair<int32, float>& BoneWeight : SkinWeights)
+		{
+			BoneIndices.Add(BoneWeight.Key);
+			BoneWeights.Add(BoneWeight.Value / TotalWeight);
+		}
+	}
+}
+
+// Merge the skin weights of border vertices if collocated
+static void MergeSkinWeights(Geometry::FDynamicMesh3& DynamicMesh, const int32 VertexOffset, const TArray<float>& SelectionMap, TArray<TArray<int32>>& FinalIndices, TArray<TArray<float>>& FinalWeights)
+{
+	const int32 NumVertices = DynamicMesh.VertexCount();
+	if(SelectionMap.Num() <= NumVertices + VertexOffset)
+	{ 
+		TArray<TPair<int32, FVector3d>> BorderVertices;
+		BorderVertices.Reserve(NumVertices);
+
+		// Gather all the border vertices
+		for (int32 VertexIndex = 0; VertexIndex < NumVertices; ++VertexIndex)
+		{
+			if (DynamicMesh.IsBoundaryVertex(VertexIndex))
+			{
+				BorderVertices.Add({VertexIndex, DynamicMesh.GetVertex(VertexIndex) });
+			}
+		}
+		// Sort them based on the position 
+		Algo::Sort(BorderVertices, [](const TPair<int32, FVector3d>& A, const TPair<int32, FVector3d>& B)
+		{
+			return (A.Value.X != B.Value.X) ? A.Value.X < B.Value.X : (A.Value.Y != B.Value.Y) ? A.Value.Y < B.Value.Y : A.Value.Z < B.Value.Z;
+		});
+		const int32 NumBorders = BorderVertices.Num();
+
+		if(NumBorders > 0)
+		{ 
+			TArray<TArray<int32>> CollocatedVertices;
+			CollocatedVertices.Reserve(NumBorders);
+
+			// Identify collocated vertices consecutively 
+			CollocatedVertices.AddDefaulted_GetRef().Add(BorderVertices[0].Key);
+			for (int32 BorderIndex = 1; BorderIndex < NumBorders; ++BorderIndex)
+			{
+				if((BorderVertices[BorderIndex].Value - BorderVertices[BorderIndex -1].Value).SquaredLength() < UE_SMALL_NUMBER)
+				{
+					CollocatedVertices.Last().Add(BorderVertices[BorderIndex].Key);
+				}
+				else
+				{
+					CollocatedVertices.AddDefaulted_GetRef().Add(BorderVertices[BorderIndex].Key);
+				}
+			}
+
+			// Loop over the selected collocated vertices (at least one of the border vertices is selected)
+			for (int32 CollocatedIndex = 0; CollocatedIndex < CollocatedVertices.Num(); ++CollocatedIndex)
+			{
+				bool bIsSelected = false;
+				for (const int32 VertexIndex : CollocatedVertices[CollocatedIndex])
+				{
+					if(SelectionMap[VertexIndex + VertexOffset] > 0.0f)
+					{
+						bIsSelected = true;
+					}
+				}
+				if(bIsSelected)
+				{ 
+					TMap<int32, float> SkinWeights;
+					float TotalWeight = 0.0;
+
+					// Accumulate weights from all the border vertices 
+					for (const int32 VertexIndex : CollocatedVertices[CollocatedIndex])
+					{
+						AccumulateSkinWeights(1.0f, FinalIndices[VertexIndex + VertexOffset], FinalWeights[VertexIndex + VertexOffset], SkinWeights, TotalWeight);
+					}
+
+					// Report the accumulated skin weight onto all the collocated vertices
+					if (TotalWeight > 0.0)
+					{
+						for (const int32 VertexIndex : CollocatedVertices[CollocatedIndex])
+						{
+							ReportSkinWeights(FinalIndices[VertexIndex + VertexOffset], FinalWeights[VertexIndex + VertexOffset], SkinWeights, TotalWeight);
+						}
+					}
+				}
+			}
+		}
+	}
+}
+
 static void PruneSkinWeights(const float PruningThreshold, const TArray<TArray<int32>>& SetupIndices, const TArray<TArray<float>>& SetupWeights,
 	const TArray<float>& SelectionMap, TArray<TArray<int32>>& FinalIndices, TArray<TArray<float>>& FinalWeights)
 {
@@ -212,7 +325,7 @@ static void RelaxSkinWeights(const UE::Geometry::FDynamicMesh3& DynamicMesh, con
 	}
 }
 
-static void HammerSkinWeights(const UE::Geometry::FDynamicMesh3& DynamicMesh, const int32 VertexOffset, const TArray<TArray<int32>>& SetupIndices, const TArray<TArray<float>>& SetupWeights,
+static void HammerSkinWeights(const UE::Geometry::FDynamicMesh3& DynamicMesh, const int32 VertexOffset, const float SelectionThreshold, const TArray<TArray<int32>>& SetupIndices, const TArray<TArray<float>>& SetupWeights,
 	const TArray<float>& SelectionMap, TArray<TArray<int32>>& FinalIndices, TArray<TArray<float>>& FinalWeights)
 {
 	const int32 NumVertices = DynamicMesh.VertexCount();
@@ -223,12 +336,12 @@ static void HammerSkinWeights(const UE::Geometry::FDynamicMesh3& DynamicMesh, co
 	for (int32 VertexIndex = 0; VertexIndex < NumVertices; ++VertexIndex)
 	{
 		const int32 GlobalVertex = VertexOffset+VertexIndex;
-		if(SelectionMap[GlobalVertex] > 0.0f)
+		if(SelectionMap[GlobalVertex] > SelectionThreshold)
 		{
 			SelectedVertices.Add(VertexIndex);
 			for (const int32 NeighborIndex : DynamicMesh.VtxVerticesItr(VertexIndex))
 			{
-				if (SelectionMap[VertexOffset+NeighborIndex] == 0.0f)
+				if (SelectionMap[VertexOffset+NeighborIndex] <= SelectionThreshold)
 				{
 					NeighborVertices.Add(NeighborIndex);
 				}
@@ -253,16 +366,27 @@ static void HammerSkinWeights(const UE::Geometry::FDynamicMesh3& DynamicMesh, co
 	for (const int32 SelectedVertex : SelectedVertices)
 	{
 		// find the closest surrounding vertex to this selected vertex
-		if (!PathFinder.FindPathToNearestSeed(SelectedVertex, VertexPath))
-		{
-			continue;
-		}
-		const int32 ClosestVertex = VertexPath.Last();
+		const int32 ClosestVertex = PathFinder.FindPathToNearestSeed(SelectedVertex, VertexPath) ? VertexPath.Last() : SelectedVertex;
 		const int32 GlobalVertex = VertexOffset+SelectedVertex;
 		const int32 GlobalClosest = VertexOffset+ClosestVertex;
-		
-		FinalIndices[GlobalVertex] = SetupIndices[GlobalClosest];
-		FinalWeights[GlobalVertex] = SetupWeights[GlobalClosest];
+
+		TMap<int32, float> SkinWeights;
+		float TotalWeight = 0.0;
+
+		const float InterpolationWeight = FMath::Clamp((SelectionMap[GlobalVertex]-SelectionThreshold) / (1.0-SelectionThreshold), 0.0f, 1.0f);
+
+		AccumulateSkinWeights(InterpolationWeight, SetupIndices[GlobalClosest], SetupWeights[GlobalClosest], SkinWeights, TotalWeight);
+		AccumulateSkinWeights(1.0f-InterpolationWeight, SetupIndices[GlobalVertex], SetupWeights[GlobalVertex], SkinWeights, TotalWeight);
+
+		if(TotalWeight != 0.0f)
+		{
+			ReportSkinWeights(FinalIndices[GlobalVertex], FinalWeights[GlobalVertex], SkinWeights, TotalWeight);
+		}
+		else
+		{
+			FinalIndices[GlobalVertex] = SetupIndices[GlobalVertex];
+			FinalWeights[GlobalVertex] = SetupWeights[GlobalVertex];
+		}
 	}
 }
 	
@@ -393,9 +517,12 @@ void FDataflowCorrectSkinWeightsNode::Evaluate(UE::Dataflow::FContext& Context,
 							}
 							else if (CorrectionType == ESkinWeightsCorrectionType::Hammer)
 							{
-								UE::Dataflow::Private::HammerSkinWeights(DynamicMesh, VertexOffset,
+								UE::Dataflow::Private::HammerSkinWeights(DynamicMesh, VertexOffset, SelectionThreshold,
 									SetupIndices, SetupWeights, SelectionMap, FinalIndices, FinalWeights);
 							}
+
+							UE::Dataflow::Private::MergeSkinWeights(DynamicMesh, VertexOffset, SelectionMap, FinalIndices, FinalWeights);
+
 							VertexOffset += DynamicMesh.VertexCount();
 						}
 					}
diff --git a/Engine/Plugins/Experimental/Dataflow/Source/DataflowEditor/Public/DataflowEditorTools/DataflowEditorCorrectSkinWeightsNode.h b/Engine/Plugins/Experimental/Dataflow/Source/DataflowEditor/Public/DataflowEditorTools/DataflowEditorCorrectSkinWeightsNode.h
--- a/Engine/Plugins/Experimental/Dataflow/Source/DataflowEditor/Public/DataflowEditorTools/DataflowEditorCorrectSkinWeightsNode.h
+++ b/Engine/Plugins/Experimental/Dataflow/Source/DataflowEditor/Public/DataflowEditorTools/DataflowEditorCorrectSkinWeightsNode.h
@@ -88,6 +88,10 @@ public:
 	/** Max number of bones to consider for the skin weights */
     UPROPERTY(EditAnywhere, Category = "Skin Weights", meta = (DisplayName = "Clamping Number", EditCondition = "CorrectionType == ESkinWeightsCorrectionType::Clamp", EditConditionHides))
     int32 ClampingNumber = 8;
+
+	/** Selection threshold to consider a neighbor skin weight */
+    UPROPERTY(EditAnywhere, Category = "Skin Weights", meta = (DisplayName = "Selection Threshold", EditCondition = "CorrectionType == ESkinWeightsCorrectionType::Hammer", EditConditionHides, ClampMin = "0.0", ClampMax = "1.0", UIMin = "0.0", UIMax = "1.0"))
+    float SelectionThreshold = 0.5f;
 	
 	/** Get the weights attribute key to retrieve/set the bone indices values*/
 	FCollectionAttributeKey GetBoneIndicesKey(UE::Dataflow::FContext& Context) const;
