Fix skin weight relaxation along the seam + limit the number of bone influences

diff --git a/Engine/Plugins/Experimental/Dataflow/Source/DataflowEditor/Private/DataflowEditorTools/DataflowEditorCorrectSkinWeightsNode.cpp b/Engine/Plugins/Experimental/Dataflow/Source/DataflowEditor/Private/DataflowEditorTools/DataflowEditorCorrectSkinWeightsNode.cpp
--- a/Engine/Plugins/Experimental/Dataflow/Source/DataflowEditor/Private/DataflowEditorTools/DataflowEditorCorrectSkinWeightsNode.cpp
+++ b/Engine/Plugins/Experimental/Dataflow/Source/DataflowEditor/Private/DataflowEditorTools/DataflowEditorCorrectSkinWeightsNode.cpp
@@ -87,78 +87,173 @@ static void ReportSkinWeights(TArray<int32>& BoneIndices, TArray<float>& BoneWei
 	}
 }
 
-// Merge the skin weights of border vertices if collocated
-static void MergeSkinWeights(Geometry::FDynamicMesh3& DynamicMesh, const int32 VertexOffset, const TArray<float>& SelectionMap, TArray<TArray<int32>>& FinalIndices, TArray<TArray<float>>& FinalWeights)
+// Small structure to store the averaged bone weight over some vertices
+struct FAveragedBoneWeight
 {
-	const int32 NumVertices = DynamicMesh.VertexCount();
-	if(SelectionMap.Num() <= NumVertices + VertexOffset)
+	FAveragedBoneWeight(const float InWeightValue, const int32 InNumVertices) : 
+		WeightValue(InWeightValue), NumVertices(InNumVertices) 
+	{}
+
+	float WeightValue;
+	int32 NumVertices;
+};
+
+static bool AverageSkinWeights(const TSet<int32>& NeighborVertices, const int32 VertexOffset, const TArray<TArray<int32>>& FinalIndices, const TArray<TArray<float>>& FinalWeights, TMap<int32, FAveragedBoneWeight>& NeighborWeights)
+{
+	for (const int32 NeighborVertex : NeighborVertices)
+	{
+		const int32 GlobalNeighbor = VertexOffset + NeighborVertex;
+		check(FinalIndices[GlobalNeighbor].Num() == FinalWeights[GlobalNeighbor].Num());
+		for (int32 WeightIndex = 0, NumWeights = FinalIndices[GlobalNeighbor].Num(); WeightIndex < NumWeights; ++WeightIndex)
+		{
+			if (FinalWeights[GlobalNeighbor][WeightIndex] > 0.0f && FinalWeights[GlobalNeighbor][WeightIndex] <= 1.0f)
+			{
+				if (FAveragedBoneWeight* WeightValue = NeighborWeights.Find(FinalIndices[GlobalNeighbor][WeightIndex]))
+				{
+					WeightValue->WeightValue += FinalWeights[GlobalNeighbor][WeightIndex];
+					WeightValue->NumVertices += 1;
+				}
+				else
+				{
+					NeighborWeights.Add(FinalIndices[GlobalNeighbor][WeightIndex], {FinalWeights[GlobalNeighbor][WeightIndex], 1});
+				}
+			}
+		}
+	}
+
+	float TotalWeight = 0.0;
+	for (TPair<int32, FAveragedBoneWeight>& NeighborWeight : NeighborWeights)
+	{
+		// Any averaged weight added has a minimum of 1 vertex 
+		NeighborWeight.Value.WeightValue /= NeighborWeight.Value.NumVertices;
+		TotalWeight += NeighborWeight.Value.WeightValue;
+	}
+	if(!FMath::IsNearlyZero(TotalWeight))
 	{ 
-		TArray<TPair<int32, FVector3d>> BorderVertices;
-		BorderVertices.Reserve(NumVertices);
+		for (TPair<int32, FAveragedBoneWeight>& NeighborWeight : NeighborWeights)
+		{
+			NeighborWeight.Value.WeightValue /= TotalWeight;
+		}
+		return true;
+	}
+	return false;
+}
 
-		// Gather all the border vertices
-		for (int32 VertexIndex = 0; VertexIndex < NumVertices; ++VertexIndex)
+static void RestrictSkinWeights(const int32 ClampingNumber, TArray<int32>& FinalIndices, TArray<float>& FinalWeights)
+{
+	if(FinalIndices.Num() > ClampingNumber)
+	{ 
+		TArray<TPair<int32, float>> SortedWeights;
+		SortedWeights.Reserve(FinalWeights.Num());
+		for (int32 WeightIndex = 0, NumWeights = FinalWeights.Num(); WeightIndex < NumWeights; ++WeightIndex)
 		{
-			if (DynamicMesh.IsBoundaryVertex(VertexIndex))
+			SortedWeights.Add({ FinalIndices[WeightIndex], FinalWeights[WeightIndex] });
+		}
+		// sort in descending order by weight
+		Algo::Sort(SortedWeights, [](const TPair<int32, float>& A, const TPair<int32, float>& B)
 			{
-				BorderVertices.Add({VertexIndex, DynamicMesh.GetVertex(VertexIndex) });
+				return A.Value > B.Value;
+			});
+
+		FinalIndices.Reset();
+		FinalWeights.Reset();
+
+		float TotalWeight = 0.0f;
+		for (int32 WeightIndex = 0; WeightIndex < ClampingNumber; ++WeightIndex)
+		{
+			FinalIndices.Add(SortedWeights[WeightIndex].Key);
+			FinalWeights.Add(SortedWeights[WeightIndex].Value);
+			TotalWeight += SortedWeights[WeightIndex].Value;
+		}
+		if (TotalWeight != 0.0f)
+		{
+			for (int32 WeightIndex = 0; WeightIndex < ClampingNumber; ++WeightIndex)
+			{
+				FinalWeights[WeightIndex] /= TotalWeight;
 			}
 		}
-		// Sort them based on the position 
-		Algo::Sort(BorderVertices, [](const TPair<int32, FVector3d>& A, const TPair<int32, FVector3d>& B)
+	}
+}
+	
+// Build all the collocated vertices
+static void BuildCollocatedVertices(const Geometry::FDynamicMesh3& DynamicMesh, TArray<TArray<int32>>& CollocatedVertices)
+{
+	CollocatedVertices.Reset();
+	
+	const int32 NumVertices = DynamicMesh.VertexCount();
+	
+	TArray<TPair<int32, FVector3d>> BorderVertices;
+	BorderVertices.Reserve(NumVertices);
+
+	// Gather all the border vertices
+	for (int32 VertexIndex = 0; VertexIndex < NumVertices; ++VertexIndex)
+	{
+		if (DynamicMesh.IsBoundaryVertex(VertexIndex))
 		{
-			return (A.Value.X != B.Value.X) ? A.Value.X < B.Value.X : (A.Value.Y != B.Value.Y) ? A.Value.Y < B.Value.Y : A.Value.Z < B.Value.Z;
-		});
-		const int32 NumBorders = BorderVertices.Num();
+			BorderVertices.Add({VertexIndex, DynamicMesh.GetVertex(VertexIndex) });
+		}
+	}
+	// Sort them based on the position 
+	Algo::Sort(BorderVertices, [](const TPair<int32, FVector3d>& A, const TPair<int32, FVector3d>& B)
+	{
+		return (A.Value.X != B.Value.X) ? A.Value.X < B.Value.X : (A.Value.Y != B.Value.Y) ? A.Value.Y < B.Value.Y : A.Value.Z < B.Value.Z;
+	});
+	const int32 NumBorders = BorderVertices.Num();
 
-		if(NumBorders > 0)
-		{ 
-			TArray<TArray<int32>> CollocatedVertices;
-			CollocatedVertices.Reserve(NumBorders);
+	if(NumBorders > 0)
+	{
+		CollocatedVertices.Reserve(NumBorders);
+		
+		// Identify collocated vertices consecutively 
+		CollocatedVertices.AddDefaulted_GetRef().Add(BorderVertices[0].Key);
+		for (int32 BorderIndex = 1; BorderIndex < NumBorders; ++BorderIndex)
+		{
+			if((BorderVertices[BorderIndex].Value - BorderVertices[BorderIndex -1].Value).SquaredLength() < UE_SMALL_NUMBER)
+			{
+				CollocatedVertices.Last().Add(BorderVertices[BorderIndex].Key);
+			}
+			else
+			{
+				CollocatedVertices.AddDefaulted_GetRef().Add(BorderVertices[BorderIndex].Key);
+			}
+		}
+	}
+}
 
-			// Identify collocated vertices consecutively 
-			CollocatedVertices.AddDefaulted_GetRef().Add(BorderVertices[0].Key);
-			for (int32 BorderIndex = 1; BorderIndex < NumBorders; ++BorderIndex)
+// Merge the skin weights of border vertices if collocated
+static void MergeSkinWeights(const Geometry::FDynamicMesh3& DynamicMesh, const TArray<TArray<int32>>& CollocatedVertices, const int32 VertexOffset, const TArray<float>& SelectionMap, TArray<TArray<int32>>& FinalIndices, TArray<TArray<float>>& FinalWeights)
+{
+	const int32 NumVertices = DynamicMesh.VertexCount();
+	if(SelectionMap.Num() <= NumVertices + VertexOffset)
+	{ 
+		// Loop over the selected collocated vertices (at least one of the border vertices is selected)
+		for (int32 CollocatedIndex = 0; CollocatedIndex < CollocatedVertices.Num(); ++CollocatedIndex)
+		{
+			bool bIsSelected = false;
+			for (const int32 VertexIndex : CollocatedVertices[CollocatedIndex])
 			{
-				if((BorderVertices[BorderIndex].Value - BorderVertices[BorderIndex -1].Value).SquaredLength() < UE_SMALL_NUMBER)
-				{
-					CollocatedVertices.Last().Add(BorderVertices[BorderIndex].Key);
-				}
-				else
+				if(SelectionMap[VertexIndex + VertexOffset] > 0.0f)
 				{
-					CollocatedVertices.AddDefaulted_GetRef().Add(BorderVertices[BorderIndex].Key);
+					bIsSelected = true;
 				}
 			}
+			if(bIsSelected)
+			{ 
+				TMap<int32, float> SkinWeights;
+				float TotalWeight = 0.0;
 
-			// Loop over the selected collocated vertices (at least one of the border vertices is selected)
-			for (int32 CollocatedIndex = 0; CollocatedIndex < CollocatedVertices.Num(); ++CollocatedIndex)
-			{
-				bool bIsSelected = false;
+				// Accumulate weights from all the border vertices 
 				for (const int32 VertexIndex : CollocatedVertices[CollocatedIndex])
 				{
-					if(SelectionMap[VertexIndex + VertexOffset] > 0.0f)
-					{
-						bIsSelected = true;
-					}
+					AccumulateSkinWeights(1.0f, FinalIndices[VertexIndex + VertexOffset], FinalWeights[VertexIndex + VertexOffset], SkinWeights, TotalWeight);
 				}
-				if(bIsSelected)
-				{ 
-					TMap<int32, float> SkinWeights;
-					float TotalWeight = 0.0;
 
-					// Accumulate weights from all the border vertices 
+				// Report the accumulated skin weight onto all the collocated vertices
+				if (TotalWeight > 0.0)
+				{
 					for (const int32 VertexIndex : CollocatedVertices[CollocatedIndex])
 					{
-						AccumulateSkinWeights(1.0f, FinalIndices[VertexIndex + VertexOffset], FinalWeights[VertexIndex + VertexOffset], SkinWeights, TotalWeight);
-					}
-
-					// Report the accumulated skin weight onto all the collocated vertices
-					if (TotalWeight > 0.0)
-					{
-						for (const int32 VertexIndex : CollocatedVertices[CollocatedIndex])
-						{
-							ReportSkinWeights(FinalIndices[VertexIndex + VertexOffset], FinalWeights[VertexIndex + VertexOffset], SkinWeights, TotalWeight);
-						}
+						ReportSkinWeights(FinalIndices[VertexIndex + VertexOffset], FinalWeights[VertexIndex + VertexOffset], SkinWeights, TotalWeight);
 					}
 				}
 			}
@@ -234,38 +329,124 @@ static void ClampSkinWeights(const int32 ClampingNumber, const TArray<TArray<int
 {
 	for(int32 VertexIndex = 0, NumVertices = FinalIndices.Num(); VertexIndex < NumVertices; ++VertexIndex)
 	{
-		if(SelectionMap[VertexIndex] > 0)
+		FinalIndices[VertexIndex] = SetupIndices[VertexIndex];
+		FinalWeights[VertexIndex] = SetupWeights[VertexIndex];
+
+		if((SelectionMap[VertexIndex] > 0) && (SetupWeights[VertexIndex].Num() > ClampingNumber))
 		{
-			TArray<TPair<int32,float>> SortedWeights;
-			SortedWeights.Reserve(SetupWeights[VertexIndex].Num());
-			for(int32 WeightIndex = 0, NumWeights = SetupWeights[VertexIndex].Num(); WeightIndex < NumWeights; ++WeightIndex)
-			{
-				SortedWeights.Add({SetupIndices[VertexIndex][WeightIndex], SetupWeights[VertexIndex][WeightIndex]});
-			}
-			// sort in descending order by weight
-			Algo::Sort(SortedWeights, [](const TPair<int32,float>& A, const TPair<int32,float>& B)
-			{
-				return A.Value > B.Value; 
-			});
-			float TotalWeight = 0.0f;
-			for(int32 WeightIndex = 0, NumWeights = FMath::Min(ClampingNumber,SetupWeights[VertexIndex].Num()); WeightIndex < NumWeights; ++WeightIndex)
+			RestrictSkinWeights(ClampingNumber, FinalIndices[VertexIndex], FinalWeights[VertexIndex]);
+		}
+	}
+}
+	
+static void SmoothVertexWeights(const UE::Geometry::FDynamicMesh3& DynamicMesh, const TArray<int32>& CollocatedVertices, const int32 VertexOffset, const float SmoothingFactor, 
+		TArray<TArray<int32>>& FinalIndices, TArray<TArray<float>>& FinalWeights)
+{
+	// Get list of all neighboring vertices, AND this vertex
+	TSet<int32> NeighborVertices;
+	for(const int32 CollocatedVertex : CollocatedVertices)
+	{
+		NeighborVertices.Add(CollocatedVertex);
+		for (const int32 NeighborVertex : DynamicMesh.VtxVerticesItr(CollocatedVertex))
+		{
+			NeighborVertices.Add(NeighborVertex);
+		}
+	}
+	// Average the per bone weights values
+	TMap<int32, FAveragedBoneWeight> NeighborWeights;
+	if (AverageSkinWeights(NeighborVertices, VertexOffset, FinalIndices, FinalWeights, NeighborWeights))
+	{
+		for (const int32 CollocatedVertex : CollocatedVertices)
+		{
+			const int32 GlobalVertex = VertexOffset + CollocatedVertex;
+
+			TArray<int32> BoneIndices;
+			TArray<float> BoneWeights;
+
+			float TotalWeight = 0.0;
+			for (const TPair<int32, FAveragedBoneWeight>& NeighborWeight : NeighborWeights)
 			{
-				FinalIndices[VertexIndex].Add(SortedWeights[WeightIndex].Key);
-				FinalWeights[VertexIndex].Add(SortedWeights[WeightIndex].Value);
-				TotalWeight += SortedWeights[WeightIndex].Value;
+				const int32 BoneIndex = NeighborWeight.Key;
+				float SmoothWeight = NeighborWeight.Value.WeightValue;
+				float StoredWeight = 0.0f;
+				
+				for (int32 WeightIndex = 0, NumWeights = FinalIndices[GlobalVertex].Num(); WeightIndex < NumWeights; ++WeightIndex)
+				{
+					if (BoneIndex == FinalIndices[GlobalVertex][WeightIndex])
+					{
+						StoredWeight = FinalWeights[GlobalVertex][WeightIndex];
+						break;
+					}
+				}
+				BoneWeights.Add(FMath::Lerp<float>(StoredWeight, SmoothWeight, SmoothingFactor));
+				BoneIndices.Add(BoneIndex);
+
+				TotalWeight += BoneWeights.Last();
 			}
-			if(TotalWeight != 0.0f)
+
+			if (!FMath::IsNearlyZero(TotalWeight))
 			{
-				for(int32 WeightIndex = 0, NumWeights = FinalWeights[VertexIndex].Num(); WeightIndex < NumWeights; ++WeightIndex)
+				for (float& VertexWeight : BoneWeights)
 				{
-					FinalWeights[VertexIndex][WeightIndex] /= TotalWeight;
+					VertexWeight /= TotalWeight;
 				}
+				FinalIndices[GlobalVertex] = BoneIndices;
+				FinalWeights[GlobalVertex] = BoneWeights;
 			}
 		}
-		else
+	}
+}
+
+// This is an equivalent of the RelaxSkinWeights using the collocated vertices in order to smooth skin weights across seam	
+static void SmoothSkinWeights(const UE::Geometry::FDynamicMesh3& DynamicMesh, const TArray<TArray<int32>>& CollocatedVertices, const int32 VertexOffset, const float SmoothStrength, const int32 NumIterations,
+			const TArray<TArray<int32>>& SetupIndices, const TArray<TArray<float>>& SetupWeights,
+			const TArray<float>& SelectionMap, TArray<TArray<int32>>& FinalIndices, TArray<TArray<float>>& FinalWeights)
+{
+	const int32 NumVertices = DynamicMesh.VertexCount();
+
+	TArray<TArray<int32>> SelectedVertices;
+	SelectedVertices.Reserve(NumVertices);
+
+	for(const TArray<int32>& CollocatedVertex : CollocatedVertices)
+	{
+		for (const int32 VertexIndex : CollocatedVertex)
 		{
-			FinalIndices[VertexIndex] = SetupIndices[VertexIndex];
-			FinalWeights[VertexIndex] = SetupWeights[VertexIndex];
+			if (SelectionMap[VertexIndex + VertexOffset] > 0.0f)
+			{
+				// Add the collocated vertex if one of the underlying vertices is selected 
+				SelectedVertices.Add(CollocatedVertex);
+				break;
+			}
+		}
+	}
+	int32 NumCollocated = SelectedVertices.Num();
+	for (int32 VertexIndex = 0; VertexIndex < NumVertices; ++VertexIndex)
+	{
+		const int32 GlobalVertex = VertexOffset+VertexIndex;
+		if(SelectionMap[GlobalVertex] > 0.0f && !DynamicMesh.IsBoundaryVertex(VertexIndex))
+		{
+			SelectedVertices.Add({VertexIndex});
+		}
+		// Initialize the final indices/weights with the setup ones
+		FinalIndices[GlobalVertex] = SetupIndices[GlobalVertex];
+		FinalWeights[GlobalVertex] = SetupWeights[GlobalVertex];
+	}
+
+	static constexpr float PercentPerIteration = 0.95f;
+	const float SmoothingFactor = SmoothStrength * PercentPerIteration;
+	
+	for (int32 IterationIndex = 0; IterationIndex < NumIterations; ++IterationIndex)
+	{
+		for (const TArray<int32>& SelectedVertex : SelectedVertices)
+		{
+			SmoothVertexWeights(DynamicMesh, SelectedVertex, VertexOffset, SmoothingFactor, FinalIndices, FinalWeights);
+		}
+	}
+	for (const TArray<int32>& SelectedVertex : SelectedVertices)
+	{
+		for (const int32 VertexIndex : SelectedVertex)
+		{
+			RestrictSkinWeights(MAX_TOTAL_INFLUENCES, FinalIndices[VertexIndex], FinalWeights[VertexIndex]);
 		}
 	}
 }
@@ -508,11 +689,14 @@ void FDataflowCorrectSkinWeightsNode::Evaluate(UE::Dataflow::FContext& Context,
 						UE::Dataflow::Private::BuildDynamicMeshes(this, Context, DynamicMeshes);
 
 						int32 VertexOffset = 0;
+						TArray<TArray<int32>> CollocatedVertices;
 						for(UE::Geometry::FDynamicMesh3& DynamicMesh : DynamicMeshes)
 						{
+							UE::Dataflow::Private::BuildCollocatedVertices(DynamicMesh, CollocatedVertices);
+							
 							if(CorrectionType == ESkinWeightsCorrectionType::Relax)
 							{
-								UE::Dataflow::Private::RelaxSkinWeights(DynamicMesh, VertexOffset, SmoothingFactor, SmoothingIterations,
+								UE::Dataflow::Private::SmoothSkinWeights(DynamicMesh, CollocatedVertices, VertexOffset, SmoothingFactor, SmoothingIterations,
 									SetupIndices, SetupWeights, SelectionMap, FinalIndices, FinalWeights);
 							}
 							else if (CorrectionType == ESkinWeightsCorrectionType::Hammer)
@@ -521,7 +705,7 @@ void FDataflowCorrectSkinWeightsNode::Evaluate(UE::Dataflow::FContext& Context,
 									SetupIndices, SetupWeights, SelectionMap, FinalIndices, FinalWeights);
 							}
 
-							UE::Dataflow::Private::MergeSkinWeights(DynamicMesh, VertexOffset, SelectionMap, FinalIndices, FinalWeights);
+							UE::Dataflow::Private::MergeSkinWeights(DynamicMesh, CollocatedVertices, VertexOffset, SelectionMap, FinalIndices, FinalWeights);
 
 							VertexOffset += DynamicMesh.VertexCount();
 						}
