Dataflow : fix double evaluation and asset update when requesting an evaluation on a terminal node

diff --git a/Engine/Plugins/Experimental/Dataflow/Source/DataflowEditor/Private/Dataflow/DataflowEditorToolkit.cpp b/Engine/Plugins/Experimental/Dataflow/Source/DataflowEditor/Private/Dataflow/DataflowEditorToolkit.cpp
--- a/Engine/Plugins/Experimental/Dataflow/Source/DataflowEditor/Private/Dataflow/DataflowEditorToolkit.cpp
+++ b/Engine/Plugins/Experimental/Dataflow/Source/DataflowEditor/Private/Dataflow/DataflowEditorToolkit.cpp
@@ -2087,16 +2087,22 @@ TSharedRef<SDataflowGraphEditor> FDataflowEditorToolkit::CreateGraphEditorWidget
 	const FDataflowEditorCommands::FGraphEvaluationCallback Evaluate =
 		[this](const FDataflowNode* Node, const FDataflowOutput* Output)
 		{
-			// evaluate may already have happen when the node was invalidated in the graph editor code 
-			// and cause on eof the terminal node to evaluate 
-			// but if the node is in a part of the graph that does not end by a terminal node , we need to explictly evaluate it 
-			if (const TObjectPtr<UDataflowBaseContent>& EditorContent = GetEditorContent())
+			// This method is called when a node is explicily called to be evaluated form the UI 
+			// Evaluate may already have happen when the node was invalidated in the graph editor code if the node is part of a branch that 
+			// ends with a terminal node and evaluation mode is automatic ( see OnNodeInvalidated )
+			// So to avoid double evaluation, we only evaluate if the node is not terminal or if we are in manual evaluation mode
+			// todo(dataflow) We should refactor this post 5.6 to make this simpler and less dependent on other method logic
+			const bool bIsTerminalNode = (Node && Node->AsType<FDataflowTerminalNode>());
+			if (!bIsTerminalNode || EvaluationMode == EDataflowEditorEvaluationMode::Manual)
 			{
-				UE::Dataflow::FTimestamp LastNodeTimestamp = EditorContent->GetLastModifiedTimestamp();
+				if (const TObjectPtr<UDataflowBaseContent>& EditorContent = GetEditorContent())
+				{
+					UE::Dataflow::FTimestamp LastNodeTimestamp = EditorContent->GetLastModifiedTimestamp();
 
-				EvaluateNode(Node, Output, LastNodeTimestamp);
+					EvaluateNode(Node, Output, LastNodeTimestamp);
 
-				EditorContent->SetLastModifiedTimestamp(LastNodeTimestamp);
+					EditorContent->SetLastModifiedTimestamp(LastNodeTimestamp);
+				}
 			}
 		};
 	
