Fixup API macro usage

#rnx

diff --git a/Engine/Plugins/Experimental/Dataflow/Source/DataflowNodes/Public/Dataflow/DataflowCollectionEditSkinWeightsNode.h b/Engine/Plugins/Experimental/Dataflow/Source/DataflowNodes/Public/Dataflow/DataflowCollectionEditSkinWeightsNode.h
--- a/Engine/Plugins/Experimental/Dataflow/Source/DataflowNodes/Public/Dataflow/DataflowCollectionEditSkinWeightsNode.h
+++ b/Engine/Plugins/Experimental/Dataflow/Source/DataflowNodes/Public/Dataflow/DataflowCollectionEditSkinWeightsNode.h
@@ -103,7 +103,7 @@ public:
     UE_API FCollectionAttributeKey GetBoneWeightsKey(UE::Dataflow::FContext& Context) const;
 
 	/** Validate the skeletal mesh construction */
-	UE_API void ValidateSkeletalMeshes() {bValidSkeletalMeshes = true;}
+	void ValidateSkeletalMeshes() {bValidSkeletalMeshes = true;}
 
 	/** Return the collection offset given a skeletal mesh*/
 	UE_API int32 GetSkeletalMeshOffset(const TObjectPtr<USkeletalMesh>& SkeletalMesh) const;
