Dataflow - Disabled the Freeze actions for 5.6.

#rnx

diff --git a/Engine/Plugins/Experimental/Dataflow/Source/DataflowEditor/Private/Dataflow/DataflowSchema.cpp b/Engine/Plugins/Experimental/Dataflow/Source/DataflowEditor/Private/Dataflow/DataflowSchema.cpp
--- a/Engine/Plugins/Experimental/Dataflow/Source/DataflowEditor/Private/Dataflow/DataflowSchema.cpp
+++ b/Engine/Plugins/Experimental/Dataflow/Source/DataflowEditor/Private/Dataflow/DataflowSchema.cpp
@@ -53,8 +53,10 @@ void UDataflowSchema::GetContextMenuActions(class UToolMenu* Menu, class UGraphN
 				Section.AddSeparator(TEXT("ActionsSeparator2"));
 				Section.AddMenuEntry(FDataflowEditorCommands::Get().EvaluateNode);
 				Section.AddSeparator(TEXT("ActionsSeparator3"));
+#if 0  // Disabled for 5.6
 				Section.AddMenuEntry(FDataflowEditorCommands::Get().FreezeNodes, LOCTEXT("DataflowContextMenu_Freeze", "Freeze"));
 				Section.AddMenuEntry(FDataflowEditorCommands::Get().UnfreezeNodes, LOCTEXT("DataflowContextMenu_Unfreeze", "Unfreeze"));
+#endif
 			}
 		}
 		{
diff --git a/Engine/Source/Runtime/Experimental/Dataflow/Core/Public/Dataflow/DataflowNode.h b/Engine/Source/Runtime/Experimental/Dataflow/Core/Public/Dataflow/DataflowNode.h
--- a/Engine/Source/Runtime/Experimental/Dataflow/Core/Public/Dataflow/DataflowNode.h
+++ b/Engine/Source/Runtime/Experimental/Dataflow/Core/Public/Dataflow/DataflowNode.h
@@ -864,10 +864,12 @@ private:
 	TMap<UE::Dataflow::FConnectionKey, FDataflowOutput*> ExpandedOutputs;
 	TMap<uint32, const FArrayProperty*> OutputArrayProperties; // Used to calculate ContainerElementOffsets
 
+#if 0  // Disabled for 5.6
 #if WITH_EDITORONLY_DATA
 	// Freeze node customization
 	UPROPERTY(EditAnywhere, Category = "Dataflow")
 	FDataflowFreezeActions FreezeActions;
+#endif
 #endif
 
 	UPROPERTY()
