Dataflow Editor : remove non working UI menu option for subgraphs in DataflowEditor

diff --git a/Engine/Plugins/Experimental/Dataflow/Source/DataflowEditor/Private/Dataflow/DataflowMembersWidget.cpp b/Engine/Plugins/Experimental/Dataflow/Source/DataflowEditor/Private/Dataflow/DataflowMembersWidget.cpp
--- a/Engine/Plugins/Experimental/Dataflow/Source/DataflowEditor/Private/Dataflow/DataflowMembersWidget.cpp
+++ b/Engine/Plugins/Experimental/Dataflow/Source/DataflowEditor/Private/Dataflow/DataflowMembersWidget.cpp
@@ -768,11 +768,14 @@ TSharedPtr<SWidget> SDataflowMembersWidget::OnContextMenuOpening()
 				MenuBuilder.AddMenuEntry(FDataflowEditorCommands::Get().ConvertToForEachSubGraph, NAME_None, EmptyTextAttribute, EmptyTextAttribute, LoopIcon);
 			}
 		}
-		MenuBuilder.AddMenuSeparator();
-		MenuBuilder.AddMenuEntry(FGenericCommands::Get().Cut);
-		MenuBuilder.AddMenuEntry(FGenericCommands::Get().Copy);
-		MenuBuilder.AddMenuEntry(FGenericCommands::Get().Paste);
-		MenuBuilder.AddMenuEntry(FGenericCommands::Get().Duplicate);
+		if (!bSubGraphActions)
+		{
+			MenuBuilder.AddMenuSeparator();
+			MenuBuilder.AddMenuEntry(FGenericCommands::Get().Cut);
+			MenuBuilder.AddMenuEntry(FGenericCommands::Get().Copy);
+			MenuBuilder.AddMenuEntry(FGenericCommands::Get().Paste);
+			MenuBuilder.AddMenuEntry(FGenericCommands::Get().Duplicate);
+		}
 		MenuBuilder.AddMenuSeparator();
 		MenuBuilder.AddMenuEntry(FGenericCommands::Get().Delete);
 	}
