Dataflow : Add customization for common anytypes to inline the value of the anytype struct type to its root level

diff --git a/Engine/Plugins/Experimental/Dataflow/Source/DataflowEditor/DataflowEditor.Build.cs b/Engine/Plugins/Experimental/Dataflow/Source/DataflowEditor/DataflowEditor.Build.cs
--- a/Engine/Plugins/Experimental/Dataflow/Source/DataflowEditor/DataflowEditor.Build.cs
+++ b/Engine/Plugins/Experimental/Dataflow/Source/DataflowEditor/DataflowEditor.Build.cs
@@ -40,6 +40,7 @@ namespace UnrealBuildTool.Rules
 					"DataflowEnginePlugin",
 					"DataflowNodes",
 					"DataflowSimulation",
+					"DetailCustomizations",
 					"DeveloperSettings",
 					"DynamicMesh",
 					"Engine",
diff --git a/Engine/Plugins/Experimental/Dataflow/Source/DataflowEditor/Private/Dataflow/DataflowAnyTypeCustomization.cpp b/Engine/Plugins/Experimental/Dataflow/Source/DataflowEditor/Private/Dataflow/DataflowAnyTypeCustomization.cpp
new file mode 100644
--- /dev/null
+++ b/Engine/Plugins/Experimental/Dataflow/Source/DataflowEditor/Private/Dataflow/DataflowAnyTypeCustomization.cpp
@@ -0,0 +1,152 @@
+// Copyright Epic Games, Inc. All Rights Reserved.
+
+#include "Dataflow/DataflowAnyTypeCustomization.h"
+#include "Dataflow/DataflowAnyType.h"
+#include "Customizations/MathStructCustomizations.h"
+#include "PropertyEditorModule.h"
+
+#include "PropertyHandle.h"
+#include "DetailWidgetRow.h"
+
+namespace UE::Dataflow
+{
+	static const FName DataflowNumericTypesName = TEXT("DataflowNumericTypes");
+	static const FName DataflowBoolTypesName = TEXT("DataflowBoolTypes");
+	static const FName DataflowStringTypesName = TEXT("DataflowStringTypes");
+	static const FName DataflowStringConvertibleTypesName = TEXT("DataflowStringConvertibleTypes");
+	static const FName DataflowVectorTypesName = TEXT("DataflowVectorTypes");
+	static const FName DataflowTransformTypesName = TEXT("DataflowTransformTypes");
+
+	void RegisterAnyTypeCustomizations(FPropertyEditorModule& PropertyModule)
+	{
+		PropertyModule.RegisterCustomPropertyTypeLayout(DataflowNumericTypesName, FOnGetPropertyTypeCustomizationInstance::CreateStatic(&UE::Dataflow::FNumericAnyTypeCustomization::MakeInstance));
+		PropertyModule.RegisterCustomPropertyTypeLayout(DataflowBoolTypesName, FOnGetPropertyTypeCustomizationInstance::CreateStatic(&UE::Dataflow::FBoolAnyTypeCustomization::MakeInstance));
+		PropertyModule.RegisterCustomPropertyTypeLayout(DataflowStringTypesName, FOnGetPropertyTypeCustomizationInstance::CreateStatic(&UE::Dataflow::FStringAnyTypeCustomization::MakeInstance));
+		PropertyModule.RegisterCustomPropertyTypeLayout(DataflowStringConvertibleTypesName, FOnGetPropertyTypeCustomizationInstance::CreateStatic(&UE::Dataflow::FStringConvertibleAnyTypeCustomization::MakeInstance));
+		PropertyModule.RegisterCustomPropertyTypeLayout(DataflowVectorTypesName, FOnGetPropertyTypeCustomizationInstance::CreateStatic(&UE::Dataflow::FVectorAnyTypeCustomization::MakeInstance));
+		PropertyModule.RegisterCustomPropertyTypeLayout(DataflowTransformTypesName, FOnGetPropertyTypeCustomizationInstance::CreateStatic(&UE::Dataflow::FTransformAnyTypeCustomization::MakeInstance));
+	}
+
+	void UnregisterAnyTypeCustomizations(FPropertyEditorModule& PropertyModule)
+	{
+		PropertyModule.UnregisterCustomPropertyTypeLayout(DataflowNumericTypesName);
+		PropertyModule.UnregisterCustomPropertyTypeLayout(DataflowBoolTypesName);
+		PropertyModule.UnregisterCustomPropertyTypeLayout(DataflowStringTypesName);
+		PropertyModule.UnregisterCustomPropertyTypeLayout(DataflowStringConvertibleTypesName);
+		PropertyModule.UnregisterCustomPropertyTypeLayout(DataflowVectorTypesName);
+		PropertyModule.UnregisterCustomPropertyTypeLayout(DataflowTransformTypesName);
+	}
+
+	/////////////////////////////////////////////////////////////////////////////////////////////////////
+
+	namespace Private
+	{
+		template <typename T>
+		void CustomizeHeaderForTypeWithoutCustomization(TSharedRef<class IPropertyHandle> StructPropertyHandle, class FDetailWidgetRow& HeaderRow)
+		{
+			TSharedPtr<IPropertyHandle> ValuePropertyHandle = StructPropertyHandle->GetChildHandle(GET_MEMBER_NAME_CHECKED(T, Value));
+			check(ValuePropertyHandle);
+
+			HeaderRow
+				.NameContent()
+				[
+					StructPropertyHandle->CreatePropertyNameWidget(StructPropertyHandle->GetPropertyDisplayName())
+				]
+				.ValueContent()
+				[
+					ValuePropertyHandle->CreatePropertyValueWidget()
+				];
+		}
+
+		template <typename T, typename TCustomization>
+		void CustomizeHeaderForTypeWithCustomization(TSharedRef<class IPropertyHandle> StructPropertyHandle, class FDetailWidgetRow& HeaderRow, IPropertyTypeCustomizationUtils& StructCustomizationUtils)
+		{
+			TSharedPtr<IPropertyHandle> ValuePropertyHandle = StructPropertyHandle->GetChildHandle(GET_MEMBER_NAME_CHECKED(T, Value));
+			check(ValuePropertyHandle);
+
+			TSharedRef<TCustomization> MathStructCustomization = MakeShareable(new TCustomization);
+			MathStructCustomization->CustomizeHeader(ValuePropertyHandle.ToSharedRef(), HeaderRow, StructCustomizationUtils);;
+
+			// make sure we fix the name to be the original parent 
+			HeaderRow
+				.NameContent()
+				[
+					StructPropertyHandle->CreatePropertyNameWidget(StructPropertyHandle->GetPropertyDisplayName())
+				];
+		}
+	}
+
+	////////////////////////////////////////////////////////////////////////////////////////////////
+
+	TSharedRef<IPropertyTypeCustomization> FNumericAnyTypeCustomization::MakeInstance()
+	{
+		return MakeShareable(new FNumericAnyTypeCustomization);
+	}
+
+	void FNumericAnyTypeCustomization::CustomizeHeader(TSharedRef<class IPropertyHandle> StructPropertyHandle, class FDetailWidgetRow& HeaderRow, IPropertyTypeCustomizationUtils& StructCustomizationUtils)
+	{
+		Private::CustomizeHeaderForTypeWithoutCustomization<FDataflowNumericTypes>(StructPropertyHandle, HeaderRow);
+	}
+
+	////////////////////////////////////////////////////////////////////////////////////////////////
+
+	TSharedRef<IPropertyTypeCustomization> FBoolAnyTypeCustomization::MakeInstance()
+	{
+		return MakeShareable(new FBoolAnyTypeCustomization);
+	}
+
+	void FBoolAnyTypeCustomization::CustomizeHeader(TSharedRef<class IPropertyHandle> StructPropertyHandle, class FDetailWidgetRow& HeaderRow, IPropertyTypeCustomizationUtils& StructCustomizationUtils)
+	{
+		Private::CustomizeHeaderForTypeWithoutCustomization<FDataflowBoolTypes>(StructPropertyHandle, HeaderRow);
+	}
+
+	////////////////////////////////////////////////////////////////////////////////////////////////
+
+	TSharedRef<IPropertyTypeCustomization> FStringAnyTypeCustomization::MakeInstance()
+	{
+		return MakeShareable(new FStringAnyTypeCustomization);
+	}
+
+	void FStringAnyTypeCustomization::CustomizeHeader(TSharedRef<class IPropertyHandle> StructPropertyHandle, class FDetailWidgetRow& HeaderRow, IPropertyTypeCustomizationUtils& StructCustomizationUtils)
+	{
+		Private::CustomizeHeaderForTypeWithoutCustomization<FDataflowStringTypes>(StructPropertyHandle, HeaderRow);
+	}
+
+	////////////////////////////////////////////////////////////////////////////////////////////////
+
+	TSharedRef<IPropertyTypeCustomization> FStringConvertibleAnyTypeCustomization::MakeInstance()
+	{
+		return MakeShareable(new FStringConvertibleAnyTypeCustomization);
+	}
+
+	void FStringConvertibleAnyTypeCustomization::CustomizeHeader(TSharedRef<class IPropertyHandle> StructPropertyHandle, class FDetailWidgetRow& HeaderRow, IPropertyTypeCustomizationUtils& StructCustomizationUtils)
+	{
+		Private::CustomizeHeaderForTypeWithoutCustomization<FDataflowStringConvertibleTypes>(StructPropertyHandle, HeaderRow);
+	}
+
+	////////////////////////////////////////////////////////////////////////////////////////////////
+	
+	TSharedRef<IPropertyTypeCustomization> FVectorAnyTypeCustomization::MakeInstance()
+	{
+		return MakeShareable(new FVectorAnyTypeCustomization);
+	}
+
+	void FVectorAnyTypeCustomization::CustomizeHeader(TSharedRef<class IPropertyHandle> StructPropertyHandle, class FDetailWidgetRow& HeaderRow, IPropertyTypeCustomizationUtils& StructCustomizationUtils)
+	{
+		Private::CustomizeHeaderForTypeWithCustomization<FDataflowVectorTypes, FMathStructCustomization>(StructPropertyHandle, HeaderRow, StructCustomizationUtils);
+	}
+
+	////////////////////////////////////////////////////////////////////////////////////////////////
+
+	TSharedRef<IPropertyTypeCustomization> FTransformAnyTypeCustomization::MakeInstance()
+	{
+		return MakeShareable(new FTransformAnyTypeCustomization);
+	}
+
+	void FTransformAnyTypeCustomization::CustomizeHeader(TSharedRef<class IPropertyHandle> StructPropertyHandle, class FDetailWidgetRow& HeaderRow, IPropertyTypeCustomizationUtils& StructCustomizationUtils)
+	{
+		Private::CustomizeHeaderForTypeWithoutCustomization<FDataflowTransformTypes>(StructPropertyHandle, HeaderRow);
+	}
+
+	////////////////////////////////////////////////////////////////////////////////////////////////
+}
diff --git a/Engine/Plugins/Experimental/Dataflow/Source/DataflowEditor/Private/Dataflow/DataflowAnyTypeCustomization.h b/Engine/Plugins/Experimental/Dataflow/Source/DataflowEditor/Private/Dataflow/DataflowAnyTypeCustomization.h
new file mode 100644
--- /dev/null
+++ b/Engine/Plugins/Experimental/Dataflow/Source/DataflowEditor/Private/Dataflow/DataflowAnyTypeCustomization.h
@@ -0,0 +1,114 @@
+// Copyright Epic Games, Inc. All Rights Reserved.
+
+#pragma once
+
+#include "IPropertyTypeCustomization.h"
+
+class SDataflowGraphEditor;
+class SWidget;
+class ITableRow;
+class STableViewBase;
+class SComboButton;
+
+class FPropertyEditorModule;
+
+namespace UE::Dataflow
+{
+	class FAnyTypeCustomizationBase
+		: public IPropertyTypeCustomization
+	{
+		//~ Begin IPropertyTypeCustomization implementation
+		virtual void CustomizeChildren(TSharedRef<IPropertyHandle> StructPropertyHandle, IDetailChildrenBuilder& StructBuilder, IPropertyTypeCustomizationUtils& StructCustomizationUtils) override {}
+		//~ End IPropertyTypeCustomization implementation
+	protected:
+	};
+
+	//////////////////////////////////////////////////////////////////////////////////////////////////////
+
+	class FNumericAnyTypeCustomization
+		: public FAnyTypeCustomizationBase
+	{
+	public:
+		static TSharedRef<IPropertyTypeCustomization> MakeInstance();
+
+	protected:
+		//~ Begin IPropertyTypeCustomization implementation
+		virtual void CustomizeHeader(TSharedRef<IPropertyHandle> StructPropertyHandle, FDetailWidgetRow& HeaderRow, IPropertyTypeCustomizationUtils& StructCustomizationUtils) override;
+		//~ End IPropertyTypeCustomization implementation
+	};
+
+	//////////////////////////////////////////////////////////////////////////////////////////////////////
+
+	class FBoolAnyTypeCustomization
+		: public FAnyTypeCustomizationBase
+	{
+	public:
+		static TSharedRef<IPropertyTypeCustomization> MakeInstance();
+
+	protected:
+		//~ Begin IPropertyTypeCustomization implementation
+		virtual void CustomizeHeader(TSharedRef<IPropertyHandle> StructPropertyHandle, FDetailWidgetRow& HeaderRow, IPropertyTypeCustomizationUtils& StructCustomizationUtils) override;
+		//~ End IPropertyTypeCustomization implementation
+	};
+
+	//////////////////////////////////////////////////////////////////////////////////////////////////////
+
+	class FStringAnyTypeCustomization
+		: public FAnyTypeCustomizationBase
+	{
+	public:
+		static TSharedRef<IPropertyTypeCustomization> MakeInstance();
+
+	protected:
+		//~ Begin IPropertyTypeCustomization implementation
+		virtual void CustomizeHeader(TSharedRef<IPropertyHandle> StructPropertyHandle, FDetailWidgetRow& HeaderRow, IPropertyTypeCustomizationUtils& StructCustomizationUtils) override;
+		//~ End IPropertyTypeCustomization implementation
+	};
+
+	//////////////////////////////////////////////////////////////////////////////////////////////////////
+
+	class FStringConvertibleAnyTypeCustomization
+		: public FAnyTypeCustomizationBase
+	{
+	public:
+		static TSharedRef<IPropertyTypeCustomization> MakeInstance();
+
+	protected:
+		//~ Begin IPropertyTypeCustomization implementation
+		virtual void CustomizeHeader(TSharedRef<IPropertyHandle> StructPropertyHandle, FDetailWidgetRow& HeaderRow, IPropertyTypeCustomizationUtils& StructCustomizationUtils) override;
+		//~ End IPropertyTypeCustomization implementation
+	};
+
+	//////////////////////////////////////////////////////////////////////////////////////////////////////
+
+	class FVectorAnyTypeCustomization
+		: public FAnyTypeCustomizationBase
+	{
+	public:
+		static TSharedRef<IPropertyTypeCustomization> MakeInstance();
+
+	protected:
+		//~ Begin IPropertyTypeCustomization implementation
+		virtual void CustomizeHeader(TSharedRef<IPropertyHandle> StructPropertyHandle, FDetailWidgetRow& HeaderRow, IPropertyTypeCustomizationUtils& StructCustomizationUtils) override;
+		//~ End IPropertyTypeCustomization implementation
+	};
+
+	//////////////////////////////////////////////////////////////////////////////////////////////////////
+
+	class FTransformAnyTypeCustomization
+		: public FAnyTypeCustomizationBase
+	{
+	public:
+		static TSharedRef<IPropertyTypeCustomization> MakeInstance();
+
+	protected:
+		//~ Begin IPropertyTypeCustomization implementation
+		virtual void CustomizeHeader(TSharedRef<IPropertyHandle> StructPropertyHandle, FDetailWidgetRow& HeaderRow, IPropertyTypeCustomizationUtils& StructCustomizationUtils) override;
+		//~ End IPropertyTypeCustomization implementation
+	};
+
+	//////////////////////////////////////////////////////////////////////////////////////////////////////
+
+	void RegisterAnyTypeCustomizations(FPropertyEditorModule& PropertyModule);
+	void UnregisterAnyTypeCustomizations(FPropertyEditorModule& PropertyModule);
+}
diff --git a/Engine/Plugins/Experimental/Dataflow/Source/DataflowEditor/Private/Dataflow/DataflowEditorModule.cpp b/Engine/Plugins/Experimental/Dataflow/Source/DataflowEditor/Private/Dataflow/DataflowEditorModule.cpp
--- a/Engine/Plugins/Experimental/Dataflow/Source/DataflowEditor/Private/Dataflow/DataflowEditorModule.cpp
+++ b/Engine/Plugins/Experimental/Dataflow/Source/DataflowEditor/Private/Dataflow/DataflowEditorModule.cpp
@@ -15,6 +15,7 @@
 #include "Dataflow/MeshStatsConstructionVisualization.h"
 #include "Dataflow/MeshConstructionVisualization.h"
 #include "Dataflow/ScalarVertexPropertyGroupCustomization.h"
+#include "Dataflow/DataflowAnyTypeCustomization.h"
 #include "Dataflow/DataflowToolRegistry.h"
 #include "DataflowEditorTools/DataflowEditorWeightMapPaintTool.h"
 #include "Dataflow/DataflowCollectionAddScalarVertexPropertyNode.h"
@@ -119,6 +120,7 @@ void FDataflowEditorModule::StartupModule()
 		PropertyModule->RegisterCustomPropertyTypeLayout(DataflowFunctionPropertyName, FOnGetPropertyTypeCustomizationInstance::CreateStatic(&UE::Dataflow::FFunctionPropertyCustomization::MakeInstance));
 		PropertyModule->RegisterCustomPropertyTypeLayout(DataflowVariableOverridesName, FOnGetPropertyTypeCustomizationInstance::CreateStatic(&FDataflowVariableOverridesDetails::MakeInstance));
 		PropertyModule->RegisterCustomPropertyTypeLayout(DataflowFreezeActionsName, FOnGetPropertyTypeCustomizationInstance::CreateStatic(&UE::Dataflow::FFreezeActionsCustomization::MakeInstance));
+		UE::Dataflow::RegisterAnyTypeCustomizations(*PropertyModule);
 	}
 
 	UE::Dataflow::RenderingCallbacks();
@@ -158,6 +160,7 @@ void FDataflowEditorModule::ShutdownModule()
 		PropertyModule->UnregisterCustomPropertyTypeLayout(DataflowFunctionPropertyName);
 		PropertyModule->UnregisterCustomPropertyTypeLayout(DataflowVariableOverridesName);
 		PropertyModule->UnregisterCustomPropertyTypeLayout(DataflowFreezeActionsName);
+		UE::Dataflow::UnregisterAnyTypeCustomizations(*PropertyModule);
 	}
 
 	FDataflowEditorCommands::Unregister();
diff --git a/Engine/Source/Runtime/Experimental/Dataflow/Core/Public/Dataflow/DataflowMathNodes.h b/Engine/Source/Runtime/Experimental/Dataflow/Core/Public/Dataflow/DataflowMathNodes.h
--- a/Engine/Source/Runtime/Experimental/Dataflow/Core/Public/Dataflow/DataflowMathNodes.h
+++ b/Engine/Source/Runtime/Experimental/Dataflow/Core/Public/Dataflow/DataflowMathNodes.h
@@ -24,7 +24,7 @@ struct FDataflowMathOneInputOperatorNode : public FDataflowNode
 {
 	GENERATED_USTRUCT_BODY()
 
-	UPROPERTY(EditAnywhere, Category = Operands, meta = (DataflowInput, ShowOnlyInnerProperties));
+	UPROPERTY(EditAnywhere, Category = Operands, meta = (DataflowInput));
 	FDataflowNumericTypes A;
 
 	UPROPERTY(meta = (DataflowOutput))
@@ -46,10 +46,10 @@ struct FDataflowMathTwoInputsOperatorNode : public FDataflowNode
 {
 	GENERATED_USTRUCT_BODY()
 
-	UPROPERTY(EditAnywhere, Category = Operands, meta = (DataflowInput, ShowOnlyInnerProperties));
+	UPROPERTY(EditAnywhere, Category = Operands, meta = (DataflowInput));
 	FDataflowNumericTypes A;
 
-	UPROPERTY(EditAnywhere, Category = Operands, meta = (DataflowInput, ShowOnlyInnerProperties));
+	UPROPERTY(EditAnywhere, Category = Operands, meta = (DataflowInput));
 	FDataflowNumericTypes B;
 
 	UPROPERTY(meta = (DataflowOutput))
@@ -111,7 +111,7 @@ struct FDataflowMathDivideNode : public FDataflowMathTwoInputsOperatorNode
 	GENERATED_USTRUCT_BODY()
 	DATAFLOW_NODE_DEFINE_INTERNAL(FDataflowMathDivideNode, "Divide", DATAFLOW_MATH_NODES_CATEGORY, "/ Division")
 
-	UPROPERTY(EditAnywhere, Category = SafeDivide, meta = (DataflowInput, ShowOnlyInnerProperties));
+	UPROPERTY(EditAnywhere, Category = SafeDivide, meta = (DataflowInput));
 	FDataflowNumericTypes Fallback;
 
 public:
@@ -229,7 +229,7 @@ struct FDataflowMathReciprocalNode : public FDataflowMathOneInputOperatorNode
 	GENERATED_USTRUCT_BODY()
 	DATAFLOW_NODE_DEFINE_INTERNAL(FDataflowMathReciprocalNode, "Reciprocal", DATAFLOW_MATH_NODES_CATEGORY, "OneOver OneDividedBy")
 
-	UPROPERTY(EditAnywhere, Category = SafeDivide, meta = (DataflowInput, ShowOnlyInnerProperties));
+	UPROPERTY(EditAnywhere, Category = SafeDivide, meta = (DataflowInput));
 	FDataflowNumericTypes Fallback;
 
 public:
@@ -283,7 +283,7 @@ struct FDataflowMathInverseSquareRootNode : public FDataflowMathOneInputOperator
 	GENERATED_USTRUCT_BODY()
 	DATAFLOW_NODE_DEFINE_INTERNAL(FDataflowMathInverseSquareRootNode, "InverseSquareRoot", DATAFLOW_MATH_NODES_CATEGORY, "OneOverSquareRoot")
 
-	UPROPERTY(EditAnywhere, Category = SafeDivide, meta = (DataflowInput, ShowOnlyInnerProperties));
+	UPROPERTY(EditAnywhere, Category = SafeDivide, meta = (DataflowInput));
 	FDataflowNumericTypes Fallback;
 
 public:
@@ -397,7 +397,7 @@ struct FDataflowMathLogXNode : public FDataflowMathOneInputOperatorNode
 	GENERATED_USTRUCT_BODY()
 	DATAFLOW_NODE_DEFINE_INTERNAL(FDataflowMathLogXNode, "LogX", DATAFLOW_MATH_NODES_CATEGORY, "Logarithm")
 
-	UPROPERTY(EditAnywhere, Category = SafeDivide, meta = (DataflowInput, ShowOnlyInnerProperties));
+	UPROPERTY(EditAnywhere, Category = SafeDivide, meta = (DataflowInput));
 	FDataflowNumericTypes Base;
 
 public:
@@ -503,10 +503,10 @@ struct FDataflowMathClampNode : public FDataflowMathOneInputOperatorNode
 	GENERATED_USTRUCT_BODY()
 	DATAFLOW_NODE_DEFINE_INTERNAL(FDataflowMathClampNode, "Clamp", DATAFLOW_MATH_NODES_CATEGORY, "Saturate Limits")
 
-	UPROPERTY(EditAnywhere, Category = "Clamp", meta = (DataflowInput, ShowOnlyInnerProperties))
+	UPROPERTY(EditAnywhere, Category = "Clamp", meta = (DataflowInput))
 	FDataflowNumericTypes Min;
 
-	UPROPERTY(EditAnywhere, Category = "Clamp", meta = (DataflowInput, ShowOnlyInnerProperties))
+	UPROPERTY(EditAnywhere, Category = "Clamp", meta = (DataflowInput))
 	FDataflowNumericTypes Max;
 
 public:
diff --git a/Engine/Source/Runtime/Experimental/Dataflow/Core/Public/Dataflow/DataflowVectorNodes.h b/Engine/Source/Runtime/Experimental/Dataflow/Core/Public/Dataflow/DataflowVectorNodes.h
--- a/Engine/Source/Runtime/Experimental/Dataflow/Core/Public/Dataflow/DataflowVectorNodes.h
+++ b/Engine/Source/Runtime/Experimental/Dataflow/Core/Public/Dataflow/DataflowVectorNodes.h
@@ -20,11 +20,11 @@ public:
 
 private:
 	/** X component */
-	UPROPERTY(EditAnywhere, Category = Operands, meta = (DataflowInput, ShowOnlyInnerProperties));
+	UPROPERTY(EditAnywhere, Category = Operands, meta = (DataflowInput));
 	FDataflowNumericTypes X;
 
 	/** Y component */
-	UPROPERTY(EditAnywhere, Category = Operands, meta = (DataflowInput, ShowOnlyInnerProperties));
+	UPROPERTY(EditAnywhere, Category = Operands, meta = (DataflowInput));
 	FDataflowNumericTypes Y;
 
 	/** 2D Vector {X, Y} */
@@ -46,15 +46,15 @@ public:
 
 private:
 	/** X component */
-	UPROPERTY(EditAnywhere, Category = Operands, meta = (DataflowInput, ShowOnlyInnerProperties));
+	UPROPERTY(EditAnywhere, Category = Operands, meta = (DataflowInput));
 	FDataflowNumericTypes X;
 
 	/** Y component */
-	UPROPERTY(EditAnywhere, Category = Operands, meta = (DataflowInput, ShowOnlyInnerProperties));
+	UPROPERTY(EditAnywhere, Category = Operands, meta = (DataflowInput));
 	FDataflowNumericTypes Y;
 
 	/** Z component */
-	UPROPERTY(EditAnywhere, Category = Operands, meta = (DataflowInput, ShowOnlyInnerProperties));
+	UPROPERTY(EditAnywhere, Category = Operands, meta = (DataflowInput));
 	FDataflowNumericTypes Z;
 
 	/** 3D Vector {X, Y, Z} */
@@ -77,15 +77,15 @@ public:
 
 private:
 	/** X component */
-	UPROPERTY(EditAnywhere, Category = Operands, meta = (DataflowInput, ShowOnlyInnerProperties));
+	UPROPERTY(EditAnywhere, Category = Operands, meta = (DataflowInput));
 	FDataflowNumericTypes X;
 
 	/** Y component */
-	UPROPERTY(EditAnywhere, Category = Operands, meta = (DataflowInput, ShowOnlyInnerProperties));
+	UPROPERTY(EditAnywhere, Category = Operands, meta = (DataflowInput));
 	FDataflowNumericTypes Y;
 
 	/** Z component */
-	UPROPERTY(EditAnywhere, Category = Operands, meta = (DataflowInput, ShowOnlyInnerProperties));
+	UPROPERTY(EditAnywhere, Category = Operands, meta = (DataflowInput));
 	FDataflowNumericTypes Z;
 
 	/** W component */
@@ -115,7 +115,7 @@ public:
 
 private:
 	/** Vector to break into components */
-	UPROPERTY(EditAnywhere, Category = Operands, meta = (DataflowInput, ShowOnlyInnerProperties));
+	UPROPERTY(EditAnywhere, Category = Operands, meta = (DataflowInput));
 	FDataflowVectorTypes V;
 
 	/** X component */
@@ -149,11 +149,11 @@ public:
 
 private:
 	/** A Vector operand */
-	UPROPERTY(EditAnywhere, Category = Operands, meta = (DataflowInput, ShowOnlyInnerProperties));
+	UPROPERTY(EditAnywhere, Category = Operands, meta = (DataflowInput));
 	FDataflowVectorTypes A;
 
 	/** B Vector operand */
-	UPROPERTY(EditAnywhere, Category = Operands, meta = (DataflowInput, ShowOnlyInnerProperties));
+	UPROPERTY(EditAnywhere, Category = Operands, meta = (DataflowInput));
 	FDataflowVectorTypes B;
 
 	/** Add result V=A+B */
@@ -175,11 +175,11 @@ public:
 
 private:
 	/** A Vector operand */
-	UPROPERTY(EditAnywhere, Category = Operands, meta = (DataflowInput, ShowOnlyInnerProperties));
+	UPROPERTY(EditAnywhere, Category = Operands, meta = (DataflowInput));
 	FDataflowVectorTypes A;
 
 	/** B Vector operand */
-	UPROPERTY(EditAnywhere, Category = Operands, meta = (DataflowInput, ShowOnlyInnerProperties));
+	UPROPERTY(EditAnywhere, Category = Operands, meta = (DataflowInput));
 	FDataflowVectorTypes B;
 
 	/** Add result V=A-B */
@@ -201,11 +201,11 @@ public:
 
 private:
 	/** A Vector operand */
-	UPROPERTY(EditAnywhere, Category = Operands, meta = (DataflowInput, ShowOnlyInnerProperties));
+	UPROPERTY(EditAnywhere, Category = Operands, meta = (DataflowInput));
 	FDataflowVectorTypes A;
 
 	/** B Vector operand */
-	UPROPERTY(EditAnywhere, Category = Operands, meta = (DataflowInput, ShowOnlyInnerProperties));
+	UPROPERTY(EditAnywhere, Category = Operands, meta = (DataflowInput));
 	FDataflowVectorTypes B;
 
 	/** Add result V=A*B */
@@ -230,15 +230,15 @@ public:
 
 private:
 	/** A Vector operand */
-	UPROPERTY(EditAnywhere, Category = Operands, meta = (DataflowInput, ShowOnlyInnerProperties));
+	UPROPERTY(EditAnywhere, Category = Operands, meta = (DataflowInput));
 	FDataflowVectorTypes A;
 
 	/** B Vector operand */
-	UPROPERTY(EditAnywhere, Category = Operands, meta = (DataflowInput, ShowOnlyInnerProperties));
+	UPROPERTY(EditAnywhere, Category = Operands, meta = (DataflowInput));
 	FDataflowVectorTypes B;
 
 	/** Fallback Vector used when components of B are zero */
-	UPROPERTY(EditAnywhere, Category = Operands, meta = (DataflowInput, ShowOnlyInnerProperties));
+	UPROPERTY(EditAnywhere, Category = Operands, meta = (DataflowInput));
 	FDataflowVectorTypes Fallback;
 
 	/** Add result V=A*B */
@@ -260,11 +260,11 @@ public:
 
 private:
 	/** A Vector operand */
-	UPROPERTY(EditAnywhere, Category = Operands, meta = (DataflowInput, ShowOnlyInnerProperties));
+	UPROPERTY(EditAnywhere, Category = Operands, meta = (DataflowInput));
 	FDataflowVectorTypes A;
 
 	/** B Vector operand */
-	UPROPERTY(EditAnywhere, Category = Operands, meta = (DataflowInput, ShowOnlyInnerProperties));
+	UPROPERTY(EditAnywhere, Category = Operands, meta = (DataflowInput));
 	FDataflowVectorTypes B;
 
 	/** Resulting dot product : DotProduct=A.B  */
@@ -286,7 +286,7 @@ public:
 
 private:
 	/** Vector to get length from */
-	UPROPERTY(EditAnywhere, Category = Operands, meta = (DataflowInput, ShowOnlyInnerProperties));
+	UPROPERTY(EditAnywhere, Category = Operands, meta = (DataflowInput));
 	FDataflowVectorTypes V;
 
 	/** Length of the input vector : Length=|V| */
@@ -308,7 +308,7 @@ public:
 
 private:
 	/** Vector to get squared length from */
-	UPROPERTY(EditAnywhere, Category = Operands, meta = (DataflowInput, ShowOnlyInnerProperties));
+	UPROPERTY(EditAnywhere, Category = Operands, meta = (DataflowInput));
 	FDataflowVectorTypes V;
 
 	/** Length of the input vector : SquaredLength = |V||V| */
@@ -330,11 +330,11 @@ public:
 
 private:
 	/** A Vector operand */
-	UPROPERTY(EditAnywhere, Category = Operands, meta = (DataflowInput, ShowOnlyInnerProperties));
+	UPROPERTY(EditAnywhere, Category = Operands, meta = (DataflowInput));
 	FDataflowVectorTypes A;
 
 	/** B Vector operand */
-	UPROPERTY(EditAnywhere, Category = Operands, meta = (DataflowInput, ShowOnlyInnerProperties));
+	UPROPERTY(EditAnywhere, Category = Operands, meta = (DataflowInput));
 	FDataflowVectorTypes B;
 
 	/** Distance between A and B : Distance=|B-A| */
@@ -359,11 +359,11 @@ public:
 
 private:
 	/** A Vector operand */
-	UPROPERTY(EditAnywhere, Category = Operands, meta = (DataflowInput, ShowOnlyInnerProperties));
+	UPROPERTY(EditAnywhere, Category = Operands, meta = (DataflowInput));
 	FDataflowVectorTypes A;
 
 	/** B Vector operand */
-	UPROPERTY(EditAnywhere, Category = Operands, meta = (DataflowInput, ShowOnlyInnerProperties));
+	UPROPERTY(EditAnywhere, Category = Operands, meta = (DataflowInput));
 	FDataflowVectorTypes B;
 
 	/** Resulting cross product of A and B : CrossProduct=B^A */
@@ -385,11 +385,11 @@ public:
 
 private:
 	/** Vector to scale */
-	UPROPERTY(EditAnywhere, Category = Operands, meta = (DataflowInput, ShowOnlyInnerProperties));
+	UPROPERTY(EditAnywhere, Category = Operands, meta = (DataflowInput));
 	FDataflowVectorTypes V;
 
 	/** Scale factor */
-	UPROPERTY(EditAnywhere, Category = Operands, meta = (DataflowInput, ShowOnlyInnerProperties));
+	UPROPERTY(EditAnywhere, Category = Operands, meta = (DataflowInput));
 	FDataflowNumericTypes Scale;
 
 	/** Scaled vector : Scaled=(V * Scale) */
@@ -411,7 +411,7 @@ public:
 
 private:
 	/** Vector to normalize */
-	UPROPERTY(EditAnywhere, Category = Operands, meta = (DataflowInput, ShowOnlyInnerProperties));
+	UPROPERTY(EditAnywhere, Category = Operands, meta = (DataflowInput));
 	FDataflowVectorTypes V;
 
 	/** Normalized vector : Normalized=V/|V| */
