Changed the dataflow 'show/hide inputs' menu to use the localized display names instead of the internal source code names.

diff --git a/Engine/Plugins/Experimental/Dataflow/Source/DataflowEditor/Private/Dataflow/DataflowSNode.cpp b/Engine/Plugins/Experimental/Dataflow/Source/DataflowEditor/Private/Dataflow/DataflowSNode.cpp
--- a/Engine/Plugins/Experimental/Dataflow/Source/DataflowEditor/Private/Dataflow/DataflowSNode.cpp
+++ b/Engine/Plugins/Experimental/Dataflow/Source/DataflowEditor/Private/Dataflow/DataflowSNode.cpp
@@ -509,7 +509,8 @@ FReply SDataflowEdNode::OnShowHideInputs()
 				{
 					if (Input->GetCanHidePin())
 					{
-						MenuBuilder.AddMenuEntry(FText::FromName(Input->GetName()), LOCTEXT("UnhidePinTooltip", "Show/Hide pin"), FSlateIcon(),
+						FText InputName = Input->GetProperty() ? Input->GetProperty()->GetDisplayNameText() : FText::FromName(Input->GetName());
+						MenuBuilder.AddMenuEntry(InputName, LOCTEXT("UnhidePinTooltip", "Show/Hide pin"), FSlateIcon(),
 							FUIAction(
 								FExecuteAction::CreateUObject(DataflowGraphNode, &UDataflowEdNode::ToggleHideInputPin, Input->GetName()),
 								FCanExecuteAction::CreateUObject(DataflowGraphNode, &UDataflowEdNode::CanToggleHideInputPin, Input->GetName()),
