Use the standard DataflowNode method to get pin names in the add/hide menu

diff --git a/Engine/Plugins/Experimental/Dataflow/Source/DataflowEditor/Private/Dataflow/DataflowSNode.cpp b/Engine/Plugins/Experimental/Dataflow/Source/DataflowEditor/Private/Dataflow/DataflowSNode.cpp
--- a/Engine/Plugins/Experimental/Dataflow/Source/DataflowEditor/Private/Dataflow/DataflowSNode.cpp
+++ b/Engine/Plugins/Experimental/Dataflow/Source/DataflowEditor/Private/Dataflow/DataflowSNode.cpp
@@ -509,7 +509,7 @@ FReply SDataflowEdNode::OnShowHideInputs()
 				{
 					if (Input->GetCanHidePin())
 					{
-						FText InputName = Input->GetProperty() ? Input->GetProperty()->GetDisplayNameText() : FText::FromName(Input->GetName());
+						FText InputName = DataflowNode->GetPinDisplayName(Input->GetName(), UE::Dataflow::FPin::EDirection::INPUT);
 						MenuBuilder.AddMenuEntry(InputName, LOCTEXT("UnhidePinTooltip", "Show/Hide pin"), FSlateIcon(),
 							FUIAction(
 								FExecuteAction::CreateUObject(DataflowGraphNode, &UDataflowEdNode::ToggleHideInputPin, Input->GetName()),
