Dataflow : show array icon for array pins

diff --git a/Engine/Plugins/Experimental/Dataflow/Source/DataflowEditor/Private/Dataflow/DataflowSNode.cpp b/Engine/Plugins/Experimental/Dataflow/Source/DataflowEditor/Private/Dataflow/DataflowSNode.cpp
--- a/Engine/Plugins/Experimental/Dataflow/Source/DataflowEditor/Private/Dataflow/DataflowSNode.cpp
+++ b/Engine/Plugins/Experimental/Dataflow/Source/DataflowEditor/Private/Dataflow/DataflowSNode.cpp
@@ -121,8 +121,6 @@ TSharedPtr<SGraphPin> SDataflowEdNode::CreatePinWidget(UEdGraphPin* Pin) const
 
 			if (Pin->Direction == EEdGraphPinDirection::EGPD_Output)
 			{
-
-
 				if (FDataflowOutput* Output = DataflowNode->FindOutput(Pin->GetFName()))
 				{
 					if (const TSharedPtr<UE::Dataflow::FContext> DataflowContext = DataflowInterface->GetDataflowContext())
@@ -143,6 +141,20 @@ TSharedPtr<SGraphPin> SDataflowEdNode::CreatePinWidget(UEdGraphPin* Pin) const
 							!DataflowContext->IsCacheEntryAfterTimestamp(Output->CacheKey(), DataflowNode->LastModifiedTimestamp);
 						PRAGMA_ENABLE_DEPRECATION_WARNINGS
 
+						if (FDataflowArrayTypePolicy::SupportsTypeStatic(Output->GetType()))
+						{
+							Pin->PinType.ContainerType = EPinContainerType::Array;
+						}
+					}
+				}
+			}
+			else if (Pin->Direction == EEdGraphPinDirection::EGPD_Input)
+			{
+				if (FDataflowInput* Input = DataflowNode->FindInput(Pin->GetFName()))
+				{
+					if (FDataflowArrayTypePolicy::SupportsTypeStatic(Input->GetType()))
+					{
+						Pin->PinType.ContainerType = EPinContainerType::Array;
 					}
 				}
 			}
