Fix datalow painting tool brush mode

diff --git a/Engine/Plugins/Experimental/Dataflow/Source/DataflowEditor/Private/DataflowEditorTools/DataflowEditorWeightMapPaintTool.cpp b/Engine/Plugins/Experimental/Dataflow/Source/DataflowEditor/Private/DataflowEditorTools/DataflowEditorWeightMapPaintTool.cpp
--- a/Engine/Plugins/Experimental/Dataflow/Source/DataflowEditor/Private/DataflowEditorTools/DataflowEditorWeightMapPaintTool.cpp
+++ b/Engine/Plugins/Experimental/Dataflow/Source/DataflowEditor/Private/DataflowEditorTools/DataflowEditorWeightMapPaintTool.cpp
@@ -769,7 +769,8 @@ bool UDataflowEditorWeightMapPaintTool::UpdateStampPosition(const FRay& WorldRay
 		break;
 	}
 
-	if (UseBrushOp->GetAlignStampToView())
+	const bool bVolumetric = (FilterProperties->BrushAreaMode == EMeshVertexPaintBrushAreaType::Volumetric);
+	if (UseBrushOp->GetAlignStampToView() || bVolumetric)
 	{
 		AlignBrushToView();
 	}
@@ -1392,7 +1393,7 @@ bool UDataflowEditorWeightMapPaintTool::UpdateBrushPosition(const FRay& WorldRay
 	{
 	case ESculptBrushOpTargetType::SculptMesh:
 	case ESculptBrushOpTargetType::TargetMesh:
-		bHit = UpdateBrushPositionOnSculptMesh(WorldRay, false);
+		bHit = UpdateBrushPositionOnSculptMesh(WorldRay, false);  
 		break;
 	case ESculptBrushOpTargetType::ActivePlane:
 		check(false);
@@ -1400,7 +1401,8 @@ bool UDataflowEditorWeightMapPaintTool::UpdateBrushPosition(const FRay& WorldRay
 		break;
 	}
 
-	if (bHit && UseBrushOp->GetAlignStampToView())
+	const bool bVolumetric = (FilterProperties->BrushAreaMode == EMeshVertexPaintBrushAreaType::Volumetric);
+	if (bHit && (UseBrushOp->GetAlignStampToView() || bVolumetric))
 	{
 		AlignBrushToView();
 	}
@@ -1494,12 +1496,6 @@ void UDataflowEditorWeightMapPaintTool::UpdateMaterialMode(EMeshEditingMaterialM
 	}
 }
 
-void UDataflowEditorWeightMapPaintTool::UpdateStampPendingState()
-{
-	if (InStroke() == false) return;
-	bIsStampPending = true;
-}
-
 void UDataflowEditorWeightMapPaintTool::OnTick(float DeltaTime)
 {
 	UMeshSculptToolBase::OnTick(DeltaTime);
@@ -1513,7 +1509,7 @@ void UDataflowEditorWeightMapPaintTool::OnTick(float DeltaTime)
 
 	check(!(bIsLasso && bIsGradient));
 
-	ConfigureIndicator(false);
+	ConfigureIndicator(FilterProperties->BrushAreaMode == EMeshVertexPaintBrushAreaType::Volumetric);
 	SetIndicatorVisibility(!bIsLasso && !bIsGradient);
 
 	if (bHavePendingAction)
diff --git a/Engine/Plugins/Experimental/Dataflow/Source/DataflowEditor/Public/DataflowEditorTools/DataflowEditorWeightMapPaintTool.h b/Engine/Plugins/Experimental/Dataflow/Source/DataflowEditor/Public/DataflowEditorTools/DataflowEditorWeightMapPaintTool.h
--- a/Engine/Plugins/Experimental/Dataflow/Source/DataflowEditor/Public/DataflowEditorTools/DataflowEditorWeightMapPaintTool.h
+++ b/Engine/Plugins/Experimental/Dataflow/Source/DataflowEditor/Public/DataflowEditorTools/DataflowEditorWeightMapPaintTool.h
@@ -279,9 +279,6 @@ public:
 	UE_API virtual void DrawHUD(FCanvas* Canvas, IToolsContextRenderAPI* RenderAPI) override;
 	UE_API virtual void Render(IToolsContextRenderAPI* RenderAPI) override;
 	UE_API virtual void UpdateMaterialMode(EMeshEditingMaterialModes MaterialMode) override;
-
-	UE_API virtual void UpdateStampPendingState() override;
-
 	UE_API virtual void OnTick(float DeltaTime) override;
 
 	virtual bool HasCancel() const override { return true; }
