Dataflow Editor - Fixed crash in Evaluation options menu, and updated the Clear Cache option to force a re-evaluation.

[RN] Dataflow Editor - Fixed crash in Evaluation options menu, and updated the Clear Cache option to force a re-evaluation.

diff --git a/Engine/Plugins/Experimental/Dataflow/Source/DataflowEditor/Private/Dataflow/DataflowEditorCommands.cpp b/Engine/Plugins/Experimental/Dataflow/Source/DataflowEditor/Private/Dataflow/DataflowEditorCommands.cpp
--- a/Engine/Plugins/Experimental/Dataflow/Source/DataflowEditor/Private/Dataflow/DataflowEditorCommands.cpp
+++ b/Engine/Plugins/Experimental/Dataflow/Source/DataflowEditor/Private/Dataflow/DataflowEditorCommands.cpp
@@ -54,10 +54,10 @@ void FDataflowEditorCommandsImpl::RegisterCommands()
 	TBaseCharacterFXEditorCommands::RegisterCommands();
 	
 	UI_COMMAND(EvaluateNode, "Evaluate", "Trigger an evaluation of the selected node.", EUserInterfaceActionType::Button, FInputChord());
-	UI_COMMAND(EvaluateGraph, "Evaluate Dataflow Graph", "Trigger an evaluation of the graph.", EUserInterfaceActionType::Button, FInputChord(EKeys::F5));
+	UI_COMMAND(EvaluateGraph, "Evaluate Dataflow Graph", "Trigger an evaluation of the graph Terminal node.", EUserInterfaceActionType::Button, FInputChord(EKeys::F5));
 	UI_COMMAND(EvaluateGraphAutomatic, "Automatic Graph Evaluation", "Set the evaluation mode of the graph to Automatic.", EUserInterfaceActionType::RadioButton, FInputChord());
 	UI_COMMAND(EvaluateGraphManual, "Manual Graph Evaluation", "Set the evaluation mode of the graph to Manual.", EUserInterfaceActionType::RadioButton, FInputChord());
-	UI_COMMAND(ClearGraphCache, "Clear Graph Cache", "Clear the graph Data cache.", EUserInterfaceActionType::Button, FInputChord());
+	UI_COMMAND(ClearGraphCache, "Force Graph Re-evaluation", "Force the entire graph to re-evaluate by deleting the evaluation cache and re-evaluating the Terminal node.", EUserInterfaceActionType::Button, FInputChord());
 	UI_COMMAND(TogglePerfData, "Performance Data", "Toggle the evaluation performance data for each node.", EUserInterfaceActionType::ToggleButton, FInputChord());
 	UI_COMMAND(ToggleAsyncEvaluation, "Asynchronous Evaluation (Experimental)", "Toggle asynchronous evaluation of the graph. This is an experimental feature.", EUserInterfaceActionType::ToggleButton, FInputChord());
 	UI_COMMAND(FreezeNodes, "FreezeNodes", "Freeze the evaluation of the selected nodes.", EUserInterfaceActionType::Button, FInputChord());
diff --git a/Engine/Plugins/Experimental/Dataflow/Source/DataflowEditor/Private/Dataflow/DataflowEditorToolkit.cpp b/Engine/Plugins/Experimental/Dataflow/Source/DataflowEditor/Private/Dataflow/DataflowEditorToolkit.cpp
--- a/Engine/Plugins/Experimental/Dataflow/Source/DataflowEditor/Private/Dataflow/DataflowEditorToolkit.cpp
+++ b/Engine/Plugins/Experimental/Dataflow/Source/DataflowEditor/Private/Dataflow/DataflowEditorToolkit.cpp
@@ -77,6 +77,7 @@
 #include "Dataflow/DataflowOutputLog.h"
 #include "Dataflow/DataflowPath.h"
 #include "Dataflow/DataflowSimulationNodes.h"
+#include "Toolkits/AssetEditorToolkitMenuContext.h"
 
 #define LOCTEXT_NAMESPACE "DataflowEditorToolkit"
 
@@ -108,6 +109,8 @@ const FName FDataflowEditorToolkit::OutputLogTabId(TEXT("DataflowEditor_OutputLo
 
 namespace UE::Dataflow::Private
 {
+	static const FName DataflowEditorToolkitName("DataflowEditor");
+
 	/** Create a debug draw component to visualize per node informations */
 	TUniquePtr<IDataflowDebugDrawInterface> CreateDebugDrawComponent(
 		FDataflowPreviewSceneBase* DataflowScene,
@@ -216,10 +219,25 @@ namespace UE::Dataflow::Private
 
 		DebugDrawOverlay = DebugDrawObject->GetOverlayText();
 	}
+
+	static TSharedPtr<FDataflowEditorToolkit> GetAssetEditorToolkitFromMenuContext(UAssetEditorToolkitMenuContext* Context)
+	{
+		if (Context)
+		{
+			if (const TSharedPtr<FAssetEditorToolkit> Toolkit = Context->Toolkit.Pin())
+			{
+				if (Toolkit->GetToolkitFName() == DataflowEditorToolkitName)  // Note: This will not detect subclasses of FDataflowEditorToolkit
+				{
+					return StaticCastSharedPtr<FDataflowEditorToolkit>(Toolkit);
+				}
+			}
+		}
+		return nullptr;
+	}
 }
 
 FDataflowEditorToolkit::FDataflowEditorToolkit(UAssetEditor* InOwningAssetEditor)
-	: FBaseCharacterFXEditorToolkit(InOwningAssetEditor, FName("DataflowEditor"))
+	: FBaseCharacterFXEditorToolkit(InOwningAssetEditor, UE::Dataflow::Private::DataflowEditorToolkitName)
 	, DataflowEditor(Cast< UDataflowEditor>(InOwningAssetEditor))
 	, bDataflowEnableGraphEval(true)
 {
@@ -902,9 +920,9 @@ void FDataflowEditorToolkit::AddEvaluationWidget(FToolMenuSection& InSection)
 		FCanExecuteAction(),
 		FIsActionChecked::CreateLambda([this]() { return EvaluationMode == EDataflowEditorEvaluationMode::Manual; })
 	);
-	ToolkitCommands->MapAction(FDataflowEditorCommands::Get().ClearGraphCache,
-		FExecuteAction::CreateSP(this, &FDataflowEditorToolkit::ClearGraphCache),
-		FCanExecuteAction::CreateSP(this, &FDataflowEditorToolkit::CanClearGraphCache)
+	ToolkitCommands->MapAction(FDataflowEditorCommands::Get().ClearGraphCache,  //  TODO: Rename ClearGraphCache to ForceGraphEvaluation in 5.7
+		FExecuteAction::CreateLambda([this]() { ClearGraphCache(); EvaluateGraph(); }),  // Note: This action used to only clear the cache, but has been changed to force a re evaluation
+		FCanExecuteAction()
 	);
 	ToolkitCommands->MapAction(FDataflowEditorCommands::Get().TogglePerfData,
 		FExecuteAction::CreateSP(this, &FDataflowEditorToolkit::TogglePerfData),
@@ -917,24 +935,38 @@ void FDataflowEditorToolkit::AddEvaluationWidget(FToolMenuSection& InSection)
 		FIsActionChecked::CreateSP(this, &FDataflowEditorToolkit::IsAsyncEvaluationEnabled)
 	);
 	
-	InSection.AddEntry(FToolMenuEntry::InitToolBarButton(
-		FDataflowEditorCommands::Get().EvaluateGraph,
-		TAttribute<FText>(),
-		TAttribute<FText>(),
-		TAttribute<FSlateIcon>(this, &FDataflowEditorToolkit::GetEvaluationStatusImage)
-	));
-
-	FToolMenuEntry EvaluationOptions = FToolMenuEntry::InitComboButton(
-		"DataflowEvaluationOptions",
-		FUIAction(),
-		FOnGetContent::CreateSP(this, &FDataflowEditorToolkit::GenerateEvaluationOptionsMenu),
-		LOCTEXT("DataflowEvaluationOptions", "Options"),
-		LOCTEXT("DataflowEvaluationOptions_ToolbarTooltip", "Options to customize how Dataflow evaluate"),
-		TAttribute<FSlateIcon>(),
-		true
-	);
-	EvaluationOptions.StyleNameOverride = "SlimToolBar";
-	InSection.AddEntry(EvaluationOptions);
+	InSection.AddDynamicEntry("DataflowEvaluateGraphDynamic", FNewToolMenuSectionDelegate::CreateLambda([](FToolMenuSection& DynamicSection)
+		{
+			using namespace UE::Dataflow::Private;
+			if (const TSharedPtr<FDataflowEditorToolkit> Toolkit = GetAssetEditorToolkitFromMenuContext(DynamicSection.FindContext<UAssetEditorToolkitMenuContext>()))
+			{
+				DynamicSection.AddEntry(FToolMenuEntry::InitToolBarButton(
+						FDataflowEditorCommands::Get().EvaluateGraph,
+						TAttribute<FText>(),
+						TAttribute<FText>(),
+						TAttribute<FSlateIcon>(Toolkit.ToSharedRef(), &FDataflowEditorToolkit::GetEvaluationStatusImage)
+					));
+			}
+		}));
+
+	InSection.AddDynamicEntry("DataflowEvaluationOptionsDynamic", FNewToolMenuSectionDelegate::CreateLambda([](FToolMenuSection& DynamicSection)
+		{
+			using namespace UE::Dataflow::Private;
+			if (const TSharedPtr<FDataflowEditorToolkit> Toolkit = GetAssetEditorToolkitFromMenuContext(DynamicSection.FindContext<UAssetEditorToolkitMenuContext>()))
+			{
+				FToolMenuEntry EvaluationOptions = FToolMenuEntry::InitComboButton(
+					"DataflowEvaluationOptions",
+					FUIAction(),
+					FOnGetContent::CreateSP(Toolkit.ToSharedRef(), &FDataflowEditorToolkit::GenerateEvaluationOptionsMenu),
+					LOCTEXT("DataflowEvaluationOptions", "Options"),
+					LOCTEXT("DataflowEvaluationOptions_ToolbarTooltip", "Options to customize how Dataflow evaluate"),
+					TAttribute<FSlateIcon>(),
+					true
+				);
+				EvaluationOptions.StyleNameOverride = "SlimToolBar";
+				DynamicSection.AddEntry(EvaluationOptions);
+			}
+		}));
 
 	// load options
 	if (UDataflowEditorOptions* const Options = UDataflowEditorOptions::StaticClass()->GetDefaultObject<UDataflowEditorOptions>())
@@ -950,10 +982,9 @@ TSharedRef<SWidget> FDataflowEditorToolkit::GenerateEvaluationOptionsMenu()
 	MenuBuilder.AddMenuEntry(FDataflowEditorCommands::Get().EvaluateGraphAutomatic);
 	MenuBuilder.AddMenuEntry(FDataflowEditorCommands::Get().EvaluateGraphManual);
 	MenuBuilder.AddSeparator();
-	MenuBuilder.AddMenuEntry(FDataflowEditorCommands::Get().TogglePerfData);
-	MenuBuilder.AddSeparator();
 	MenuBuilder.AddMenuEntry(FDataflowEditorCommands::Get().ClearGraphCache);
 	MenuBuilder.AddSeparator();
+	MenuBuilder.AddMenuEntry(FDataflowEditorCommands::Get().TogglePerfData);
 	MenuBuilder.AddMenuEntry(FDataflowEditorCommands::Get().ToggleAsyncEvaluation);
 	MenuBuilder.EndSection();
 	return MenuBuilder.MakeWidget();
