Dataflow Editor - Fixed random crashes when closing an editor window.

[RN] Dataflow Editor - Fixed random crashes when closing an editor window.

diff --git a/Engine/Plugins/Experimental/Dataflow/Source/DataflowEditor/Private/Dataflow/DataflowEditorPreviewSceneBase.cpp b/Engine/Plugins/Experimental/Dataflow/Source/DataflowEditor/Private/Dataflow/DataflowEditorPreviewSceneBase.cpp
--- a/Engine/Plugins/Experimental/Dataflow/Source/DataflowEditor/Private/Dataflow/DataflowEditorPreviewSceneBase.cpp
+++ b/Engine/Plugins/Experimental/Dataflow/Source/DataflowEditor/Private/Dataflow/DataflowEditorPreviewSceneBase.cpp
@@ -67,6 +67,13 @@ FDataflowPreviewSceneBase::FDataflowPreviewSceneBase(FPreviewScene::Construction
 
 FDataflowPreviewSceneBase::~FDataflowPreviewSceneBase()
 {
+	// There is a crash in 5.6, but the fix requires the preview scene to become a weakptr and this isn't hotfixable for 5.6.1.
+	// Instead sending the following guid to notify the destruction of the preview scene, so the pointer can be safely cleared up.
+	// TODO: Remove this hack for 5.7.
+	static const FGuid TerminationGuid(0x8A4F3813, 0x5D7E4932, 0x8F9F92AF, 0x74341FAA);
+	static const FBox TerminationBox(FVector((double)TerminationGuid.A, (double)TerminationGuid.B, 0.0), FVector((double)TerminationGuid.C, (double)TerminationGuid.D, 0.0));
+	FocusRequestDelegate.Broadcast(TerminationBox);
+
 	UnregisterSceneElements();
 	ModifySceneElements().Reset();
 }
diff --git a/Engine/Plugins/Experimental/Dataflow/Source/DataflowEditor/Private/Dataflow/DataflowEditorViewportClientBase.cpp b/Engine/Plugins/Experimental/Dataflow/Source/DataflowEditor/Private/Dataflow/DataflowEditorViewportClientBase.cpp
--- a/Engine/Plugins/Experimental/Dataflow/Source/DataflowEditor/Private/Dataflow/DataflowEditorViewportClientBase.cpp
+++ b/Engine/Plugins/Experimental/Dataflow/Source/DataflowEditor/Private/Dataflow/DataflowEditorViewportClientBase.cpp
@@ -75,27 +75,29 @@ void FDataflowEditorViewportClientBase::RegisterDelegates()
 	UAssetViewerSettings::Get()->OnAssetViewerSettingsChanged().Remove(OnAssetViewerSettingsChangedDelegateHandle);
 
 	auto SyncEngineShowFlags = [this]()
-	{
-		checkf(DataflowPreviewScene, TEXT("Expected valid DataflowPreviewScene pointer"));
-		const int32 CurrentProfileIndex = DataflowPreviewScene->GetCurrentProfileIndex();
+		{
+			if (ensureMsgf(DataflowPreviewScene, TEXT("Expected valid DataflowPreviewScene pointer")))
+			{
+				const int32 CurrentProfileIndex = DataflowPreviewScene->GetCurrentProfileIndex();
 
-		const UAssetViewerSettings* const DefaultSettings = UAssetViewerSettings::Get();
+				const UAssetViewerSettings* const DefaultSettings = UAssetViewerSettings::Get();
 
-		checkf(DefaultSettings && DefaultSettings->Profiles.IsValidIndex(CurrentProfileIndex), TEXT("Invalid default settings pointer or current profile index"));
+				checkf(DefaultSettings && DefaultSettings->Profiles.IsValidIndex(CurrentProfileIndex), TEXT("Invalid default settings pointer or current profile index"));
 
-		const FPreviewSceneProfile& Profile = DefaultSettings->Profiles[CurrentProfileIndex];
-		EngineShowFlags.Grid = Profile.bShowGrid;
-		DrawHelper.bDrawGrid = Profile.bShowGrid;
+				const FPreviewSceneProfile& Profile = DefaultSettings->Profiles[CurrentProfileIndex];
+				EngineShowFlags.Grid = Profile.bShowGrid;
+				DrawHelper.bDrawGrid = Profile.bShowGrid;
 
-		if (Profile.bPostProcessingEnabled)
-		{
-			EngineShowFlags.EnableAdvancedFeatures();
-		}
-		else
-		{
-			EngineShowFlags.DisableAdvancedFeatures();
-		}
-	};
+				if (Profile.bPostProcessingEnabled)
+				{
+					EngineShowFlags.EnableAdvancedFeatures();
+				}
+				else
+				{
+					EngineShowFlags.DisableAdvancedFeatures();
+				}
+			}
+		};
 
 	// Run the lambda once to get things in sync
 	SyncEngineShowFlags();
@@ -128,6 +130,16 @@ void FDataflowEditorViewportClientBase::DeregisterDelegates()
 
 void FDataflowEditorViewportClientBase::HandleFocusRequest(const FBox& BoundingBox)
 {
+	// There is a crash in 5.6, but the fix requires the preview scene to become a weakptr and this isn't hotfixable for 5.6.1.
+	// Instead sending the following guid to notify the destruction of the preview scene, so the pointer can be safely cleared up.
+	// TODO: Remove this hack for 5.7.
+	static const FGuid TerminationGuid(0x8A4F3813, 0x5D7E4932, 0x8F9F92AF, 0x74341FAA);
+	static const FBox TerminationBox(FVector((double)TerminationGuid.A, (double)TerminationGuid.B, 0.0), FVector((double)TerminationGuid.C, (double)TerminationGuid.D, 0.0));
+	if (BoundingBox == TerminationBox)
+	{
+		DataflowPreviewScene = nullptr;
+		return;
+	}
 	FocusViewportOnBox(BoundingBox);
 }
 
