Dataflow : Add support for auto convert nodes

diff --git a/Engine/Plugins/Experimental/Dataflow/Source/DataflowEditor/Private/Dataflow/DataflowSchema.cpp b/Engine/Plugins/Experimental/Dataflow/Source/DataflowEditor/Private/Dataflow/DataflowSchema.cpp
--- a/Engine/Plugins/Experimental/Dataflow/Source/DataflowEditor/Private/Dataflow/DataflowSchema.cpp
+++ b/Engine/Plugins/Experimental/Dataflow/Source/DataflowEditor/Private/Dataflow/DataflowSchema.cpp
@@ -17,6 +17,7 @@
 #include "ToolMenuSection.h"
 #include "ToolMenu.h"
 #include "GraphEditorActions.h"
+#include "Dataflow/DataflowAnyTypeRegistry.h"
 #include "Dataflow/DataflowSettings.h"
 #include "Dataflow/DataflowObject.h"
 #include "Dataflow/DataflowSubGraph.h"
@@ -320,33 +321,73 @@ bool HasLoopIfConnected(const UEdGraphNode* FromNode, const UEdGraphNode* ToNode
 	return false;
 }
 
-bool UDataflowSchema::CanConnectPins(const UEdGraphPin& OutputPin, const UEdGraphPin& InputPin)
+TSharedPtr<UE::Dataflow::FGraph> UDataflowSchema::GetDataflowGraphFromPin(const UEdGraphPin& Pin)
+{
+	if (UDataflowEdNode* EdNode = Cast<UDataflowEdNode>(Pin.GetOwningNode()))
+	{
+		return EdNode->DataflowGraph;
+	}
+	return {};
+}
+
+TSharedPtr<const FDataflowNode> UDataflowSchema::GetDataflowNodeFromPin(const UEdGraphPin& Pin)
 {
-	if (ensure(
-		OutputPin.Direction == EEdGraphPinDirection::EGPD_Output
-		&& InputPin.Direction == EEdGraphPinDirection::EGPD_Input
-		))
+	if (UDataflowEdNode* EdNode = Cast<UDataflowEdNode>(Pin.GetOwningNode()))
 	{
-		UDataflowEdNode* EdOutputNode = Cast<UDataflowEdNode>(OutputPin.GetOwningNode());
-		UDataflowEdNode* EdInputNode = Cast<UDataflowEdNode>(InputPin.GetOwningNode());
-		if (EdOutputNode && EdInputNode)
+		return EdNode->GetDataflowNode();
+	}
+	return {};
+}
+
+const FDataflowInput* UDataflowSchema::GetDataflowInputFromPin(const UEdGraphPin& Pin)
+{
+	if (TSharedPtr<const FDataflowNode> DataflowNode = GetDataflowNodeFromPin(Pin))
+	{
+		if (Pin.Direction == EEdGraphPinDirection::EGPD_Input)
 		{
-			if (EdOutputNode->DataflowGraph && EdInputNode->DataflowGraph)
-			{
-				TSharedPtr<const FDataflowNode> OutputNode = EdOutputNode->GetDataflowNode();
-				TSharedPtr<const FDataflowNode> InputNode = EdInputNode->GetDataflowNode();
-				if (OutputNode && InputNode)
-				{
-					const FDataflowOutput* Output = OutputNode->FindOutput(OutputPin.PinName);
-					const FDataflowInput* Input = InputNode->FindInput(InputPin.PinName);
-					if (Output && Input)
-					{
-						return EdOutputNode->DataflowGraph->CanConnect(*Output, *Input);
-					}
-				}
-			}
+			return DataflowNode->FindInput(Pin.PinName);
+		}
+	}
+	return nullptr;
+}
+
+const FDataflowOutput* UDataflowSchema::GetDataflowOutputFromPin(const UEdGraphPin& Pin)
+{
+	if (TSharedPtr<const FDataflowNode> DataflowNode = GetDataflowNodeFromPin(Pin))
+	{
+		if (Pin.Direction == EEdGraphPinDirection::EGPD_Output)
+		{
+			return DataflowNode->FindOutput(Pin.PinName);
 		}
 	}
+	return nullptr;
+}
+
+bool UDataflowSchema::CanConnectPins(const UEdGraphPin& OutputPin, const UEdGraphPin& InputPin)
+{
+	ensure(OutputPin.Direction == EEdGraphPinDirection::EGPD_Output && InputPin.Direction == EEdGraphPinDirection::EGPD_Input);
+
+	const TSharedPtr<UE::Dataflow::FGraph> Graph = GetDataflowGraphFromPin(OutputPin);
+	const FDataflowOutput* Output = GetDataflowOutputFromPin(OutputPin);
+	const FDataflowInput* Input  = GetDataflowInputFromPin(InputPin);
+	if (Graph && Output && Input)
+	{
+		return Graph->CanConnect(*Output, *Input);
+	}
+	return false;
+}
+
+bool UDataflowSchema::CanConvertPins(const UEdGraphPin& OutputPin, const UEdGraphPin& InputPin)
+{
+	ensure(OutputPin.Direction == EEdGraphPinDirection::EGPD_Output && InputPin.Direction == EEdGraphPinDirection::EGPD_Input);
+	
+	const FDataflowOutput* Output = GetDataflowOutputFromPin(OutputPin);
+	const FDataflowInput* Input = GetDataflowInputFromPin(InputPin);
+	if (Output && Input)
+	{
+		const FName AutoConvertNodeName = UE::Dataflow::FAnyTypesRegistry::GetAutoConvertNodeTypeStatic(Output->GetType(), Input->GetType());
+		return AutoConvertNodeName != NAME_None;
+	}
 	return false;
 }
 
@@ -372,26 +413,28 @@ const FPinConnectionResponse UDataflowSchema::CanCreateConnection(const UEdGraph
 
 			if (EdNodeA && EdNodeB && (EdNodeA != EdNodeB))
 			{
-				if (CanConnectPins(*PinA, *PinB))
+				if (HasLoopIfConnected(EdNodeA, EdNodeB))
 				{
-					// cycle checking on connect
-					if (!HasLoopIfConnected(EdNodeA, EdNodeB))
-					{
-						if (PinB->LinkedTo.Num())
-						{
-							return (bSwapped) ?
-								FPinConnectionResponse(CONNECT_RESPONSE_BREAK_OTHERS_A, LOCTEXT("PinSteal", "Disconnect existing input and connect new input."))
-								:
-								FPinConnectionResponse(CONNECT_RESPONSE_BREAK_OTHERS_B, LOCTEXT("PinSteal", "Disconnect existing input and connect new input."));
+					return FPinConnectionResponse(CONNECT_RESPONSE_DISALLOW, LOCTEXT("PinError_Loop", "Graph Cycle"));
+				}
 
-						}
+				const bool bCompatibleTypes = CanConnectPins(*PinA, *PinB);
+				const bool bCanConvert = CanConvertPins(*PinA, *PinB);
 
-						return FPinConnectionResponse(CONNECT_RESPONSE_MAKE, LOCTEXT("PinConnect", "Connect input to output."));
-					}
-					else
+				if (bCanConvert && !bCompatibleTypes)
+				{
+					return FPinConnectionResponse(CONNECT_RESPONSE_MAKE_WITH_CONVERSION_NODE, LOCTEXT("PinConnectWithConversion", "Connect input to output using a conversion node."));
+				}
+				if (bCompatibleTypes)
+				{
+					if (PinB->LinkedTo.Num())
 					{
-						return FPinConnectionResponse(CONNECT_RESPONSE_DISALLOW, LOCTEXT("PinError_Loop", "Graph Cycle"));
+						return (bSwapped) ?
+							FPinConnectionResponse(CONNECT_RESPONSE_BREAK_OTHERS_A, LOCTEXT("PinSteal", "Disconnect existing input and connect new input."))
+							:
+							FPinConnectionResponse(CONNECT_RESPONSE_BREAK_OTHERS_B, LOCTEXT("PinSteal", "Disconnect existing input and connect new input."));
 					}
+					return FPinConnectionResponse(CONNECT_RESPONSE_MAKE, LOCTEXT("PinConnect", "Connect input to output."));
 				}
 				else
 				{
@@ -410,6 +453,91 @@ const FPinConnectionResponse UDataflowSchema::CanCreateConnection(const UEdGraph
 	return FPinConnectionResponse(CONNECT_RESPONSE_DISALLOW, NoConnectionResponse[FMath::RandRange(0, NoConnectionResponse.Num()-1)]);
 }
 
+bool UDataflowSchema::CreateAutomaticConversionNodeAndConnections(UEdGraphPin* InSourcePin, UEdGraphPin* InTargetPin) const
+{
+	if (!InSourcePin || !InTargetPin)
+	{
+		return false;
+	}
+
+	UEdGraphPin* FromPin = InSourcePin;
+	UEdGraphPin* ToPin = InTargetPin;
+	if (FromPin->Direction == EEdGraphPinDirection::EGPD_Input)
+	{
+		Swap(FromPin, ToPin);
+	}
+	if (FromPin->Direction != EEdGraphPinDirection::EGPD_Output && ToPin->Direction != EEdGraphPinDirection::EGPD_Input)
+	{
+		return false;
+	}
+
+	UDataflowEdNode* const SourceEdNode = Cast<UDataflowEdNode>(FromPin->GetOwningNode());
+	UDataflowEdNode* const TargetEdNode = Cast<UDataflowEdNode>(ToPin->GetOwningNode());
+
+	if (!SourceEdNode || !TargetEdNode)
+	{
+		return false;
+	}
+
+	const FDataflowOutput* Output = GetDataflowOutputFromPin(*FromPin);
+	const FDataflowInput* Input = GetDataflowInputFromPin(*ToPin);
+	if (!Output || !Input)
+	{
+		return false;
+	}
+
+	const FName AutoConvertNodeName = UE::Dataflow::FAnyTypesRegistry::GetAutoConvertNodeTypeStatic(Output->GetType(), Input->GetType());
+	if (AutoConvertNodeName == NAME_None)
+	{
+		return false;
+	}
+
+	UEdGraph* EdGraph = SourceEdNode->GetGraph();
+
+
+	TSharedPtr<FAssetSchemaAction_Dataflow_CreateNode_DataflowEdNode> NewAutoConvertNodeAction
+		= FAssetSchemaAction_Dataflow_CreateNode_DataflowEdNode::CreateAction(EdGraph, AutoConvertNodeName);
+	if (NewAutoConvertNodeAction)
+	{
+		const FVector2f AutoConvertNodeLocation = (SourceEdNode->GetPosition() + TargetEdNode->GetPosition()) * 0.5f;
+		if (UEdGraphNode* AutoConvertEdNode = NewAutoConvertNodeAction->PerformAction(EdGraph, nullptr, AutoConvertNodeLocation, false))
+		{
+			UEdGraphPin* InputPin = nullptr;
+			UEdGraphPin* OutputPin = nullptr;
+			for (UEdGraphPin* Pin : AutoConvertEdNode->GetAllPins())
+			{
+				if (InputPin == nullptr && Pin && Pin->Direction == EEdGraphPinDirection::EGPD_Input)
+				{
+					if (CanConnectPins(*FromPin, *Pin))
+					{
+						InputPin = Pin;
+					}
+				}
+				if (OutputPin == nullptr && Pin && Pin->Direction == EEdGraphPinDirection::EGPD_Output)
+				{
+					if (CanConnectPins(*Pin, *ToPin))
+					{
+						OutputPin = Pin;
+					}
+				}
+				if (InputPin && OutputPin)
+				{
+					break;
+				}
+			}
+			if (InputPin && OutputPin)
+			{
+				ToPin->BreakAllPinLinks(true);
+				EdGraph->GetSchema()->TryCreateConnection(FromPin, InputPin);
+				EdGraph->GetSchema()->TryCreateConnection(OutputPin, ToPin);
+				return true;
+			}
+		}
+	}
+
+	return false;
+}
+
 UE::Dataflow::FPin::EDirection UDataflowSchema::GetDirectionFromPinDirection(EEdGraphPinDirection InPinDirection)
 {
 	switch (InPinDirection)
diff --git a/Engine/Plugins/Experimental/Dataflow/Source/DataflowEditor/Public/Dataflow/DataflowSchema.h b/Engine/Plugins/Experimental/Dataflow/Source/DataflowEditor/Public/Dataflow/DataflowSchema.h
--- a/Engine/Plugins/Experimental/Dataflow/Source/DataflowEditor/Public/Dataflow/DataflowSchema.h
+++ b/Engine/Plugins/Experimental/Dataflow/Source/DataflowEditor/Public/Dataflow/DataflowSchema.h
@@ -34,6 +34,7 @@ public:
 	UE_API virtual void OnPinConnectionDoubleCicked(UEdGraphPin* PinA, UEdGraphPin* PinB, const FVector2f& GraphPosition) const override;
 	UE_API virtual void BreakPinLinks(UEdGraphPin& TargetPin, bool bSendsNodeNotifcation) const override;
 	UE_API virtual bool TryCreateConnection(UEdGraphPin* PinA, UEdGraphPin* PinB) const override;
+	UE_API virtual bool CreateAutomaticConversionNodeAndConnections(UEdGraphPin* InSourcePin, UEdGraphPin* InTargetPin) const override;
 
 	UE_API virtual void GetAssetsGraphHoverMessage(const TArray<FAssetData>& Assets, const UEdGraph* HoverGraph, FString& OutTooltipText, bool& OutOkIcon) const override;
 	UE_API virtual void DroppedAssetsOnGraph(const TArray<FAssetData>& Assets, const FVector2f& GraphPosition, UEdGraph* Graph) const override;
@@ -54,10 +55,16 @@ public:
 private:
 	static UE_API bool CanPinBeConnectedToNode(const UEdGraphPin* Pin, const UE::Dataflow::FFactoryParameters& NodeParameters);
 	static UE_API bool CanConnectPins(const UEdGraphPin& OutputPin, const UEdGraphPin& InputPin);
+	static UE_API bool CanConvertPins(const UEdGraphPin& OutputPin, const UEdGraphPin& InputPin);
 	static UE_API bool IsCategorySupported(FName NodeCategory, FName AssetType);
 	static UE_API FName GetEditedAssetType();
 	static UE_API UE::Dataflow::FPin::EDirection GetDirectionFromPinDirection(EEdGraphPinDirection InPinDirection);
 
+	static TSharedPtr<UE::Dataflow::FGraph> GetDataflowGraphFromPin(const UEdGraphPin& Pin);
+	static TSharedPtr<const FDataflowNode> GetDataflowNodeFromPin(const UEdGraphPin& Pin);
+	static const FDataflowInput* GetDataflowInputFromPin(const UEdGraphPin& Pin);
+	static const FDataflowOutput* GetDataflowOutputFromPin(const UEdGraphPin& Pin);
+
 	// used by SupportsDropPinOnNode to know which original connection is being dropped
 	mutable UEdGraphPin* PinBeingDropped = nullptr;
 };
diff --git a/Engine/Plugins/Experimental/GeometryCollectionPlugin/Source/GeometryCollectionNodes/Private/Dataflow/GeometryCollectionMeshNodes.cpp b/Engine/Plugins/Experimental/GeometryCollectionPlugin/Source/GeometryCollectionNodes/Private/Dataflow/GeometryCollectionMeshNodes.cpp
--- a/Engine/Plugins/Experimental/GeometryCollectionPlugin/Source/GeometryCollectionNodes/Private/Dataflow/GeometryCollectionMeshNodes.cpp
+++ b/Engine/Plugins/Experimental/GeometryCollectionPlugin/Source/GeometryCollectionNodes/Private/Dataflow/GeometryCollectionMeshNodes.cpp
@@ -34,6 +34,7 @@
 #include "FractureEngineClustering.h"
 #include "FractureEngineSelection.h"
 #include "FractureEngineUtility.h"
+#include "Dataflow/DataflowAnyTypeRegistry.h"
 #include "DynamicMesh/DynamicMesh3.h"
 #include "DynamicMesh/MeshTransforms.h"
 #include "UDynamicMesh.h"
@@ -64,6 +65,12 @@ namespace UE::Dataflow
 		DATAFLOW_NODE_REGISTER_CREATION_FACTORY(FApplyMeshProcessorToGeometryCollectionDataflowNode);
 		DATAFLOW_NODE_REGISTER_CREATION_FACTORY(FCollectionSelectionToMeshesDataflowNode);
 		DATAFLOW_NODE_REGISTER_CREATION_FACTORY(FAppendMeshesToCollectionDataflowNode);
+
+		UE_DATAFLOW_REGISTER_AUTOCONVERT(TObjectPtr<UStaticMesh>, TObjectPtr<UDynamicMesh>, FStaticMeshToMeshDataflowNode);
+		UE_DATAFLOW_REGISTER_AUTOCONVERT(FManagedArrayCollection, TObjectPtr<UDynamicMesh>, FCollectionToMeshDataflowNode);
+		UE_DATAFLOW_REGISTER_AUTOCONVERT(TObjectPtr<UDynamicMesh>, FManagedArrayCollection, FMeshToCollectionDataflowNode);
+		UE_DATAFLOW_REGISTER_AUTOCONVERT(TObjectPtr<UDynamicMesh>, FBox, FGetMeshBoundingBoxDataflowNode);
+		UE_DATAFLOW_REGISTER_AUTOCONVERT(FBox, TObjectPtr<UDynamicMesh>, FBoxToMeshDataflowNode);
 	}
 }
 
diff --git a/Engine/Source/Runtime/Experimental/Dataflow/Core/Private/Dataflow/DataflowAnyType.cpp b/Engine/Source/Runtime/Experimental/Dataflow/Core/Private/Dataflow/DataflowAnyType.cpp
--- a/Engine/Source/Runtime/Experimental/Dataflow/Core/Private/Dataflow/DataflowAnyType.cpp
+++ b/Engine/Source/Runtime/Experimental/Dataflow/Core/Private/Dataflow/DataflowAnyType.cpp
@@ -2,6 +2,7 @@
 
 #include "Dataflow/DataflowAnyType.h"
 #include "Dataflow/DataflowAnyTypeRegistry.h"
+#include "Dataflow/DataflowConvertNodes.h"
 
 namespace UE::Dataflow
 {
@@ -26,6 +27,18 @@ namespace UE::Dataflow
 		UE_DATAFLOW_REGISTER_ANYTYPE(FDataflowStringArrayTypes);
 		UE_DATAFLOW_REGISTER_ANYTYPE(FDataflowBoolArrayTypes);
 		UE_DATAFLOW_REGISTER_ANYTYPE(FDataflowTransformArrayTypes);
+
+		UE_DATAFLOW_REGISTER_AUTOCONVERT(FDataflowNumericTypes, FDataflowNumericTypes, FConvertNumericTypesDataflowNode);
+		UE_DATAFLOW_REGISTER_AUTOCONVERT(FDataflowVectorTypes, FDataflowVectorTypes, FConvertVectorTypesDataflowNode);
+		UE_DATAFLOW_REGISTER_AUTOCONVERT(FDataflowStringTypes, FDataflowStringTypes, FConvertStringTypesDataflowNode);
+		UE_DATAFLOW_REGISTER_AUTOCONVERT(FDataflowBoolTypes, FDataflowBoolTypes, FConvertBoolTypesDataflowNode);
+		UE_DATAFLOW_REGISTER_AUTOCONVERT(FDataflowTransformTypes, FDataflowTransformTypes, FConvertTransformTypesDataflowNode);
+		UE_DATAFLOW_REGISTER_AUTOCONVERT(FDataflowRotationTypes, FDataflowRotationTypes, FConvertRotationDataflowNode);
+		UE_DATAFLOW_REGISTER_AUTOCONVERT(FDataflowVectorArrayTypes, FDataflowVectorArrayTypes, FConvertVectorArrayTypesDataflowNode);
+		UE_DATAFLOW_REGISTER_AUTOCONVERT(FDataflowNumericArrayTypes, FDataflowNumericArrayTypes, FConvertNumericArrayTypesDataflowNode);
+		UE_DATAFLOW_REGISTER_AUTOCONVERT(FDataflowStringArrayTypes, FDataflowStringArrayTypes, FConvertStringArrayTypesDataflowNode);
+		UE_DATAFLOW_REGISTER_AUTOCONVERT(FDataflowBoolArrayTypes, FDataflowBoolArrayTypes, FConvertBoolArrayTypesDataflowNode);
+		UE_DATAFLOW_REGISTER_AUTOCONVERT(FDataflowTransformArrayTypes, FDataflowTransformArrayTypes, FConvertTransformArrayTypesDataflowNode);
 	}
 
 	bool AreTypesCompatible(FName TypeA, FName TypeB)
diff --git a/Engine/Source/Runtime/Experimental/Dataflow/Core/Private/Dataflow/DataflowAnyTypeRegistry.cpp b/Engine/Source/Runtime/Experimental/Dataflow/Core/Private/Dataflow/DataflowAnyTypeRegistry.cpp
--- a/Engine/Source/Runtime/Experimental/Dataflow/Core/Private/Dataflow/DataflowAnyTypeRegistry.cpp
+++ b/Engine/Source/Runtime/Experimental/Dataflow/Core/Private/Dataflow/DataflowAnyTypeRegistry.cpp
@@ -14,6 +14,16 @@ namespace UE::Dataflow
 		return GetInstance().GetStorageType(Type);
 	}
 
+	void FAnyTypesRegistry::RegisterAutoConvertNodeTypeStatic(FName FromType, FName ToType, FName AutoConvertNode, FTypeFilterFunction TypeFilterFunction)
+	{
+		GetInstance().RegisterAutoConvertNodeType(FromType, ToType, AutoConvertNode, TypeFilterFunction);
+	}
+
+	FName FAnyTypesRegistry::GetAutoConvertNodeTypeStatic(FName FromType, FName ToType)
+	{
+		return GetInstance().GetAutoConvertNodeType(FromType, ToType);
+	}
+
 	bool FAnyTypesRegistry::IsAnyTypeStatic(FName Type)
 	{
 		return GetInstance().IsAnyType(Type);
@@ -57,6 +67,34 @@ namespace UE::Dataflow
 		return NAME_None;
 	}
 
+	void FAnyTypesRegistry::RegisterAutoConvertNodeType(FName FromType, FName ToType, FName AutoConvertNode, FTypeFilterFunction TypeFilterFunction)
+	{
+		AutoConvertNodes.Add(
+			FAutoConvertNodeInfo
+			{
+				.FromType = FromType,
+				.ToType = ToType,
+				.AutoConvertNode = AutoConvertNode,
+				.TypeFilterFunction = TypeFilterFunction,
+			});
+	}
+
+	FName FAnyTypesRegistry::GetAutoConvertNodeType(FName FromType, FName ToType) const
+	{
+		// check the node that matches both the from and to types
+		for (const FAutoConvertNodeInfo& Info : AutoConvertNodes)
+		{
+			if (!Info.TypeFilterFunction.IsSet() || Info.TypeFilterFunction(FromType, ToType))
+			{
+				if (AreTypesCompatible(FromType, Info.FromType) && AreTypesCompatible(ToType, Info.ToType))
+				{
+					return Info.AutoConvertNode;
+				}
+			}
+		}
+		return NAME_None;
+	}
+
 	bool FAnyTypesRegistry::IsAnyType(FName Type) const
 	{
 		return TypeInfosByName.Contains(Type);
diff --git a/Engine/Source/Runtime/Experimental/Dataflow/Core/Private/Dataflow/DataflowVectorNodes.cpp b/Engine/Source/Runtime/Experimental/Dataflow/Core/Private/Dataflow/DataflowVectorNodes.cpp
--- a/Engine/Source/Runtime/Experimental/Dataflow/Core/Private/Dataflow/DataflowVectorNodes.cpp
+++ b/Engine/Source/Runtime/Experimental/Dataflow/Core/Private/Dataflow/DataflowVectorNodes.cpp
@@ -4,6 +4,7 @@
 
 #include "Dataflow/DataflowNodeFactory.h"
 #include "Dataflow/DataflowNodeColorsRegistry.h"
+#include "Dataflow/DataflowAnyTypeRegistry.h"
 
 namespace UE::Dataflow
 {
@@ -24,10 +25,12 @@ namespace UE::Dataflow
 		DATAFLOW_NODE_REGISTER_CREATION_FACTORY(FDataflowVectorCrossProductNode);
 		DATAFLOW_NODE_REGISTER_CREATION_FACTORY(FDataflowVectorScaleNode);
 		DATAFLOW_NODE_REGISTER_CREATION_FACTORY(FDataflowVectorNormalize);
+
+		UE_DATAFLOW_REGISTER_AUTOCONVERT(FDataflowVectorTypes, FDataflowNumericTypes, FDataflowVectorBreakNode);
+		UE_DATAFLOW_REGISTER_AUTOCONVERT(FDataflowNumericTypes, FDataflowVectorTypes, FDataflowVectorMakeVec3Node);
 	}
 }
 
-
 //-----------------------------------------------------------
 
 FDataflowVectorMakeVec2Node::FDataflowVectorMakeVec2Node(const UE::Dataflow::FNodeParameters& InParam, FGuid InGuid)
diff --git a/Engine/Source/Runtime/Experimental/Dataflow/Core/Public/Dataflow/DataflowAnyTypeRegistry.h b/Engine/Source/Runtime/Experimental/Dataflow/Core/Public/Dataflow/DataflowAnyTypeRegistry.h
--- a/Engine/Source/Runtime/Experimental/Dataflow/Core/Public/Dataflow/DataflowAnyTypeRegistry.h
+++ b/Engine/Source/Runtime/Experimental/Dataflow/Core/Public/Dataflow/DataflowAnyTypeRegistry.h
@@ -11,14 +11,20 @@ namespace UE::Dataflow
 	struct FAnyTypesRegistry
 	{
 	public:
+		using FTypeFilterFunction = TFunction<bool(FName, FName)>;
+
 		template<typename T>
 		static void RegisterTypeStatic(FName TypeName)
 		{
 			GetInstance().RegisterType<T>(TypeName);
 		}
-		static bool AreTypesCompatibleStatic(FName TypeA, FName TypeB);
-		static FName GetStorageTypeStatic(FName Type);
-		static bool IsAnyTypeStatic(FName Type);
+		DATAFLOWCORE_API static bool AreTypesCompatibleStatic(FName TypeA, FName TypeB);
+		DATAFLOWCORE_API static FName GetStorageTypeStatic(FName Type);
+
+		DATAFLOWCORE_API static void RegisterAutoConvertNodeTypeStatic(FName FromType, FName ToType, FName AutoConvertNode, FTypeFilterFunction TypeFilterFunction = {});
+		DATAFLOWCORE_API static FName GetAutoConvertNodeTypeStatic(FName FromType, FName ToType);
+
+		DATAFLOWCORE_API static bool IsAnyTypeStatic(FName Type);
 		DATAFLOWCORE_API static FAnyTypesRegistry& GetInstance();
 
 	private:
@@ -30,6 +36,14 @@ namespace UE::Dataflow
 			FName StorageType;
 		};
 
+		struct FAutoConvertNodeInfo
+		{
+			FName FromType;
+			FName ToType;
+			FName AutoConvertNode;
+			FTypeFilterFunction TypeFilterFunction;
+		};
+
 		FAnyTypesRegistry() {};
 
 		template<typename T>
@@ -43,12 +57,18 @@ namespace UE::Dataflow
 			TypeInfosByName.Emplace(TypeName, TypeInfo);
 		}
 
+		void RegisterAutoConvertNodeType(FName FromType, FName ToType, FName AutoConvertNode, FTypeFilterFunction TypeFilterFunction);
+
 		bool AreTypesCompatible(FName TypeA, FName TypeB) const;
 		FName GetStorageType(FName Type) const;
+		FName GetAutoConvertNodeType(FName FromType, FName ToType) const;
 		bool IsAnyType(FName Type) const;
 
 		TMap<FName, FTypeInfo> TypeInfosByName;
+		TArray<FAutoConvertNodeInfo> AutoConvertNodes;
 	};
 };
 
-#define UE_DATAFLOW_REGISTER_ANYTYPE(TYPE) UE::Dataflow::FAnyTypesRegistry::RegisterTypeStatic<TYPE>(#TYPE)
+#define UE_DATAFLOW_REGISTER_ANYTYPE(Type) UE::Dataflow::FAnyTypesRegistry::RegisterTypeStatic<Type>(#Type)
+#define UE_DATAFLOW_REGISTER_AUTOCONVERT(FromType, ToType, NodeType) UE::Dataflow::FAnyTypesRegistry::RegisterAutoConvertNodeTypeStatic(#FromType, #ToType, #NodeType)
+#define UE_DATAFLOW_REGISTER_AUTOCONVERT_WITH_FILTER(FromType, ToType, NodeType, FilterFunc) UE::Dataflow::FAnyTypesRegistry::RegisterAutoConvertNodeTypeStatic(#FromType, #ToType, #NodeType, FilterFunc)
