Dataflow : improve node pin behavior
- Pinning a node now unpin all the other nodes unless CTRL is pressed
- Updated the interface to pass the node that had his pin status changed

diff --git a/Engine/Plugins/Experimental/Dataflow/Source/DataflowEditor/Private/Dataflow/DataflowGraphEditor.cpp b/Engine/Plugins/Experimental/Dataflow/Source/DataflowEditor/Private/Dataflow/DataflowGraphEditor.cpp
--- a/Engine/Plugins/Experimental/Dataflow/Source/DataflowEditor/Private/Dataflow/DataflowGraphEditor.cpp
+++ b/Engine/Plugins/Experimental/Dataflow/Source/DataflowEditor/Private/Dataflow/DataflowGraphEditor.cpp
@@ -420,8 +420,27 @@ UDataflow* SDataflowGraphEditor::GetDataflowAsset() const
 	return nullptr;
 }
 
-void SDataflowGraphEditor::OnRenderToggleChanged() const
+void SDataflowGraphEditor::OnRenderToggleChanged(UDataflowEdNode* UpdatedEdNode) const
 {
+	// deselect all the existing pinned nodes unless CTRL is pressed
+	if (UpdatedEdNode->ShouldWireframeRenderNode() && !IsControlDown())
+	{
+		if (UEdGraph* EdGraph = GetCurrentGraph())
+		{
+			for (UEdGraphNode* EdNode : EdGraph->Nodes)
+			{
+				if (UDataflowEdNode* DataflowEdNode = Cast<UDataflowEdNode>(EdNode))
+				{
+					if (UpdatedEdNode != DataflowEdNode)
+					{
+						DataflowEdNode->SetShouldWireframeRenderNode(false);
+					}
+				}
+			}
+		}
+	}
+	// todo : if nothing is selected should we render the terminal node ?
+
 	// need to refresh the UI by refreshing the selection
 	FGraphPanelSelectionSet SelectionSet = GetSelectedNodes();
 	GetGraphPanel()->SelectionManager.SetSelectionSet(SelectionSet);
diff --git a/Engine/Plugins/Experimental/Dataflow/Source/DataflowEditor/Private/Dataflow/DataflowSNode.cpp b/Engine/Plugins/Experimental/Dataflow/Source/DataflowEditor/Private/Dataflow/DataflowSNode.cpp
--- a/Engine/Plugins/Experimental/Dataflow/Source/DataflowEditor/Private/Dataflow/DataflowSNode.cpp
+++ b/Engine/Plugins/Experimental/Dataflow/Source/DataflowEditor/Private/Dataflow/DataflowSNode.cpp
@@ -200,7 +200,7 @@ TSharedRef<SWidget> SDataflowEdNode::CreateTitleRightWidget()
 				{
 					const bool bShouldWireframeRenderNode = (NewState == ECheckBoxState::Checked);
 					DataflowGraphNode->SetShouldWireframeRenderNode(bShouldWireframeRenderNode);
-					DataflowInterface->OnRenderToggleChanged();
+					DataflowInterface->OnRenderToggleChanged(DataflowGraphNode);
 				}
 			})
 				.IsEnabled_Lambda([this]()->bool
diff --git a/Engine/Plugins/Experimental/Dataflow/Source/DataflowEditor/Public/Dataflow/DataflowGraphEditor.h b/Engine/Plugins/Experimental/Dataflow/Source/DataflowEditor/Public/Dataflow/DataflowGraphEditor.h
--- a/Engine/Plugins/Experimental/Dataflow/Source/DataflowEditor/Public/Dataflow/DataflowGraphEditor.h
+++ b/Engine/Plugins/Experimental/Dataflow/Source/DataflowEditor/Public/Dataflow/DataflowGraphEditor.h
@@ -162,7 +162,7 @@ public:
 
 	/** FDataflowSNodeInterface */
 	UE_API virtual TSharedPtr<UE::Dataflow::FContext> GetDataflowContext() const override;
-	UE_API virtual void OnRenderToggleChanged() const override;
+	UE_API virtual void OnRenderToggleChanged(UDataflowEdNode* UpdatedEdNode) const override;
 
 	const TSharedPtr<FUICommandList> GetCommands() const { return GraphEditorCommands; }
 
diff --git a/Engine/Plugins/Experimental/Dataflow/Source/DataflowEditor/Public/Dataflow/DataflowSEditorInterface.h b/Engine/Plugins/Experimental/Dataflow/Source/DataflowEditor/Public/Dataflow/DataflowSEditorInterface.h
--- a/Engine/Plugins/Experimental/Dataflow/Source/DataflowEditor/Public/Dataflow/DataflowSEditorInterface.h
+++ b/Engine/Plugins/Experimental/Dataflow/Source/DataflowEditor/Public/Dataflow/DataflowSEditorInterface.h
@@ -9,6 +9,8 @@ namespace UE::Dataflow
 	class FContext;
 }
 
+class UDataflowEdNode;
+
 /**
 * FDataflowSEditorInterface
 * 
@@ -29,7 +31,7 @@ public:
 		return true;
 	}
 
-	virtual void OnRenderToggleChanged() const
+	virtual void OnRenderToggleChanged(UDataflowEdNode* UpdatedEdNode) const
 	{
 	}
 
