Dataflow : Simplify the logic in preparation to more changes around how selection works
- Remove tracking of the primary selection
- Simplify the logic of running various code based of what is nealy selected
- Select node set are not strongly type on UDataflowEdNode
- Updated related classes
There's still some work to do but that will be for another change list

diff --git a/Engine/Plugins/ChaosClothAssetEditor/Source/ChaosClothAssetEditor/Private/ChaosClothAsset/ClothEditorToolkit.cpp b/Engine/Plugins/ChaosClothAssetEditor/Source/ChaosClothAssetEditor/Private/ChaosClothAsset/ClothEditorToolkit.cpp
--- a/Engine/Plugins/ChaosClothAssetEditor/Source/ChaosClothAssetEditor/Private/ChaosClothAsset/ClothEditorToolkit.cpp
+++ b/Engine/Plugins/ChaosClothAssetEditor/Source/ChaosClothAssetEditor/Private/ChaosClothAsset/ClothEditorToolkit.cpp
@@ -1059,13 +1059,20 @@ TSharedPtr<IStructureDetailsView> FChaosClothAssetEditorToolkit::CreateNodeDetai
 
 void FChaosClothAssetEditorToolkit::OnPropertyValueChanged(const FPropertyChangedEvent& PropertyChangedEvent)
 {
-	auto AsObjectPointers = [](const TSet<UObject*>& Set) {
-		TSet<TObjectPtr<UObject> > Objs; for (UObject* Elem : Set) Objs.Add(Elem);
-		return Objs;
+	auto AsEdNodePointers = [](const TSet<UObject*>& Set) {
+		TSet<TObjectPtr<UDataflowEdNode>> Nodes;
+		for (UObject* Elem : Set)
+		{
+			if (UDataflowEdNode* Node = Cast<UDataflowEdNode>(Elem))
+			{
+				Nodes.Add(Node);
+			}
+		}
+		return Nodes;
 	};
 
 	FDataflowEditorCommands::OnPropertyValueChanged(GetDataflow(), DataflowContext, LastDataflowNodeTimestamp, PropertyChangedEvent,
-		GraphEditor ? AsObjectPointers(GraphEditor->GetSelectedNodes()) : TSet<TObjectPtr<UObject>>());
+		GraphEditor ? AsEdNodePointers(GraphEditor->GetSelectedNodes()) : TSet<TObjectPtr<UDataflowEdNode>>());
 }
 
 bool FChaosClothAssetEditorToolkit::OnNodeVerifyTitleCommit(const FText& NewText, UEdGraphNode* GraphNode, FText& OutErrorMessage) const
diff --git a/Engine/Plugins/Experimental/Dataflow/Source/DataflowEditor/Private/Dataflow/DataflowConstructionScene.cpp b/Engine/Plugins/Experimental/Dataflow/Source/DataflowEditor/Private/Dataflow/DataflowConstructionScene.cpp
--- a/Engine/Plugins/Experimental/Dataflow/Source/DataflowEditor/Private/Dataflow/DataflowConstructionScene.cpp
+++ b/Engine/Plugins/Experimental/Dataflow/Source/DataflowEditor/Private/Dataflow/DataflowConstructionScene.cpp
@@ -588,6 +588,46 @@ void FDataflowConstructionScene::UpdateConstructionScene()
 	}
 }
 
+void FDataflowConstructionScene::SelectNodeComponents(const UDataflowEdNode* EdNode)
+{
+	TArray<AActor*> FoundActors;
+	if (USelection* SelectedComponents = GetDataflowModeManager()->GetSelectedComponents())
+	{
+		SelectedComponents->Modify();
+		SelectedComponents->BeginBatchSelectOperation();
+
+		TArray<TWeakObjectPtr<UObject>> SelectedObjects;
+		const int32 NumSelected = SelectedComponents->GetSelectedObjects(SelectedObjects);
+		for (TWeakObjectPtr<UObject> WeakObject : SelectedObjects)
+		{
+			if (WeakObject.IsValid())
+			{
+				if (UDataflowEditorCollectionComponent* ActorComponent = Cast< UDataflowEditorCollectionComponent>(WeakObject.Get()))
+				{
+					SelectedComponents->Deselect(ActorComponent);
+					ActorComponent->PushSelectionToProxy();
+				}
+			}
+		}
+
+		if (TObjectPtr<AActor> RootActor = GetRootActor())
+		{
+			for (UActorComponent* ActorComponent : RootActor->GetComponents())
+			{
+				if (UDataflowEditorCollectionComponent* Component = Cast<UDataflowEditorCollectionComponent>(ActorComponent))
+				{
+					if (Component->Node == EdNode)
+					{
+						SelectedComponents->Select(Component);
+						Component->PushSelectionToProxy();
+					}
+				}
+			}
+		}
+		SelectedComponents->EndBatchSelectOperation();
+	}
+}
+
 
 #undef LOCTEXT_NAMESPACE
 
diff --git a/Engine/Plugins/Experimental/Dataflow/Source/DataflowEditor/Private/Dataflow/DataflowEditorCommands.cpp b/Engine/Plugins/Experimental/Dataflow/Source/DataflowEditor/Private/Dataflow/DataflowEditorCommands.cpp
--- a/Engine/Plugins/Experimental/Dataflow/Source/DataflowEditor/Private/Dataflow/DataflowEditorCommands.cpp
+++ b/Engine/Plugins/Experimental/Dataflow/Source/DataflowEditor/Private/Dataflow/DataflowEditorCommands.cpp
@@ -81,7 +81,8 @@ void FDataflowEditorCommandsImpl::RegisterCommands()
 	UI_COMMAND(AddNewSubGraph, "SubGraph", "Adds a new subgraph to this dataflow graph.", EUserInterfaceActionType::Button, FInputChord());
 	UI_COMMAND(ConvertToBasicSubGraph, "Convert to Basic Subgraph", "Convert the subgraph to a basic one (no loop).", EUserInterfaceActionType::Button, FInputChord());
 	UI_COMMAND(ConvertToForEachSubGraph, "Convert to For Each Subgraph", "Convert the subgraph to be able to iterate through an array.", EUserInterfaceActionType::Button, FInputChord());
-	
+	UI_COMMAND(CollapseToSubGraph, "Collapse to Sub-graph", "Collapse the selection into a subgraph", EUserInterfaceActionType::Button, FInputChord());
+
 
 	for (const TPair<FName, TUniquePtr<UE::Dataflow::IDataflowConstructionViewMode>>& NameAndMode : UE::Dataflow::FRenderingViewModeFactory::GetInstance().GetViewModes())
 	{
@@ -393,7 +394,7 @@ void FDataflowEditorCommands::OnAssetPropertyValueChanged(TObjectPtr<UDataflowBa
 	}
 }
 
-void FDataflowEditorCommands::OnPropertyValueChanged(UDataflow* OutDataflow, TSharedPtr<UE::Dataflow::FEngineContext>& Context, UE::Dataflow::FTimestamp& OutLastNodeTimestamp, const FPropertyChangedEvent& InPropertyChangedEvent, const TSet<TObjectPtr<UObject> >& SelectedNodes)
+void FDataflowEditorCommands::OnPropertyValueChanged(UDataflow* OutDataflow, TSharedPtr<UE::Dataflow::FEngineContext>& Context, UE::Dataflow::FTimestamp& OutLastNodeTimestamp, const FPropertyChangedEvent& InPropertyChangedEvent, const TSet<TObjectPtr<UDataflowEdNode>>& SelectedNodes)
 {
 	if (InPropertyChangedEvent.ChangeType != EPropertyChangeType::Interactive)
 	{
@@ -401,33 +402,30 @@ void FDataflowEditorCommands::OnPropertyValueChanged(UDataflow* OutDataflow, TSh
 		{
 			OutDataflow->Modify();  // Modify must be called even if SelectedNodes is empty because comment nodes aren't part of the selection set but still have properties
 
-			for (UObject* const SelectedNode : SelectedNodes)
+			for (UDataflowEdNode* const SelectedNode : SelectedNodes)
 			{
-				if (UDataflowEdNode* const Node = Cast<UDataflowEdNode>(SelectedNode))
+				if (TSharedPtr<FDataflowNode> DataflowNode = SelectedNode->GetDataflowNode())
 				{
-					if (TSharedPtr<FDataflowNode> DataflowNode = Node->GetDataflowNode())
+					// Active state update
+					if (InPropertyChangedEvent.GetPropertyName() == FDataflowNode::GetActivePropertyName())
 					{
-						// Active state update
-						if (InPropertyChangedEvent.GetPropertyName() == FDataflowNode::GetActivePropertyName())
-						{
-							// Reflect the active state on the drawing of the node
-							constexpr bool bCheckIsActiveFlagOnly = true;
-							if (DataflowNode->IsActive(bCheckIsActiveFlagOnly) != Node->IsNodeEnabled())
-							{
-								Node->SetEnabledState(DataflowNode->IsActive(bCheckIsActiveFlagOnly) ? ENodeEnabledState::Enabled : ENodeEnabledState::Disabled);
-							}
-						}
-						if (DataflowNode->ShouldInvalidateOnPropertyChanged(InPropertyChangedEvent))
+						// Reflect the active state on the drawing of the node
+						constexpr bool bCheckIsActiveFlagOnly = true;
+						if (DataflowNode->IsActive(bCheckIsActiveFlagOnly) != SelectedNode->IsNodeEnabled())
 						{
-							// Invalidate the node and reset the editor timestamp
-							DataflowNode->Invalidate();
-							OutLastNodeTimestamp = UE::Dataflow::FTimestamp::Invalid;
-						}
-						if (Context.IsValid())
-						{
-							DataflowNode->OnPropertyChanged(*Context, InPropertyChangedEvent);
+							SelectedNode->SetEnabledState(DataflowNode->IsActive(bCheckIsActiveFlagOnly) ? ENodeEnabledState::Enabled : ENodeEnabledState::Disabled);
 						}
 					}
+					if (DataflowNode->ShouldInvalidateOnPropertyChanged(InPropertyChangedEvent))
+					{
+						// Invalidate the node and reset the editor timestamp
+						DataflowNode->Invalidate();
+						OutLastNodeTimestamp = UE::Dataflow::FTimestamp::Invalid;
+					}
+					if (Context.IsValid())
+					{
+						DataflowNode->OnPropertyChanged(*Context, InPropertyChangedEvent);
+					}
 				}
 			}
 		}
diff --git a/Engine/Plugins/Experimental/Dataflow/Source/DataflowEditor/Private/Dataflow/DataflowEditorToolkit.cpp b/Engine/Plugins/Experimental/Dataflow/Source/DataflowEditor/Private/Dataflow/DataflowEditorToolkit.cpp
--- a/Engine/Plugins/Experimental/Dataflow/Source/DataflowEditor/Private/Dataflow/DataflowEditorToolkit.cpp
+++ b/Engine/Plugins/Experimental/Dataflow/Source/DataflowEditor/Private/Dataflow/DataflowEditorToolkit.cpp
@@ -1379,235 +1379,127 @@ void FDataflowEditorToolkit::OnNodeDoubleClicked(UEdGraphNode* ClickedNode)
 	}
 }
 
-void FDataflowEditorToolkit::OnNodeSelectionChanged(const TSet<UObject*>& InNewSelection)
+void FDataflowEditorToolkit::ResetViewListeners(UDataflowEdNode* Node)
 {
-	//
-	// Local helper lambdas
-	//
-
-	auto FindDataflowNodesInSet = [](const TSet<TObjectPtr<UObject>>& InSet) {
-		TSet<TObjectPtr<UObject>> Results;
-		for (UObject* Item : InSet)
-		{
-			if (Cast<UDataflowEdNode>(Item))
-			{
-				Results.Add(Item);
-			}
-		}
-		return Results;
-	};
-
-	auto ResetListeners = [&ViewListeners = ViewListeners](UDataflowEdNode* Node = nullptr)
+	for (IDataflowViewListener* Listener : ViewListeners)
 	{
-		for (IDataflowViewListener* Listener : ViewListeners)
-		{
-			Listener->OnSelectedNodeChanged(nullptr);
-		}
-		if (Node)
-		{
-			for (IDataflowViewListener* Listener : ViewListeners)
-			{
-				Listener->OnSelectedNodeChanged(Node);
-			}
-		}
-	};
-
-	auto IsControlDown = [&DataflowEditor = GraphEditor]()
+		Listener->OnSelectedNodeChanged(nullptr);
+	}
+	if (Node)
 	{
-		if (DataflowEditor)
+		for (IDataflowViewListener* Listener : ViewListeners)
 		{
-			return DataflowEditor->IsControlDown();
+			Listener->OnSelectedNodeChanged(Node);
 		}
-		return false;
-	};
+	}
+};
 
-	auto SelectComponentsInView = [&](TObjectPtr<UDataflowEdNode> Node)
+TSet<TObjectPtr<UDataflowEdNode>> FDataflowEditorToolkit::FilterDataflowEdNodesFromSet(const TSet<UObject*>& Set)
+{
+	TSet<TObjectPtr<UDataflowEdNode>> DataflowEdNodes;
+	for (UObject* ObjectPtr : Set)
 	{
-		TArray<AActor*> FoundActors;
-		if (USelection* SelectedComponents = ConstructionScene->GetDataflowModeManager()->GetSelectedComponents())
+		if (UDataflowEdNode* DataflowEdNode = Cast<UDataflowEdNode>(ObjectPtr))
 		{
-			SelectedComponents->Modify();
-			SelectedComponents->BeginBatchSelectOperation();
-
-			TArray<TWeakObjectPtr<UObject>> SelectedObjects;
-			const int32 NumSelected = SelectedComponents->GetSelectedObjects(SelectedObjects);
-			for (TWeakObjectPtr<UObject> WeakObject : SelectedObjects)
-			{
-				if (WeakObject.IsValid())
-				{
-					if (UDataflowEditorCollectionComponent* ActorComponent = Cast< UDataflowEditorCollectionComponent>(WeakObject.Get()))
-					{
-						SelectedComponents->Deselect(ActorComponent);
-						ActorComponent->PushSelectionToProxy();
-					}
-				}
-			}
-
-			if (TObjectPtr<AActor> RootActor = ConstructionScene->GetRootActor())
-			{
-				for (UActorComponent* ActorComponent : RootActor->GetComponents())
-				{
-					if (UDataflowEditorCollectionComponent* Component = Cast<UDataflowEditorCollectionComponent>(ActorComponent))
-					{
-						if (Component->Node == Node)
-						{
-							SelectedComponents->Select(Component);
-							Component->PushSelectionToProxy();
-						}
-					}
-				}
-			}
-			SelectedComponents->EndBatchSelectOperation();
+			DataflowEdNodes.Add(DataflowEdNode);
 		}
-	};
-
-	//
-	// Actual function
-	// 
-	
-	// Despite this function's name, we might not have actually changed which node is selected
-	bool bPrimarySelectionChanged = false;
+	}
+	return DataflowEdNodes;
+}
 
-	if (const TObjectPtr<UDataflowBaseContent>& EditorContent = GetEditorContent(); EditorContent->GetDataflowAsset())
+const TObjectPtr<UDataflowEdNode> FDataflowEditorToolkit::GetFirstFromSet(const TSet<TObjectPtr<UDataflowEdNode>>& Set)
+{
+	if (Set.Num() > 0)
 	{
-		auto AsObjectPointers = [](const TSet<UObject*>& Set) {
-			TSet<TObjectPtr<UObject> > Objs; for (UObject* Elem : Set) Objs.Add(Elem);
-			return Objs;
-		};
+		return *Set.CreateConstIterator();
+	}
+	return nullptr;
+}
 
-		const TSet<TObjectPtr<UObject>> PreviouslySelectedNodes = SelectedDataflowNodes;
-		for (UObject* const PreviouslySelectedNode : SelectedDataflowNodes)
-		{
-			if (UDataflowEdNode* const PreviouslySelectedEdNode = Cast<UDataflowEdNode>(PreviouslySelectedNode))
-			{
-				PreviouslySelectedEdNode->SetShouldRenderNode(false);
-			}
-		}
+void FDataflowEditorToolkit::OnNodeSelectionChanged(const TSet<UObject*>& InNewSelection)
+{
+	// update the current selection set and keep the previous one
+	TSet<TObjectPtr<UDataflowEdNode>> PreviouslySelectedNodes = SelectedDataflowNodes;
+	SelectedDataflowNodes = FilterDataflowEdNodesFromSet(InNewSelection);
 
-		// Only keep UDataflowEdNode from NewSelection
-		TSet< TObjectPtr<UObject> > NodeSelection = FindDataflowNodesInSet(AsObjectPointers(InNewSelection));
+	// compute the various delta sets
+	TSet<TObjectPtr<UDataflowEdNode>> DeselectedNodes = PreviouslySelectedNodes.Difference(SelectedDataflowNodes);
+	TSet<TObjectPtr<UDataflowEdNode>> StillSelectedNodes = PreviouslySelectedNodes.Intersect(SelectedDataflowNodes);
+	TSet<TObjectPtr<UDataflowEdNode>> NewlySelectedNodes = SelectedDataflowNodes.Difference(PreviouslySelectedNodes);
 
-		if (!NodeSelection.Num())
+	// update the render flag 
+	auto SetShouldRenderNodeFlag = [](const TSet<TObjectPtr<UDataflowEdNode>>& Nodes, bool bValue)
 		{
-			// The selection is empty. 
-			ResetListeners();
-			SelectedDataflowNodes = TSet<TObjectPtr<UObject>>();
-			if (PrimarySelection) bPrimarySelectionChanged = true;
-			PrimarySelection = nullptr;
-		}
-		else
-		{
-			TSet<TObjectPtr<UObject>> DeselectedNodes = SelectedDataflowNodes.Difference(NodeSelection);
-			TSet<TObjectPtr<UObject>> StillSelectedNodes = SelectedDataflowNodes.Intersect(NodeSelection);
-			TSet<TObjectPtr<UObject>> NewlySelectedNodes = NodeSelection.Difference(SelectedDataflowNodes);
-
-			// Something has been removed
-			if (DeselectedNodes.Num())
-			{
-				if (DeselectedNodes.Contains(PrimarySelection))
-				{
-					ResetListeners();
-
-					if (PrimarySelection) bPrimarySelectionChanged = true;
-					PrimarySelection = nullptr;
-
-					// pick a new primary if nothing new was selected
-					if (!NewlySelectedNodes.Num() && StillSelectedNodes.Num())
-					{
-						PrimarySelection = Cast< UDataflowEdNode>(StillSelectedNodes.Array()[0]);
-						ResetListeners(PrimarySelection);
-						bPrimarySelectionChanged = true;
-					}
-				}
-			}
-
-			// Something new has been selected.
-			if (NewlySelectedNodes.Num() >= 1)
+			for (UDataflowEdNode* const Node : Nodes)
 			{
-				PrimarySelection = Cast< UDataflowEdNode>(NewlySelectedNodes.Array()[0]);
-				ResetListeners(PrimarySelection);
-				bPrimarySelectionChanged = true;
+				Node->SetShouldRenderNode(bValue);
 			}
+		};
 
-			SelectedDataflowNodes = NodeSelection;
-		}
+	SetShouldRenderNodeFlag(DeselectedNodes, false);
+	SetShouldRenderNodeFlag(NewlySelectedNodes, true);
 
-		for (UObject* const SelectedNode : NodeSelection)
-		{
-			if (UDataflowEdNode* SelectedEdNode = Cast<UDataflowEdNode>(SelectedNode))
-			{
-				SelectedEdNode->SetShouldRenderNode(true);
-			}
-		}
+	const TObjectPtr<UDataflowEdNode> PrevSelectedNode = GetFirstFromSet(PreviouslySelectedNodes);
+	const TObjectPtr<UDataflowEdNode> NewSelectedNode = GetFirstFromSet(SelectedDataflowNodes);
+	ResetViewListeners(NewSelectedNode);
 
-		if (bPrimarySelectionChanged)
+	// todo : move the tool logic somewhere else: we should not be activelly looking for a collection here, we should have some sort of interface to talk to the tools
+	if (const TObjectPtr<UDataflowBaseContent>& EditorContent = GetEditorContent(); EditorContent->GetDataflowAsset())
+	{
+		if (NewSelectedNode != PrevSelectedNode)
 		{
-			for (UObject* const PreviouslySelectedNode : PreviouslySelectedNodes)
-			{
-				if (UDataflowEdNode* PreviouslySelectedEdNode = Cast<UDataflowEdNode>(PreviouslySelectedNode))
-				{
-					PreviouslySelectedEdNode->SetShouldRenderNode(false);
-				}
-			}
-
-			for (UObject* const SelectedNode : NodeSelection)
-			{
-				if (UDataflowEdNode* SelectedEdNode = Cast<UDataflowEdNode>(SelectedNode))
-				{
-					SelectedEdNode->SetShouldRenderNode(true);
-				}
-			}
-
 			EditorContent->SetSelectedNode(nullptr);
-
 			EditorContent->SetSelectedCollection(nullptr, /*bCollectionIsInput=*/ false);
 
-			if( UDataflowEditorMode* const DataflowMode = Cast<UDataflowEditorMode>(EditorModeManager->GetActiveScriptableMode(UDataflowEditorMode::EM_DataflowEditorModeId)) )
+			if (NewSelectedNode)
 			{
-				// Close any running tool. OnNodeSingleClicked() will start a new tool if a new node was clicked.
+				if (UDataflowEditorMode* const DataflowMode = Cast<UDataflowEditorMode>(EditorModeManager->GetActiveScriptableMode(UDataflowEditorMode::EM_DataflowEditorModeId)))
+				{
+					// Close any running tool. OnNodeSingleClicked() will start a new tool if a new node was clicked.
 
-				UEditorInteractiveToolsContext* const ToolsContext = DataflowMode->GetInteractiveToolsContext();
-				checkf(ToolsContext, TEXT("No valid ToolsContext found for FDataflowEditorToolkit"));
+					UEditorInteractiveToolsContext* const ToolsContext = DataflowMode->GetInteractiveToolsContext();
+					checkf(ToolsContext, TEXT("No valid ToolsContext found for FDataflowEditorToolkit"));
 
-				UInteractiveToolManager* const ToolManager = ToolsContext->ToolManager;
-				checkf(ToolManager, TEXT("ToolsContext doesn't have a valid ToolManager"));
+					UInteractiveToolManager* const ToolManager = ToolsContext->ToolManager;
+					checkf(ToolManager, TEXT("ToolsContext doesn't have a valid ToolManager"));
 
-				if (UInteractiveTool* const ActiveTool = ToolManager->GetActiveTool(EToolSide::Left))
-				{
-					ToolManager->PostActiveToolShutdownRequest(ActiveTool, EToolShutdownType::Accept);
-				}
+					if (UInteractiveTool* const ActiveTool = ToolManager->GetActiveTool(EToolSide::Left))
+					{
+						ToolManager->PostActiveToolShutdownRequest(ActiveTool, EToolShutdownType::Accept);
+					}
 
-				EditorContent->SetSelectedNode(PrimarySelection);
+					EditorContent->SetSelectedNode(NewSelectedNode);
 
-				// Call the node's OnSelected function. Some nodes use this to cache information from the inputs (e.g. FDataflowCollectionAddScalarVertexPropertyNode::CachedCollectionGroupNames)
-				TSharedPtr<UE::Dataflow::FEngineContext> DataflowContext = EditorContent->GetDataflowContext();
-				if (PrimarySelection && DataflowContext.IsValid())
-				{
-					if (const TSharedPtr<FDataflowNode> DataflowNode = PrimarySelection->GetDataflowNode())
+					// Call the node's OnSelected function. Some nodes use this to cache information from the inputs (e.g. FDataflowCollectionAddScalarVertexPropertyNode::CachedCollectionGroupNames)
+					TSharedPtr<UE::Dataflow::FEngineContext> DataflowContext = EditorContent->GetDataflowContext();
+					if (NewSelectedNode && DataflowContext.IsValid())
 					{
-						// Update selected Collection in the ContextObject
-						for (const FDataflowOutput* const Output : DataflowNode->GetOutputs())
+						if (const TSharedPtr<FDataflowNode> DataflowNode = NewSelectedNode->GetDataflowNode())
 						{
-							if (Output->GetType() == FName(TEXT("FManagedArrayCollection")))
+							// Update selected Collection in the ContextObject
+							for (const FDataflowOutput* const Output : DataflowNode->GetOutputs())
 							{
-								const FManagedArrayCollection DefaultValue;
-								TSharedRef<FManagedArrayCollection> Collection = MakeShared<FManagedArrayCollection>(Output->GetValue<FManagedArrayCollection>(*DataflowContext, DefaultValue));
-								constexpr bool bCollectionIsInput = false;
-								EditorContent->SetSelectedCollection(Collection, bCollectionIsInput);
+								if (Output->GetType() == FName(TEXT("FManagedArrayCollection")))
+								{
+									const FManagedArrayCollection DefaultValue;
+									TSharedRef<FManagedArrayCollection> Collection = MakeShared<FManagedArrayCollection>(Output->GetValue<FManagedArrayCollection>(*DataflowContext, DefaultValue));
+									constexpr bool bCollectionIsInput = false;
+									EditorContent->SetSelectedCollection(Collection, bCollectionIsInput);
+								}
 							}
 						}
 					}
 				}
 			}
-
+			
 			if (GetDataflowGraphEditor()->IsAltDown())
 			{
-				SelectComponentsInView(PrimarySelection);
+				ConstructionScene->SelectNodeComponents(NewSelectedNode);
 			}
-
-			UpdateDebugDraw();
 		}
 
+		// todo: need a better way to trigger the update of the various views ( should also include what ResetViewListeners does )
+		UpdateDebugDraw();
 		EditorContent->SetConstructionDirty(true);
 	}
 
@@ -1617,13 +1509,13 @@ void FDataflowEditorToolkit::OnNodeSelectionChanged(const TSet<UObject*>& InNewS
 
 	if (UDataflowEditorMode* const DataflowMode = Cast<UDataflowEditorMode>(EditorModeManager->GetActiveScriptableMode(UDataflowEditorMode::EM_DataflowEditorModeId)))
 	{
-		if (PrimarySelection && GetEditorContent())
+		if (NewSelectedNode && GetEditorContent())
 		{
 			// Check the most recently used view mode for this "kind" of node -- where "kind" means nodes that have the same set of valid view modes
 			auto FindRecentViewMode = [&](FName& OutViewModeName) -> bool
 			{
 				TArray<FName> ViewModesForNode;
-				UE::Dataflow::GetViewModesForNode(*PrimarySelection, *GetEditorContent(), ViewModesForNode);
+				UE::Dataflow::GetViewModesForNode(*NewSelectedNode, *GetEditorContent(), ViewModesForNode);
 
 				if (const FName* const FoundViewMode = DataflowMode->FindRecentlyUsedViewMode(ViewModesForNode))
 				{
@@ -1639,7 +1531,7 @@ void FDataflowEditorToolkit::OnNodeSelectionChanged(const TSet<UObject*>& InNewS
 				for (const TPair<FName, TUniquePtr<UE::Dataflow::IDataflowConstructionViewMode>>& ViewMode : UE::Dataflow::FRenderingViewModeFactory::GetInstance().GetViewModes())
 				{
 					check(ViewMode.Value.IsValid());
-					const bool bCanRender = UE::Dataflow::CanRenderNodeOutput(*PrimarySelection, *GetEditorContent(), *ViewMode.Value);
+					const bool bCanRender = UE::Dataflow::CanRenderNodeOutput(*NewSelectedNode, *GetEditorContent(), *ViewMode.Value);
 					if (bCanRender)
 					{
 						OutViewModeName = ViewMode.Key;
@@ -1655,7 +1547,7 @@ void FDataflowEditorToolkit::OnNodeSelectionChanged(const TSet<UObject*>& InNewS
 			{
 				DataflowMode->SetConstructionViewMode(NewViewMode);
 			}
-			else if (UE::Dataflow::CanRenderNodeOutput(*PrimarySelection, *GetEditorContent(), *DataflowMode->GetConstructionViewMode()))        // Can we use the current view mode?
+			else if (UE::Dataflow::CanRenderNodeOutput(*NewSelectedNode, *GetEditorContent(), *DataflowMode->GetConstructionViewMode()))        // Can we use the current view mode?
 			{
 				// no need to switch
 			}		
@@ -1762,7 +1654,8 @@ void FDataflowEditorToolkit::OnNodeSingleClicked(UObject* ClickedNode) const
 
 void FDataflowEditorToolkit::OnNodeDeleted(const TSet<UObject*>& NewSelection)
 {
-	for (UObject* Node : NewSelection)
+	const TSet<TObjectPtr<UDataflowEdNode>> DeletedNodes = FilterDataflowEdNodesFromSet(NewSelection);
+	for (const UDataflowEdNode* Node : DeletedNodes)
 	{
 		if (SelectedDataflowNodes.Contains(Node))
 		{
@@ -2894,7 +2787,7 @@ FDataflowConstructionScene* FDataflowEditorToolkit::GetConstructionScene() const
 void FDataflowEditorToolkit::AddReferencedObjects(FReferenceCollector& Collector)
 {
 	Collector.AddReferencedObjects(SelectedDataflowNodes);
-	Collector.AddReferencedObject(PrimarySelection);
+	//Collector.AddReferencedObject(PrimarySelection);
 	if(ConstructionScene)
 	{
 		Collector.AddReferencedObject(ConstructionScene->ModifyDebugDrawComponent());
diff --git a/Engine/Plugins/Experimental/Dataflow/Source/DataflowEditor/Public/Dataflow/DataflowConstructionScene.h b/Engine/Plugins/Experimental/Dataflow/Source/DataflowEditor/Public/Dataflow/DataflowConstructionScene.h
--- a/Engine/Plugins/Experimental/Dataflow/Source/DataflowEditor/Public/Dataflow/DataflowConstructionScene.h
+++ b/Engine/Plugins/Experimental/Dataflow/Source/DataflowEditor/Public/Dataflow/DataflowConstructionScene.h
@@ -72,6 +72,9 @@ public:
 	/** Array of all DynamicMeshComponents in the scene */
 	UE_API TArray<TObjectPtr<UDynamicMeshComponent>> GetDynamicMeshComponents() const;
 
+	/** select only the components that are relates to a specific node */
+	UE_API void SelectNodeComponents(const UDataflowEdNode* EdNode);
+
 	FDebugMesh DebugMesh;
 
 private:
diff --git a/Engine/Plugins/Experimental/Dataflow/Source/DataflowEditor/Public/Dataflow/DataflowEditorCommands.h b/Engine/Plugins/Experimental/Dataflow/Source/DataflowEditor/Public/Dataflow/DataflowEditorCommands.h
--- a/Engine/Plugins/Experimental/Dataflow/Source/DataflowEditor/Public/Dataflow/DataflowEditorCommands.h
+++ b/Engine/Plugins/Experimental/Dataflow/Source/DataflowEditor/Public/Dataflow/DataflowEditorCommands.h
@@ -78,6 +78,7 @@ public:
 	TSharedPtr<FUICommandInfo> ConvertToBasicSubGraph;
 	TSharedPtr<FUICommandInfo> ConvertToForEachSubGraph;
 	TSharedPtr<FUICommandInfo> ToggleAsyncEvaluation;
+	TSharedPtr<FUICommandInfo> CollapseToSubGraph;
 
 	TMap<FName, TSharedPtr<FUICommandInfo>> SetConstructionViewModeCommands;
 
@@ -173,7 +174,7 @@ public:
 	/*
 	*  OnPropertyValueChanged
 	*/
-	static UE_API void OnPropertyValueChanged(UDataflow* Graph, TSharedPtr<UE::Dataflow::FEngineContext>& Context, UE::Dataflow::FTimestamp& OutLastNodeTimestamp, const FPropertyChangedEvent& PropertyChangedEvent, const TSet<TObjectPtr<UObject>>& NewSelection = TSet<TObjectPtr<UObject>>());
+	static UE_API void OnPropertyValueChanged(UDataflow* Graph, TSharedPtr<UE::Dataflow::FEngineContext>& Context, UE::Dataflow::FTimestamp& OutLastNodeTimestamp, const FPropertyChangedEvent& PropertyChangedEvent, const TSet<TObjectPtr<UDataflowEdNode>>& NewSelection = {});
 	static UE_API void OnAssetPropertyValueChanged(TObjectPtr<UDataflowBaseContent> Content, const FPropertyChangedEvent& PropertyChangedEvent);
 
 	/*
diff --git a/Engine/Plugins/Experimental/Dataflow/Source/DataflowEditor/Public/Dataflow/DataflowEditorToolkit.h b/Engine/Plugins/Experimental/Dataflow/Source/DataflowEditor/Public/Dataflow/DataflowEditorToolkit.h
--- a/Engine/Plugins/Experimental/Dataflow/Source/DataflowEditor/Public/Dataflow/DataflowEditorToolkit.h
+++ b/Engine/Plugins/Experimental/Dataflow/Source/DataflowEditor/Public/Dataflow/DataflowEditorToolkit.h
@@ -267,6 +267,10 @@ private:
 	/** Update the debug draw based on a change of currently selected or pinned nodes */
 	UE_API void UpdateDebugDraw();
 
+	void ResetViewListeners(UDataflowEdNode* Node = nullptr);
+	static TSet<TObjectPtr<UDataflowEdNode>> FilterDataflowEdNodesFromSet(const TSet<UObject*>& Set);
+	static const TObjectPtr<UDataflowEdNode> GetFirstFromSet(const TSet<TObjectPtr<UDataflowEdNode>>& Set);
+
 	// List of editor commands used  for the dataflow asset
 	TSharedPtr<FUICommandList> GraphEditorCommands;
 
@@ -294,12 +298,8 @@ private:
 
 	// The currently selected set of dataflow nodes. 
 	UPROPERTY()
-	TSet< TObjectPtr<UObject> > SelectedDataflowNodes;
+	TSet<TObjectPtr<UDataflowEdNode>> SelectedDataflowNodes;
 
-	// The most recently selected dataflow node.
-	UPROPERTY()
-	TObjectPtr<UDataflowEdNode> PrimarySelection;
-	
 	/** Pointer to the construction viewport scene. Note this is an alias of ObjectScene in FBaseCharacterFXEditorToolkit but with the specific type */
 	FDataflowConstructionScene* ConstructionScene;
 
