Dataflow : improve update logic of dataflow views to avoid calling update is if there's no change

diff --git a/Engine/Plugins/Experimental/Dataflow/Source/DataflowEditor/Private/Dataflow/DataflowView.cpp b/Engine/Plugins/Experimental/Dataflow/Source/DataflowEditor/Private/Dataflow/DataflowView.cpp
--- a/Engine/Plugins/Experimental/Dataflow/Source/DataflowEditor/Private/Dataflow/DataflowView.cpp
+++ b/Engine/Plugins/Experimental/Dataflow/Source/DataflowEditor/Private/Dataflow/DataflowView.cpp
@@ -27,7 +27,7 @@ bool FDataflowNodeView::SelectedNodeHaveSupportedOutputTypes(UDataflowEdNode* In
 {
 	SetSupportedOutputTypes();
 
-	if (InNode->IsBound())
+	if (InNode && InNode->IsBound())
 	{
 		if (TSharedPtr<FDataflowNode> DataflowNode = InNode->DataflowGraph->FindBaseNode(InNode->DataflowNodeGuid))
 		{
@@ -59,27 +59,56 @@ void FDataflowNodeView::OnSimulationViewSelectionChanged(const TArray<UPrimitive
 	SimulationViewSelectionChanged(SelectedComponents, SelectedElements);
 }
 
+UE::Dataflow::FTimestamp  FDataflowNodeView::GetNodeLastModifiedTimeStamp(UDataflowEdNode* InNode) const
+{
+	if (InNode)
+	{
+		return GetNodeLastModifiedTimeStamp(InNode->GetDataflowNode().Get());
+	}
+	return UE::Dataflow::FTimestamp::Invalid;
+}
+UE::Dataflow::FTimestamp  FDataflowNodeView::GetNodeLastModifiedTimeStamp(FDataflowNode* InNode) const
+{
+	if (InNode)
+	{
+		return InNode->GetTimestamp();
+	}
+	return UE::Dataflow::FTimestamp::Invalid;
+}
+
 void FDataflowNodeView::OnSelectedNodeChanged(UDataflowEdNode* InNode)
 {
 	if (!bIsPinnedDown)
 	{
-		SelectedNode = nullptr;
+		UDataflowEdNode* ValidNodeToSelect = nullptr;
+		if (SelectedNodeHaveSupportedOutputTypes(InNode))
+		{
+			ValidNodeToSelect = InNode;
+		}
 
-		if (InNode)  // nullptr is valid
+		const UE::Dataflow::FTimestamp NodeTimestamp = GetNodeLastModifiedTimeStamp(ValidNodeToSelect);
+		const bool bIsMoreRecentThanRefresh = (LastRefreshTimestamp < NodeTimestamp);
+		const bool bIsDifferentNode = (ValidNodeToSelect != SelectedNode);
+		
+		if (bIsDifferentNode)
 		{
-			if (SelectedNodeHaveSupportedOutputTypes(InNode))
-			{
-				SelectedNode = InNode;
-			}
+			SelectedNode = ValidNodeToSelect;
 		}
 
-		UpdateViewData();
+		if (bIsMoreRecentThanRefresh || bIsDifferentNode)
+		{
+			LastRefreshTimestamp = NodeTimestamp;
+			UpdateViewData();
+		}
 	}
 }
 
 void FDataflowNodeView::RefreshView()
 {
-	if (!bIsRefreshLocked && SelectedNode)
+	const UE::Dataflow::FTimestamp NodeTimestamp = GetNodeLastModifiedTimeStamp(SelectedNode);
+	const bool bIsMoreRecentThanRefresh = (LastRefreshTimestamp < NodeTimestamp);
+
+	if (!bIsRefreshLocked && SelectedNode && bIsMoreRecentThanRefresh)
 	{
 		UpdateViewData();
 	}
diff --git a/Engine/Plugins/Experimental/Dataflow/Source/DataflowEditor/Public/Dataflow/DataflowView.h b/Engine/Plugins/Experimental/Dataflow/Source/DataflowEditor/Public/Dataflow/DataflowView.h
--- a/Engine/Plugins/Experimental/Dataflow/Source/DataflowEditor/Public/Dataflow/DataflowView.h
+++ b/Engine/Plugins/Experimental/Dataflow/Source/DataflowEditor/Public/Dataflow/DataflowView.h
@@ -75,10 +75,15 @@ public:
 	virtual void AddReferencedObjects(FReferenceCollector& Collector) override;
 	virtual FString GetReferencerName() const override { return TEXT("FDataflowNodeView"); }
 
+public:
+	UE::Dataflow::FTimestamp GetNodeLastModifiedTimeStamp(UDataflowEdNode* InNode) const;
+	UE::Dataflow::FTimestamp GetNodeLastModifiedTimeStamp(FDataflowNode* InNode) const;
+
 private:
 	TObjectPtr<UDataflowBaseContent> EditorContent = nullptr;
 
 	TObjectPtr<UDataflowEdNode> SelectedNode = nullptr;
+	UE::Dataflow::FTimestamp LastRefreshTimestamp = UE::Dataflow::FTimestamp::Invalid;
 
 	bool bIsPinnedDown = false;
 
