Dataflow : continued imporving and simplifying the selection logic
- Consolidate the views update with the debug draw one
- extracted some logic to external methods

diff --git a/Engine/Plugins/Experimental/Dataflow/Source/DataflowEditor/Private/Dataflow/DataflowEditorMode.cpp b/Engine/Plugins/Experimental/Dataflow/Source/DataflowEditor/Private/Dataflow/DataflowEditorMode.cpp
--- a/Engine/Plugins/Experimental/Dataflow/Source/DataflowEditor/Private/Dataflow/DataflowEditorMode.cpp
+++ b/Engine/Plugins/Experimental/Dataflow/Source/DataflowEditor/Private/Dataflow/DataflowEditorMode.cpp
@@ -891,6 +891,65 @@ void UDataflowEditorMode::SetConstructionViewMode(const FName& NewViewModeName)
 
 }
 
+void UDataflowEditorMode::SetConstructionViewModeForNode(const UDataflowEdNode* Node)
+{
+	const TObjectPtr<UDataflowBaseContent> EditorContent = ConstructionScene ? ConstructionScene->GetEditorContent() : nullptr;
+	if (Node && EditorContent)
+	{
+		// Check the most recently used view mode for this "kind" of node -- where "kind" means nodes that have the same set of valid view modes
+		auto FindRecentViewMode =
+			[this, &Node, &EditorContent](FName& OutViewModeName) -> bool
+			{
+				TArray<FName> ViewModesForNode;
+				UE::Dataflow::GetViewModesForNode(*Node, *EditorContent, ViewModesForNode);
+
+				if (const FName* const FoundViewMode = FindRecentlyUsedViewMode(ViewModesForNode))
+				{
+					OutViewModeName = *FoundViewMode;
+					return true;
+				}
+				return false;
+			};
+
+		// Get the first view mode that this node can use
+		auto FindAnyValidViewMode =
+			[this, &Node, &EditorContent](FName& OutViewModeName) -> bool
+			{
+				for (const TPair<FName, TUniquePtr<UE::Dataflow::IDataflowConstructionViewMode>>& ViewMode : UE::Dataflow::FRenderingViewModeFactory::GetInstance().GetViewModes())
+				{
+					check(ViewMode.Value.IsValid());
+					const bool bCanRender = UE::Dataflow::CanRenderNodeOutput(*Node, *EditorContent, *ViewMode.Value);
+					if (bCanRender)
+					{
+						OutViewModeName = ViewMode.Key;
+						return true;
+					}
+				}
+				return false;
+			};
+
+		FName NewViewMode;
+		if (FindRecentViewMode(NewViewMode))               // Is there a recently used view mode we can switch to?
+		{
+			SetConstructionViewMode(NewViewMode);
+		}
+		else if (UE::Dataflow::CanRenderNodeOutput(*Node, *EditorContent, *GetConstructionViewMode()))        // Can we use the current view mode?
+		{
+			// no need to switch
+		}
+		else if (FindAnyValidViewMode(NewViewMode))             // Is there *any* view mode that works?
+		{
+			SetConstructionViewMode(NewViewMode);
+		}
+		else                                                    // No valid view mode found
+		{
+			// TODO: We should clear and disable View Mode Button. For now set default mode to the built-in 3D view mode.
+			SetConstructionViewMode(UE::Dataflow::FDataflowConstruction3DViewMode::Name);
+		}
+	}
+}
+
+
 const UE::Dataflow::IDataflowConstructionViewMode* UDataflowEditorMode::GetConstructionViewMode() const
 {
 	return ConstructionViewMode;
@@ -1043,6 +1102,20 @@ void UDataflowEditorMode::StartToolForSelectedNode(const UObject* SelectedNode)
 	}
 }
 
+void UDataflowEditorMode::ShutdownActiveToolIfNeeded(EToolShutdownType ShutdownType)
+{
+	UEditorInteractiveToolsContext* const ToolsContext = GetInteractiveToolsContext();
+	checkf(ToolsContext, TEXT("No valid ToolsContext found for UDataflowEditorMode"));
+
+	UInteractiveToolManager* const ToolManager = ToolsContext->ToolManager;
+	checkf(ToolManager, TEXT("ToolsContext doesn't have a valid ToolManager"));
+
+	if (UInteractiveTool* const ActiveTool = ToolManager->GetActiveTool(EToolSide::Left))
+	{
+		ToolManager->PostActiveToolShutdownRequest(ActiveTool, ShutdownType);
+	}
+}
+
 void UDataflowEditorMode::OnDataflowNodeDeleted(const TSet<UObject*>& DeletedNodes)
 {
 	UEditorInteractiveToolsContext* const ToolsContext = GetInteractiveToolsContext();
diff --git a/Engine/Plugins/Experimental/Dataflow/Source/DataflowEditor/Private/Dataflow/DataflowEditorToolkit.cpp b/Engine/Plugins/Experimental/Dataflow/Source/DataflowEditor/Private/Dataflow/DataflowEditorToolkit.cpp
--- a/Engine/Plugins/Experimental/Dataflow/Source/DataflowEditor/Private/Dataflow/DataflowEditorToolkit.cpp
+++ b/Engine/Plugins/Experimental/Dataflow/Source/DataflowEditor/Private/Dataflow/DataflowEditorToolkit.cpp
@@ -87,6 +87,10 @@ bool bDataflowEnableSkeletonView = false;
 FAutoConsoleVariableRef CVARDataflowEnableSkeletonView(TEXT("p.Dataflow.Editor.EnableSkeletonView"), bDataflowEnableSkeletonView,
 	TEXT("Deprecated Tool! Allows the Dataflow editor to create a skeleton view that reflects the hierarchy and selection state of the construction viewport.[def:false]"));
 
+bool bDataflowLogNodeEvaluation = false;
+FAutoConsoleVariableRef CVARDataflowLogNodeEvaluation(TEXT("p.Dataflow.Editor.LogNodeEvaluation"), bDataflowLogNodeEvaluation,
+	TEXT("When enabled, logs the begin and end evaluation of Dataflow nodes - if the graph is large this may cause evry large amounty of log messages to be displayed in the output window.[def:false]"));
+
 const FName FDataflowEditorToolkit::GraphCanvasTabId(TEXT("DataflowEditor_GraphCanvas"));
 const FName FDataflowEditorToolkit::SubGraphCanvasTabId(TEXT("DataflowEditor_SubGraphTab"));
 const FName FDataflowEditorToolkit::NodeDetailsTabId(TEXT("DataflowEditor_NodeDetails"));
@@ -1287,8 +1291,6 @@ void FDataflowEditorToolkit::OnPropertyValueChanged(const FPropertyChangedEvent&
 
 			// Refresh graph display to display nodes with warning/error
 			GraphEditor->NotifyGraphChanged();
-
-			UpdateDebugDraw();
 		}
 	}
 }
@@ -1379,21 +1381,6 @@ void FDataflowEditorToolkit::OnNodeDoubleClicked(UEdGraphNode* ClickedNode)
 	}
 }
 
-void FDataflowEditorToolkit::ResetViewListeners(UDataflowEdNode* Node)
-{
-	for (IDataflowViewListener* Listener : ViewListeners)
-	{
-		Listener->OnSelectedNodeChanged(nullptr);
-	}
-	if (Node)
-	{
-		for (IDataflowViewListener* Listener : ViewListeners)
-		{
-			Listener->OnSelectedNodeChanged(Node);
-		}
-	}
-};
-
 TSet<TObjectPtr<UDataflowEdNode>> FDataflowEditorToolkit::FilterDataflowEdNodesFromSet(const TSet<UObject*>& Set)
 {
 	TSet<TObjectPtr<UDataflowEdNode>> DataflowEdNodes;
@@ -1441,9 +1428,10 @@ void FDataflowEditorToolkit::OnNodeSelectionChanged(const TSet<UObject*>& InNewS
 
 	const TObjectPtr<UDataflowEdNode> PrevSelectedNode = GetFirstFromSet(PreviouslySelectedNodes);
 	const TObjectPtr<UDataflowEdNode> NewSelectedNode = GetFirstFromSet(SelectedDataflowNodes);
-	ResetViewListeners(NewSelectedNode);
 
 	// todo : move the tool logic somewhere else: we should not be activelly looking for a collection here, we should have some sort of interface to talk to the tools
+	UDataflowEditorMode* const DataflowMode = Cast<UDataflowEditorMode>(EditorModeManager->GetActiveScriptableMode(UDataflowEditorMode::EM_DataflowEditorModeId));
+
 	if (const TObjectPtr<UDataflowBaseContent>& EditorContent = GetEditorContent(); EditorContent->GetDataflowAsset())
 	{
 		if (NewSelectedNode != PrevSelectedNode)
@@ -1453,22 +1441,17 @@ void FDataflowEditorToolkit::OnNodeSelectionChanged(const TSet<UObject*>& InNewS
 
 			if (NewSelectedNode)
 			{
-				if (UDataflowEditorMode* const DataflowMode = Cast<UDataflowEditorMode>(EditorModeManager->GetActiveScriptableMode(UDataflowEditorMode::EM_DataflowEditorModeId)))
+				if (DataflowMode)
 				{
-					// Close any running tool. OnNodeSingleClicked() will start a new tool if a new node was clicked.
+					// Close any running tool by accepting the changes.
+					// OnNodeSingleClicked() will start a new tool if a new node was clicked.
+					DataflowMode->ShutdownActiveToolIfNeeded(EToolShutdownType::Accept);
 
-					UEditorInteractiveToolsContext* const ToolsContext = DataflowMode->GetInteractiveToolsContext();
-					checkf(ToolsContext, TEXT("No valid ToolsContext found for FDataflowEditorToolkit"));
-
-					UInteractiveToolManager* const ToolManager = ToolsContext->ToolManager;
-					checkf(ToolManager, TEXT("ToolsContext doesn't have a valid ToolManager"));
+					EditorContent->SetSelectedNode(NewSelectedNode);
 
-					if (UInteractiveTool* const ActiveTool = ToolManager->GetActiveTool(EToolSide::Left))
-					{
-						ToolManager->PostActiveToolShutdownRequest(ActiveTool, EToolShutdownType::Accept);
-					}
 
-					EditorContent->SetSelectedNode(NewSelectedNode);
+// todo(dataflow) : delete when confirmed this is not causing any issue
+#if 0 // after discussion this does not seem to be used at all in the context of the Dataflow editor - disabling it for now and see if this cause any issues
 
 					// Call the node's OnSelected function. Some nodes use this to cache information from the inputs (e.g. FDataflowCollectionAddScalarVertexPropertyNode::CachedCollectionGroupNames)
 					TSharedPtr<UE::Dataflow::FEngineContext> DataflowContext = EditorContent->GetDataflowContext();
@@ -1489,6 +1472,7 @@ void FDataflowEditorToolkit::OnNodeSelectionChanged(const TSet<UObject*>& InNewS
 							}
 						}
 					}
+#endif
 				}
 			}
 			
@@ -1498,70 +1482,14 @@ void FDataflowEditorToolkit::OnNodeSelectionChanged(const TSet<UObject*>& InNewS
 			}
 		}
 
-		// todo: need a better way to trigger the update of the various views ( should also include what ResetViewListeners does )
-		UpdateDebugDraw();
-		EditorContent->SetConstructionDirty(true);
+		// refresh the views
+		UpdateViewsFromNode(NewSelectedNode);
 	}
 
-	//
 	// Check if the current view mode can render the selected node. If not, try to find a view mode that can. 
-	//
-
-	if (UDataflowEditorMode* const DataflowMode = Cast<UDataflowEditorMode>(EditorModeManager->GetActiveScriptableMode(UDataflowEditorMode::EM_DataflowEditorModeId)))
+	if (DataflowMode)
 	{
-		if (NewSelectedNode && GetEditorContent())
-		{
-			// Check the most recently used view mode for this "kind" of node -- where "kind" means nodes that have the same set of valid view modes
-			auto FindRecentViewMode = [&](FName& OutViewModeName) -> bool
-			{
-				TArray<FName> ViewModesForNode;
-				UE::Dataflow::GetViewModesForNode(*NewSelectedNode, *GetEditorContent(), ViewModesForNode);
-
-				if (const FName* const FoundViewMode = DataflowMode->FindRecentlyUsedViewMode(ViewModesForNode))
-				{
-					OutViewModeName = *FoundViewMode;
-					return true;
-				}
-				return false;
-			};
-
-			// Get the first view mode that this node can use
-			auto FindAnyValidViewMode = [&](FName& OutViewModeName) -> bool
-			{
-				for (const TPair<FName, TUniquePtr<UE::Dataflow::IDataflowConstructionViewMode>>& ViewMode : UE::Dataflow::FRenderingViewModeFactory::GetInstance().GetViewModes())
-				{
-					check(ViewMode.Value.IsValid());
-					const bool bCanRender = UE::Dataflow::CanRenderNodeOutput(*NewSelectedNode, *GetEditorContent(), *ViewMode.Value);
-					if (bCanRender)
-					{
-						OutViewModeName = ViewMode.Key;
-						return true;
-					}
-				}
-			
-				return false;
-			};
-
-			FName NewViewMode;
-			if (FindRecentViewMode(NewViewMode))               // Is there a recently used view mode we can switch to?
-			{
-				DataflowMode->SetConstructionViewMode(NewViewMode);
-			}
-			else if (UE::Dataflow::CanRenderNodeOutput(*NewSelectedNode, *GetEditorContent(), *DataflowMode->GetConstructionViewMode()))        // Can we use the current view mode?
-			{
-				// no need to switch
-			}		
-			else if (FindAnyValidViewMode(NewViewMode))             // Is there *any* view mode that works?
-			{
-				DataflowMode->SetConstructionViewMode(NewViewMode);
-			}
-			else                                                    // No valid view mode found
-			{
-				// TODO: We should clear and disable View Mode Button. For now set default mode to the built-in 3D view mode.
-				DataflowMode->SetConstructionViewMode(UE::Dataflow::FDataflowConstruction3DViewMode::Name);
-			}
-		}
-
+		DataflowMode->SetConstructionViewModeForNode(NewSelectedNode);
 		DataflowMode->SetPendingNodeSelectionChanged(false);
 	}
 }
@@ -1682,24 +1610,30 @@ void FDataflowEditorToolkit::OnSimulationViewSelectionChanged(const TArray<UPrim
 
 void FDataflowEditorToolkit::OnNodeBeginEvaluate(const FDataflowNode* Node, const FDataflowOutput* Output)
 {
-	if (const TObjectPtr<UDataflowBaseContent>& EditorContent = GetEditorContent())
+	if (bDataflowLogNodeEvaluation)
 	{
-		ensure(EditorContent);
-		if (TSharedPtr<UE::Dataflow::FEngineContext> DataflowContext = EditorContent->GetDataflowContext())
+		if (const TObjectPtr<UDataflowBaseContent>& EditorContent = GetEditorContent())
 		{
-			DataflowContext->Info(TEXT("Begin Evaluate"), Node, Output);
+			ensure(EditorContent);
+			if (TSharedPtr<UE::Dataflow::FEngineContext> DataflowContext = EditorContent->GetDataflowContext())
+			{
+				DataflowContext->Info(TEXT("Begin Evaluate"), Node, Output);
+			}
 		}
 	}
 }
 
 void FDataflowEditorToolkit::OnNodeFinishEvaluate(const FDataflowNode* Node, const FDataflowOutput* Output)
 {
-	if (const TObjectPtr<UDataflowBaseContent>& EditorContent = GetEditorContent())
+	if (bDataflowLogNodeEvaluation)
 	{
-		ensure(EditorContent);
-		if (TSharedPtr<UE::Dataflow::FEngineContext> DataflowContext = EditorContent->GetDataflowContext())
+		if (const TObjectPtr<UDataflowBaseContent>& EditorContent = GetEditorContent())
 		{
-			DataflowContext->Info(TEXT("End Evaluate"), Node, Output);
+			ensure(EditorContent);
+			if (TSharedPtr<UE::Dataflow::FEngineContext> DataflowContext = EditorContent->GetDataflowContext())
+			{
+				DataflowContext->Info(TEXT("End Evaluate"), Node, Output);
+			}
 		}
 	}
 }
@@ -1991,6 +1925,8 @@ void FDataflowEditorToolkit::RefreshViewsIfNeeded(bool bForce)
 			Listener->RefreshView();
 		}
 
+		UpdateDebugDraw();
+
 		bViewsNeedRefresh = false;
 	}
 
@@ -2003,6 +1939,21 @@ void FDataflowEditorToolkit::RefreshViewsIfNeeded(bool bForce)
 
 }
 
+void FDataflowEditorToolkit::UpdateViewsFromNode(UDataflowEdNode* Node)
+{
+	for (IDataflowViewListener* Listener : ViewListeners)
+	{
+		Listener->OnSelectedNodeChanged(Node);
+	}
+
+	if (const TObjectPtr<UDataflowBaseContent>& EditorContent = GetEditorContent())
+	{
+		EditorContent->SetConstructionDirty(true);
+	}
+
+	UpdateDebugDraw();
+};
+
 TSharedRef<SDataflowGraphEditor> FDataflowEditorToolkit::CreateGraphEditorWidget(UEdGraph* GraphToEdit, TSharedPtr<IStructureDetailsView> InNodeDetailsEditor)
 {
 	ensure(GraphToEdit);
diff --git a/Engine/Plugins/Experimental/Dataflow/Source/DataflowEditor/Public/Dataflow/DataflowEditorMode.h b/Engine/Plugins/Experimental/Dataflow/Source/DataflowEditor/Public/Dataflow/DataflowEditorMode.h
--- a/Engine/Plugins/Experimental/Dataflow/Source/DataflowEditor/Public/Dataflow/DataflowEditorMode.h
+++ b/Engine/Plugins/Experimental/Dataflow/Source/DataflowEditor/Public/Dataflow/DataflowEditorMode.h
@@ -23,6 +23,7 @@ namespace UE::Dataflow
 {
 	class IDataflowConstructionViewMode;
 }
+enum class EToolShutdownType : uint8;
 
 /**
  * The dataflow editor mode is the mode used in the cloth asset editor. It holds most of the inter-tool state.
@@ -60,6 +61,8 @@ public:
 	* Construction View Mode
 	*/
 	UE_API void SetConstructionViewMode(const FName& NewViewModeName);
+	UE_API void SetConstructionViewModeForNode(const UDataflowEdNode* Node);
+
 	UE_API const UE::Dataflow::IDataflowConstructionViewMode* GetConstructionViewMode() const;
 	UE_API bool CanChangeConstructionViewModeTo(const FName& NewViewModeName) const;
 
@@ -90,6 +93,7 @@ public:
 
 
 	UE_API void StartToolForSelectedNode(const UObject* SelectedNode);
+	UE_API void ShutdownActiveToolIfNeeded(EToolShutdownType ShutdownType);
 
 	/** Each node has a set of view modes it can use. For each distinct set of available view modes, we store the last mode that the user has chosen.
 	 *  Use this function to find the most recently used view mode for nodes having the same set of available modes. 
diff --git a/Engine/Plugins/Experimental/Dataflow/Source/DataflowEditor/Public/Dataflow/DataflowEditorToolkit.h b/Engine/Plugins/Experimental/Dataflow/Source/DataflowEditor/Public/Dataflow/DataflowEditorToolkit.h
--- a/Engine/Plugins/Experimental/Dataflow/Source/DataflowEditor/Public/Dataflow/DataflowEditorToolkit.h
+++ b/Engine/Plugins/Experimental/Dataflow/Source/DataflowEditor/Public/Dataflow/DataflowEditorToolkit.h
@@ -267,7 +267,7 @@ private:
 	/** Update the debug draw based on a change of currently selected or pinned nodes */
 	UE_API void UpdateDebugDraw();
 
-	void ResetViewListeners(UDataflowEdNode* Node = nullptr);
+	void UpdateViewsFromNode(UDataflowEdNode* Node = nullptr);
 	static TSet<TObjectPtr<UDataflowEdNode>> FilterDataflowEdNodesFromSet(const TSet<UObject*>& Set);
 	static const TObjectPtr<UDataflowEdNode> GetFirstFromSet(const TSet<TObjectPtr<UDataflowEdNode>>& Set);
 
