Dataflow : add support for Dataflow Node to execute code when double clicked on
- implement opening the geometry script blueprint when double clicking on the mesh processor nodes

diff --git a/Engine/Plugins/Experimental/Dataflow/Source/DataflowEditor/Private/Dataflow/DataflowSNode.cpp b/Engine/Plugins/Experimental/Dataflow/Source/DataflowEditor/Private/Dataflow/DataflowSNode.cpp
--- a/Engine/Plugins/Experimental/Dataflow/Source/DataflowEditor/Private/Dataflow/DataflowSNode.cpp
+++ b/Engine/Plugins/Experimental/Dataflow/Source/DataflowEditor/Private/Dataflow/DataflowSNode.cpp
@@ -387,25 +387,33 @@ void SDataflowEdNode::UpdateErrorInfo()
 
 FReply SDataflowEdNode::OnMouseButtonDoubleClick(const FGeometry& InMyGeometry, const FPointerEvent& InMouseEvent)
 {
-	if (InMouseEvent.GetModifierKeys().IsControlDown())
+	TSharedPtr<FDataflowNode> Node;
+	if (GraphNode)
 	{
-		if (GraphNode)
+		if (UDataflowEdNode* DataflowNode = Cast<UDataflowEdNode>(GraphNode))
 		{
-			if (UDataflowEdNode* DataflowNode = Cast<UDataflowEdNode>(GraphNode))
+			if (TSharedPtr<UE::Dataflow::FGraph> Graph = DataflowNode->GetDataflowGraph())
 			{
-				if (TSharedPtr<UE::Dataflow::FGraph> Graph = DataflowNode->GetDataflowGraph())
-				{
-					if (TSharedPtr<FDataflowNode> Node = Graph->FindBaseNode(DataflowNode->GetDataflowNodeGuid()))
-					{
-						if (FSourceCodeNavigation::CanNavigateToStruct(Node->TypedScriptStruct()))
-						{
-							FSourceCodeNavigation::NavigateToStruct(Node->TypedScriptStruct());
-						}
-					}
-				}
+				Node = Graph->FindBaseNode(DataflowNode->GetDataflowNodeGuid());
+			}
+		}
+	}
+
+	if (Node)
+	{
+		if (InMouseEvent.GetModifierKeys().IsControlDown())
+		{
+			if (FSourceCodeNavigation::CanNavigateToStruct(Node->TypedScriptStruct()))
+			{
+				FSourceCodeNavigation::NavigateToStruct(Node->TypedScriptStruct());
 			}
 		}
+		else
+		{
+			Node->OnDoubleClicked();
+		}
 	}
+	
 	return Super::OnMouseButtonDoubleClick(InMyGeometry, InMouseEvent);
 }
 
diff --git a/Engine/Plugins/Experimental/GeometryCollectionPlugin/Source/GeometryCollectionNodes/GeometryCollectionNodes.Build.cs b/Engine/Plugins/Experimental/GeometryCollectionPlugin/Source/GeometryCollectionNodes/GeometryCollectionNodes.Build.cs
--- a/Engine/Plugins/Experimental/GeometryCollectionPlugin/Source/GeometryCollectionNodes/GeometryCollectionNodes.Build.cs
+++ b/Engine/Plugins/Experimental/GeometryCollectionPlugin/Source/GeometryCollectionNodes/GeometryCollectionNodes.Build.cs
@@ -44,7 +44,8 @@ namespace UnrealBuildTool.Rules
 				PrivateDependencyModuleNames.AddRange(
 					new string[]
 					{
-					"DataflowEditor",
+						"DataflowEditor",
+						"UnrealEd",
 					}
 					);
 			}
diff --git a/Engine/Plugins/Experimental/GeometryCollectionPlugin/Source/GeometryCollectionNodes/Private/Dataflow/GeometryCollectionMeshNodes.cpp b/Engine/Plugins/Experimental/GeometryCollectionPlugin/Source/GeometryCollectionNodes/Private/Dataflow/GeometryCollectionMeshNodes.cpp
--- a/Engine/Plugins/Experimental/GeometryCollectionPlugin/Source/GeometryCollectionNodes/Private/Dataflow/GeometryCollectionMeshNodes.cpp
+++ b/Engine/Plugins/Experimental/GeometryCollectionPlugin/Source/GeometryCollectionNodes/Private/Dataflow/GeometryCollectionMeshNodes.cpp
@@ -39,6 +39,11 @@
 #include "DynamicMesh/MeshTransforms.h"
 #include "UDynamicMesh.h"
 
+#if WITH_EDITOR
+#include "Editor.h"		// For GEditor
+#include "Subsystems/AssetEditorSubsystem.h"
+#endif
+
 #include UE_INLINE_GENERATED_CPP_BY_NAME(GeometryCollectionMeshNodes)
 
 DEFINE_LOG_CATEGORY_STATIC(LogGeometryCollectionMeshNodes, Warning, All);
@@ -935,6 +940,20 @@ void FMeshProcessorDataflowNodeBase::ApplyParametersToBlueprintInstance(UE::Data
 #endif
 }
 
+#if WITH_EDITOR
+void FMeshProcessorDataflowNodeBase::OnDoubleClicked() const
+{
+	if (UAssetEditorSubsystem* AssetEditorSubsystem = GEditor->GetEditorSubsystem<UAssetEditorSubsystem>())
+	{
+		if (UBlueprint* Blueprint = Cast<UBlueprint>(MeshProcessor->ClassGeneratedBy))
+		{
+			AssetEditorSubsystem->OpenEditorForAssets({ Blueprint });
+		}
+	}
+}
+#endif
+
+
 ///////////////////////////////////////////////////////////////////////////////////////////////////////////
 
 void FApplyMeshProcessorToMeshDataflowNode::Evaluate(UE::Dataflow::FContext& Context, const FDataflowOutput* Out) const
diff --git a/Engine/Plugins/Experimental/GeometryCollectionPlugin/Source/GeometryCollectionNodes/Public/Dataflow/GeometryCollectionMeshNodes.h b/Engine/Plugins/Experimental/GeometryCollectionPlugin/Source/GeometryCollectionNodes/Public/Dataflow/GeometryCollectionMeshNodes.h
--- a/Engine/Plugins/Experimental/GeometryCollectionPlugin/Source/GeometryCollectionNodes/Public/Dataflow/GeometryCollectionMeshNodes.h
+++ b/Engine/Plugins/Experimental/GeometryCollectionPlugin/Source/GeometryCollectionNodes/Public/Dataflow/GeometryCollectionMeshNodes.h
@@ -120,6 +120,10 @@ private:
 		return true;
 	}
 
+#if WITH_EDITOR
+	virtual void OnDoubleClicked() const override;
+#endif
+
 	// Handling for the selected blueprint being changed under the node (e.g., recompiled)
 	FDelegateHandle BlueprintChangeDelegateHandle;
 	void TeardownBlueprintEvent();
diff --git a/Engine/Source/Runtime/Experimental/Dataflow/Core/Public/Dataflow/DataflowNode.h b/Engine/Source/Runtime/Experimental/Dataflow/Core/Public/Dataflow/DataflowNode.h
--- a/Engine/Source/Runtime/Experimental/Dataflow/Core/Public/Dataflow/DataflowNode.h
+++ b/Engine/Source/Runtime/Experimental/Dataflow/Core/Public/Dataflow/DataflowNode.h
@@ -150,6 +150,8 @@ PRAGMA_ENABLE_DEPRECATION_WARNINGS
 	virtual void DebugDraw(UE::Dataflow::FContext& Context,
 		IDataflowDebugDrawInterface& DataflowRenderingInterface,
 		const FDebugDrawParameters& DebugDrawParameters) const {}
+
+	virtual void OnDoubleClicked() const {};
 #endif
 
 	//
