Dataflow Editor - Fixed dangling pointers issue with viewports closing order.
- Tidied up hacky fix made for hotfix 5.6.1.

#rnx

diff --git a/Engine/Plugins/ChaosClothAssetEditor/Source/ChaosClothAssetEditor/Private/ChaosClothAsset/ClothEditorToolkit.cpp b/Engine/Plugins/ChaosClothAssetEditor/Source/ChaosClothAssetEditor/Private/ChaosClothAsset/ClothEditorToolkit.cpp
--- a/Engine/Plugins/ChaosClothAssetEditor/Source/ChaosClothAssetEditor/Private/ChaosClothAsset/ClothEditorToolkit.cpp
+++ b/Engine/Plugins/ChaosClothAssetEditor/Source/ChaosClothAssetEditor/Private/ChaosClothAsset/ClothEditorToolkit.cpp
@@ -222,7 +222,7 @@ FChaosClothAssetEditorToolkit::FChaosClothAssetEditorToolkit(UAssetEditor* InOwn
 	};
 
 	// Construction view scene
-	ObjectScene = MakeUnique<FPreviewScene>(FPreviewScene::ConstructionValues().SetSkyBrightness(0.0f).SetLightBrightness(0.0f));
+	ObjectScene = MakeShared<FPreviewScene>(FPreviewScene::ConstructionValues().SetSkyBrightness(0.0f).SetLightBrightness(0.0f));
 }
 
 FChaosClothAssetEditorToolkit::~FChaosClothAssetEditorToolkit()
diff --git a/Engine/Plugins/Experimental/CharacterFXEditor/BaseCharacterFXEditor/Source/BaseCharacterFXEditor/Private/BaseCharacterFXEditorToolkit.cpp b/Engine/Plugins/Experimental/CharacterFXEditor/BaseCharacterFXEditor/Source/BaseCharacterFXEditor/Private/BaseCharacterFXEditorToolkit.cpp
--- a/Engine/Plugins/Experimental/CharacterFXEditor/BaseCharacterFXEditor/Source/BaseCharacterFXEditor/Private/BaseCharacterFXEditorToolkit.cpp
+++ b/Engine/Plugins/Experimental/CharacterFXEditor/BaseCharacterFXEditor/Source/BaseCharacterFXEditor/Private/BaseCharacterFXEditorToolkit.cpp
@@ -71,7 +71,7 @@ FBaseCharacterFXEditorToolkit::FBaseCharacterFXEditorToolkit(UAssetEditor* InOwn
 	// We could create the preview scene in CreateEditorViewportClient() the way that FBaseAssetToolkit
 	// does, but it seems more intuitive to create them right off the bat and pass it in later. 
 	FPreviewScene::ConstructionValues SceneArgs;
-	ObjectScene = MakeUnique<FAdvancedPreviewScene>(SceneArgs);
+	ObjectScene = MakeShared<FAdvancedPreviewScene>(SceneArgs);
 }
 
 void FBaseCharacterFXEditorToolkit::OnToolkitHostingStarted(const TSharedRef<IToolkit>& Toolkit)
diff --git a/Engine/Plugins/Experimental/CharacterFXEditor/BaseCharacterFXEditor/Source/BaseCharacterFXEditor/Public/BaseCharacterFXEditorToolkit.h b/Engine/Plugins/Experimental/CharacterFXEditor/BaseCharacterFXEditor/Source/BaseCharacterFXEditor/Public/BaseCharacterFXEditorToolkit.h
--- a/Engine/Plugins/Experimental/CharacterFXEditor/BaseCharacterFXEditor/Source/BaseCharacterFXEditor/Public/BaseCharacterFXEditorToolkit.h
+++ b/Engine/Plugins/Experimental/CharacterFXEditor/BaseCharacterFXEditor/Source/BaseCharacterFXEditor/Public/BaseCharacterFXEditorToolkit.h
@@ -2,7 +2,6 @@
 
 #pragma once
 
-#include "CoreMinimal.h"
 #include "Tools/BaseAssetToolkit.h"
 #include "BaseCharacterFXEditorModeUILayer.h"
 #include "PreviewScene.h"
@@ -61,7 +60,7 @@ protected:
 	TSharedPtr<FBaseCharacterFXEditorModeUILayer> ModeUILayer;
 
 	// PreviewScene showing the objects being edited
-	TUniquePtr<FPreviewScene> ObjectScene;
+	TSharedPtr<FPreviewScene> ObjectScene;
 
 	// A default "CharacterFXEditor" category is created in RegisterTabSpawners. Override that function to set this to a concrete class type
 	TSharedPtr<FWorkspaceItem> EditorMenuCategory;
diff --git a/Engine/Plugins/Experimental/Dataflow/Source/DataflowEditor/Private/Dataflow/DataflowConstructionViewport.cpp b/Engine/Plugins/Experimental/Dataflow/Source/DataflowEditor/Private/Dataflow/DataflowConstructionViewport.cpp
--- a/Engine/Plugins/Experimental/Dataflow/Source/DataflowEditor/Private/Dataflow/DataflowConstructionViewport.cpp
+++ b/Engine/Plugins/Experimental/Dataflow/Source/DataflowEditor/Private/Dataflow/DataflowConstructionViewport.cpp
@@ -34,7 +34,7 @@ void SDataflowConstructionViewport::Construct(const FArguments& InArgs, const FA
 	Client->VisibilityDelegate.BindSP(this, &SDataflowConstructionViewport::IsVisible);
 }
 
-FDataflowConstructionScene* SDataflowConstructionViewport::GetConstructionScene() const
+TWeakPtr<FDataflowConstructionScene> SDataflowConstructionViewport::GetConstructionScene() const
 {
 	const TSharedPtr<FDataflowConstructionViewportClient> DataflowClient = StaticCastSharedPtr<FDataflowConstructionViewportClient>(Client);
 	if (DataflowClient)
diff --git a/Engine/Plugins/Experimental/Dataflow/Source/DataflowEditor/Private/Dataflow/DataflowConstructionViewportClient.cpp b/Engine/Plugins/Experimental/Dataflow/Source/DataflowEditor/Private/Dataflow/DataflowConstructionViewportClient.cpp
--- a/Engine/Plugins/Experimental/Dataflow/Source/DataflowEditor/Private/Dataflow/DataflowConstructionViewportClient.cpp
+++ b/Engine/Plugins/Experimental/Dataflow/Source/DataflowEditor/Private/Dataflow/DataflowConstructionViewportClient.cpp
@@ -14,6 +14,7 @@
 #include "Dataflow/DataflowEditorPreviewSceneBase.h"
 #include "Dataflow/DataflowRenderingViewMode.h"
 #include "EditorModeManager.h"
+#include "Framework/Application/SlateApplication.h"
 #include "Tools/EdModeInteractiveToolsContext.h"
 #include "GraphEditor.h"
 #include "PreviewScene.h"
@@ -26,9 +27,10 @@
 #define LOCTEXT_NAMESPACE "DataflowConstructionViewportClient"
 
 FDataflowConstructionViewportClient::FDataflowConstructionViewportClient(FEditorModeTools* InModeTools,
-                                                             FPreviewScene* InPreviewScene,  const bool bCouldTickScene,
+                                                             TWeakPtr<FDataflowConstructionScene> InConstructionScene, const bool bCouldTickScene,
                                                              const TWeakPtr<SEditorViewport> InEditorViewportWidget)
-	: FDataflowEditorViewportClientBase(InModeTools, InPreviewScene, bCouldTickScene, InEditorViewportWidget)
+	: FDataflowEditorViewportClientBase(InModeTools, InConstructionScene, bCouldTickScene, InEditorViewportWidget)
+	, ConstructionScene(InConstructionScene)
 {
 	// We want our near clip plane to be quite close so that we can zoom in further.
 	OverrideNearClipPlane(KINDA_SMALL_NUMBER);
@@ -39,7 +41,6 @@ FDataflowConstructionViewportClient::FDataflowConstructionViewportClient(FEditor
 	EngineShowFlags.SetSelectionOutline(true);
 	EngineShowFlags.EnableAdvancedFeatures();
 
-	PreviewScene = static_cast<FDataflowPreviewSceneBase*>(InPreviewScene);
 	bEnableSceneTicking = bCouldTickScene;
 
 	bool bUseRightMouseButton = true;
@@ -111,16 +112,22 @@ void FDataflowConstructionViewportClient::SetToolCommandList(TWeakPtr<FUICommand
 void FDataflowConstructionViewportClient::Tick(float DeltaSeconds)
 {
 	Super::Tick(DeltaSeconds);
-	if (PreviewScene)
+	if (const TSharedPtr<FDataflowConstructionScene> DataflowConstructionScenePtr = ConstructionScene.Pin())
 	{
-		PreviewScene->TickDataflowScene(DeltaSeconds);
+		DataflowConstructionScenePtr->TickDataflowScene(DeltaSeconds);
 	}
 }
 
 USelection* FDataflowConstructionViewportClient::GetSelectedComponents() 
 {
-	USelection* SceneSelection = FDataflowPreviewSceneBase::GetSelectedComponents(PreviewScene);
-	return SceneSelection ? SceneSelection : ModeTools->GetSelectedComponents(); 
+	if (const TSharedPtr<FDataflowConstructionScene> DataflowConstructionScenePtr = ConstructionScene.Pin())
+	{
+		if (USelection* const SceneSelection = DataflowConstructionScenePtr->GetSelectedComponents())
+		{
+			return SceneSelection;
+		}
+	}
+	return ModeTools->GetSelectedComponents();
 }
 
 bool FDataflowConstructionViewportClient::InputKey(const FInputKeyEventArgs& EventArgs)
@@ -152,9 +159,10 @@ void FDataflowConstructionViewportClient::OnViewportClicked(HHitProxy* HitProxy)
 	{
 		if (TSharedPtr<FDataflowEditorToolkit> DataflowEditorToolkit = DataflowEditorToolkitPtr.Pin())
 		{
-			if (PreviewScene && PreviewScene->GetDataflowModeManager())
+			const TSharedPtr<FDataflowConstructionScene> ConstructionScenePtr = ConstructionScene.Pin();
+			if (ConstructionScenePtr && ConstructionScenePtr->GetDataflowModeManager())
 			{
-				if (UDataflowEditorMode* DataflowMode = Cast<UDataflowEditorMode>(PreviewScene->GetDataflowModeManager()->GetActiveScriptableMode(UDataflowEditorMode::EM_DataflowEditorModeId)))
+				if (UDataflowEditorMode* DataflowMode = Cast<UDataflowEditorMode>(ConstructionScenePtr->GetDataflowModeManager()->GetActiveScriptableMode(UDataflowEditorMode::EM_DataflowEditorModeId)))
 				{
 					if (TSharedPtr<SDataflowGraphEditor> GraphEditor = DataflowEditorToolkit->GetDataflowGraphEditor())
 					{
@@ -213,12 +221,15 @@ void FDataflowConstructionViewportClient::OnViewportClicked(HHitProxy* HitProxy)
 	
 	auto IsInteractiveToolActive = [&]()
 	{
-		if (UDataflowEditorMode* DataflowMode = Cast<UDataflowEditorMode>(
-			PreviewScene->GetDataflowModeManager()->GetActiveScriptableMode(UDataflowEditorMode::EM_DataflowEditorModeId)))
+		if (const TSharedPtr<FDataflowConstructionScene> ConstructionScenePtr = ConstructionScene.Pin())
 		{
-			if (UEditorInteractiveToolsContext* const ToolsContext = DataflowMode->GetInteractiveToolsContext())
+			if (UDataflowEditorMode* const DataflowMode = Cast<UDataflowEditorMode>(
+				ConstructionScenePtr->GetDataflowModeManager()->GetActiveScriptableMode(UDataflowEditorMode::EM_DataflowEditorModeId)))
 			{
-				return ToolsContext->HasActiveTool();
+				if (UEditorInteractiveToolsContext* const ToolsContext = DataflowMode->GetInteractiveToolsContext())
+				{
+					return ToolsContext->HasActiveTool();
+				}
 			}
 		}
 		return false;
@@ -314,10 +325,13 @@ void FDataflowConstructionViewportClient::Draw(const FSceneView* View, FPrimitiv
 {
 	FDataflowEditorViewportClientBase::Draw(View, PDI);
 
-	using namespace UE::Dataflow;
-	for (const TPair<FName, TUniquePtr<IDataflowConstructionVisualization>>& Visualization : FDataflowConstructionVisualizationRegistry::GetInstance().GetVisualizations())
+	if (const TSharedPtr<FDataflowConstructionScene> ConstructionScenePtr = ConstructionScene.Pin())
 	{
-		Visualization.Value->Draw(static_cast<FDataflowConstructionScene*>(PreviewScene), PDI, View);
+		using namespace UE::Dataflow;
+		for (const TPair<FName, TUniquePtr<IDataflowConstructionVisualization>>& Visualization : FDataflowConstructionVisualizationRegistry::GetInstance().GetVisualizations())
+		{
+			Visualization.Value->Draw(ConstructionScenePtr.Get(), PDI, View);
+		}
 	}
 }
 
@@ -325,10 +339,13 @@ void FDataflowConstructionViewportClient::DrawCanvas(FViewport& InViewport, FSce
 {
 	FEditorViewportClient::DrawCanvas(InViewport, View, Canvas);
 
-	using namespace UE::Dataflow;
-	for (const TPair<FName, TUniquePtr<IDataflowConstructionVisualization>>& Visualization : FDataflowConstructionVisualizationRegistry::GetInstance().GetVisualizations())
+	if (const TSharedPtr<FDataflowConstructionScene> ConstructionScenePtr = ConstructionScene.Pin())
 	{
-		Visualization.Value->DrawCanvas(static_cast<FDataflowConstructionScene*>(PreviewScene), &Canvas, &View);
+		using namespace UE::Dataflow;
+		for (const TPair<FName, TUniquePtr<IDataflowConstructionVisualization>>& Visualization : FDataflowConstructionVisualizationRegistry::GetInstance().GetVisualizations())
+		{
+			Visualization.Value->DrawCanvas(ConstructionScenePtr.Get(), &Canvas, &View);
+		}
 	}
 }
 
diff --git a/Engine/Plugins/Experimental/Dataflow/Source/DataflowEditor/Private/Dataflow/DataflowEditorMode.cpp b/Engine/Plugins/Experimental/Dataflow/Source/DataflowEditor/Private/Dataflow/DataflowEditorMode.cpp
--- a/Engine/Plugins/Experimental/Dataflow/Source/DataflowEditor/Private/Dataflow/DataflowEditorMode.cpp
+++ b/Engine/Plugins/Experimental/Dataflow/Source/DataflowEditor/Private/Dataflow/DataflowEditorMode.cpp
@@ -170,8 +170,10 @@ void UDataflowEditorMode::RegisterDataflowTool(TSharedPtr<FUICommandInfo> UIComm
 						bShouldRestoreSavedConstructionViewMode = true;
 					}
 
+					const TSharedPtr<FDataflowConstructionScene> ConstructionScenePtr = ConstructionScene.Pin();
+
 					bool bHadSingleSelectedMesh = false;
-					if (const USelection* const SelectedComponents = FDataflowPreviewSceneBase::GetSelectedComponents(ConstructionScene))
+					if (const USelection* const SelectedComponents = ConstructionScenePtr ? ConstructionScenePtr->GetSelectedComponents() : nullptr)
 					{
 						bHadSingleSelectedMesh = (SelectedComponents->Num() == 1);
 					}
@@ -182,16 +184,15 @@ void UDataflowEditorMode::RegisterDataflowTool(TSharedPtr<FUICommandInfo> UIComm
 					if (bHadSingleSelectedMesh)
 					{
 						// If there is a single dynamic mesh component in the scene, select it so the tool can start
-						const TArray<UDynamicMeshComponent*> DynamicMeshComponents = ConstructionScene->GetDynamicMeshComponents();
+						const TArray<TObjectPtr<UDynamicMeshComponent>> DynamicMeshComponents = ConstructionScenePtr ? ConstructionScenePtr->GetDynamicMeshComponents() : TArray<TObjectPtr<UDynamicMeshComponent>>();
 						if (DynamicMeshComponents.Num() == 1)
 						{
-							if (USelection* const SelectedComponents = FDataflowPreviewSceneBase::GetSelectedComponents(ConstructionScene))
+							if (USelection* const SelectedComponents = ConstructionScenePtr->GetSelectedComponents())
 							{
 								SelectedComponents->Select(DynamicMeshComponents[0]);
 							}
 						}
 					}
-
 				}
 			}
 
@@ -268,9 +269,9 @@ const TArray<FName>& UDataflowEditorMode::GetToolCategories()
 	const UDataflowEditor* LocalEditor = DataflowEditor;
 	if(!LocalEditor)
 	{
-		if(ConstructionScene)
+		if (const TSharedPtr<FDataflowConstructionScene> ConstructionScenePtr = ConstructionScene.Pin())
 		{
-			LocalEditor = ConstructionScene->GetDataflowEditor();
+			LocalEditor = ConstructionScenePtr->GetDataflowEditor();
 		}
 		else if(FAssetEditorModeManager* ModeManager = StaticCast<FAssetEditorModeManager*>(Owner))
 		{
@@ -329,15 +330,15 @@ bool UDataflowEditorMode::ShouldToolStartBeAllowed(const FString& ToolIdentifier
 		return false;
 	}
 
-	if (ConstructionScene && ConstructionScene->GetDataflowModeManager() && ConstructionScene->GetDataflowModeManager()->GetInteractiveToolsContext())
+	const TSharedPtr<FDataflowConstructionScene> ConstructionScenePtr = ConstructionScene.Pin();
+	if (ConstructionScenePtr && ConstructionScenePtr->GetDataflowModeManager() && ConstructionScenePtr->GetDataflowModeManager()->GetInteractiveToolsContext())
 	{
-		if (ConstructionScene->GetDataflowModeManager()->GetInteractiveToolsContext()->HasActiveTool())
+		if (ConstructionScenePtr->GetDataflowModeManager()->GetInteractiveToolsContext()->HasActiveTool())
 		{
 			return false;
 		}
 	}
 
-
 	if (UE::Dataflow::Private::bDataflowEditorEnableToolsInPIE)
 	{
 		// UEdMode::ShouldToolStartBeAllowed returns (!GEditor->PlayWorld && !GIsPlayInEditorWorld) but we want to allow tools to start while in PIE
@@ -357,7 +358,8 @@ void UDataflowEditorMode::CreateToolkit()
 
 void UDataflowEditorMode::SetWireframeRenderToggleEnabled(bool bEnable)
 {
-	if (const TObjectPtr<UDataflowBaseContent> EditorContent = ConstructionScene->GetEditorContent())
+	const TSharedPtr<FDataflowConstructionScene> ConstructionScenePtr = ConstructionScene.Pin();
+	if (const TObjectPtr<UDataflowBaseContent> EditorContent = ConstructionScenePtr ? ConstructionScenePtr->GetEditorContent() : nullptr)
 	{
 		if (const TObjectPtr<UDataflow> DataflowGraph = EditorContent->GetDataflowAsset())
 		{
@@ -395,9 +397,9 @@ void UDataflowEditorMode::OnToolEnded(UInteractiveToolManager* Manager, UInterac
 		SetConstructionViewMode(SavedConstructionViewMode);
 		bShouldRestoreSavedConstructionViewMode = false;
 	}
-	else
+	else if (const TSharedPtr<FDataflowConstructionScene> ConstructionScenePtr = ConstructionScene.Pin())
 	{
-		ConstructionScene->ResetConstructionScene();
+		ConstructionScenePtr->ResetConstructionScene();
 	}
 
 	if (TSharedPtr<SDataflowGraphEditor> GraphEditor = DataflowGraphEditor.Pin())
@@ -440,16 +442,14 @@ void UDataflowEditorMode::Exit()
 {
 	UActorComponent::MarkRenderStateDirtyEvent.RemoveAll(this);
 
-	if(ConstructionScene)
+	if (const TSharedPtr<FDataflowConstructionScene> ConstructionScenePtr = ConstructionScene.Pin())
 	{
-		ConstructionScene->ResetConstructionScene();
-		ConstructionScene = nullptr;
+		ConstructionScenePtr->ResetConstructionScene();
 	}
 
-	if(SimulationScene)
+	if (const TSharedPtr<FDataflowSimulationScene> SimulationScenePtr = SimulationScene.Pin())
 	{
-		SimulationScene->ResetSimulationScene();
-		SimulationScene = nullptr;
+		SimulationScenePtr->ResetSimulationScene();
 	}
 
 	// Log mode exit
@@ -467,50 +467,52 @@ void UDataflowEditorMode::Exit()
 	Super::Exit();
 }
 
-void UDataflowEditorMode::SetDataflowConstructionScene(FDataflowConstructionScene* InConstructionScene)
+void UDataflowEditorMode::SetDataflowConstructionScene(TWeakPtr<FDataflowConstructionScene> InConstructionScene)
 {
 	ConstructionScene = InConstructionScene;
 
-	UEditorInteractiveToolsContext* const PreviewToolsContext = ConstructionScene->GetDataflowModeManager()->GetInteractiveToolsContext();
-	UInteractiveToolManager* const PreviewToolManager = PreviewToolsContext->ToolManager;
-	//#todo(brice): Make sure AddToolTargetFactories has been called. 
-	//PreviewToolsContext->TargetManager->AddTargetFactory(NewObject<UClothComponentToolTargetFactory>(PreviewToolManager));
-	//PreviewToolsContext->TargetManager->AddTargetFactory(NewObject<USkeletalMeshComponentToolTargetFactory>(PreviewToolManager));
-
-	PreviewToolManager->OnToolStarted.AddUObject(this, &UDataflowEditorMode::OnToolStarted);
-	PreviewToolManager->OnToolEnded.AddUObject(this, &UDataflowEditorMode::OnToolEnded);
+	if (const TSharedPtr<FDataflowConstructionScene> ConstructionScenePtr = ConstructionScene.Pin())
+	{
+		UEditorInteractiveToolsContext* const PreviewToolsContext = ConstructionScenePtr->GetDataflowModeManager()->GetInteractiveToolsContext();
+		UInteractiveToolManager* const PreviewToolManager = PreviewToolsContext->ToolManager;
+		//#todo(brice): Make sure AddToolTargetFactories has been called. 
+		//PreviewToolsContext->TargetManager->AddTargetFactory(NewObject<UClothComponentToolTargetFactory>(PreviewToolManager));
+		//PreviewToolsContext->TargetManager->AddTargetFactory(NewObject<USkeletalMeshComponentToolTargetFactory>(PreviewToolManager));
 
-	check(Toolkit.IsValid());
+		PreviewToolManager->OnToolStarted.AddUObject(this, &UDataflowEditorMode::OnToolStarted);
+		PreviewToolManager->OnToolEnded.AddUObject(this, &UDataflowEditorMode::OnToolEnded);
 
-	// FBaseToolkit's OnToolStarted and OnToolEnded are protected, so we use the subclass to get at them
-	FDataflowEditorModeToolkit* const DataflowModeToolkit = static_cast<FDataflowEditorModeToolkit*>(Toolkit.Get());
+		check(Toolkit.IsValid());
 
-	PreviewToolManager->OnToolStarted.AddSP(DataflowModeToolkit, &FDataflowEditorModeToolkit::OnToolStarted);
-	PreviewToolManager->OnToolEnded.AddSP(DataflowModeToolkit, &FDataflowEditorModeToolkit::OnToolEnded);
-}
+		// FBaseToolkit's OnToolStarted and OnToolEnded are protected, so we use the subclass to get at them
+		FDataflowEditorModeToolkit* const DataflowModeToolkit = static_cast<FDataflowEditorModeToolkit*>(Toolkit.Get());
 
-void UDataflowEditorMode::SetDataflowSimlationScene(FDataflowSimulationScene* InSimulationScene)
-{
-	SimulationScene = InSimulationScene;
+		PreviewToolManager->OnToolStarted.AddSP(DataflowModeToolkit, &FDataflowEditorModeToolkit::OnToolStarted);
+		PreviewToolManager->OnToolEnded.AddSP(DataflowModeToolkit, &FDataflowEditorModeToolkit::OnToolEnded);
+	}
 }
 
 void UDataflowEditorMode::CreateToolTargets(const TArray<TObjectPtr<UObject>>& AssetsIn)
 {
 	ToolTargets.Reset();
-	if (const TObjectPtr<UDataflowBaseContent>& EditorContent = ConstructionScene->GetEditorContent())
+	if (const TSharedPtr<FDataflowConstructionScene> ConstructionScenePtr = ConstructionScene.Pin())
 	{
-		if (UToolTarget* Target = GetInteractiveToolsContext()->TargetManager->BuildTarget(EditorContent, GetToolTargetRequirements()))
+		if (const TObjectPtr<UDataflowBaseContent>& EditorContent = ConstructionScenePtr->GetEditorContent())
 		{
-			ToolTargets.Add(Target);
+			if (UToolTarget* const Target = GetInteractiveToolsContext()->TargetManager->BuildTarget(EditorContent, GetToolTargetRequirements()))
+			{
+				ToolTargets.Add(Target);
+			}
 		}
 	}
 }
 
 bool UDataflowEditorMode::IsComponentSelected(const UPrimitiveComponent* InComponent)
 {
-	if (ConstructionScene && ConstructionScene->GetDataflowModeManager())
+	const TSharedPtr<FDataflowConstructionScene> ConstructionScenePtr = ConstructionScene.Pin();
+	if (ConstructionScenePtr&& ConstructionScenePtr->GetDataflowModeManager())
 	{
-		if (const UTypedElementSelectionSet* const TypedElementSelectionSet = ConstructionScene->GetDataflowModeManager()->GetEditorSelectionSet())
+		if (const UTypedElementSelectionSet* const TypedElementSelectionSet = ConstructionScenePtr->GetDataflowModeManager()->GetEditorSelectionSet())
 		{
 			if (const FTypedElementHandle ComponentElement = UEngineElementsLibrary::AcquireEditorComponentElementHandle(InComponent))
 			{
@@ -542,30 +544,31 @@ void UDataflowEditorMode::RefocusConstructionViewportClient()
 
 void UDataflowEditorMode::RefocusSimulationViewportClient()
 {
-	TSharedPtr<FDataflowSimulationViewportClient, ESPMode::ThreadSafe> PinnedVC = SimulationViewportClient.Pin();
-	if (PinnedVC.IsValid() && SimulationScene)
+	if (const TSharedPtr<FDataflowSimulationScene> SimulationScenePtr = SimulationScene.Pin())
 	{
-		// This will happen in FocusViewportOnBox anyways; do it now to get a consistent end result
-		PinnedVC->ToggleOrbitCamera(false);
+		if (const TSharedPtr<FDataflowSimulationViewportClient> PinnedVC = SimulationViewportClient.Pin())
+		{
+			// This will happen in FocusViewportOnBox anyways; do it now to get a consistent end result
+			PinnedVC->ToggleOrbitCamera(false);
 
-		const FBox SceneBounds = SimulationScene->GetBoundingBox();
+			const FBox SceneBounds = SimulationScenePtr->GetBoundingBox();
 
-		// Set up camera for an angled view by default 
-		PinnedVC->SetInitialViewTransform(ELevelViewportType::LVT_Perspective, FVector(0.0, 0.0, 0.0), FRotator(-15.0, -40.0, 0.0), DEFAULT_ORTHOZOOM);
+			// Set up camera for an angled view by default 
+			PinnedVC->SetInitialViewTransform(ELevelViewportType::LVT_Perspective, FVector(0.0, 0.0, 0.0), FRotator(-15.0, -40.0, 0.0), DEFAULT_ORTHOZOOM);
 
-		constexpr bool bInstant = true;
-		PinnedVC->FocusViewportOnBox(SceneBounds, bInstant);
+			constexpr bool bInstant = true;
+			PinnedVC->FocusViewportOnBox(SceneBounds, bInstant);
+		}
 	}
 }
 
 void UDataflowEditorMode::FirstTimeFocusConstructionViewport()
 {
 	// If this is the first time seeing a valid 2D or 3D mesh, refocus the camera on it.
-	const bool bIsValid = (ConstructionScene->HasRenderableGeometry());
-	const bool bIs2D = !ConstructionViewMode->IsPerspective();
-
-	if (bIsValid)
+	const TSharedPtr<FDataflowConstructionScene> ConstructionScenePtr = ConstructionScene.Pin();
+	if (ConstructionScenePtr && ConstructionScenePtr->HasRenderableGeometry())
 	{
+		const bool bIs2D = !ConstructionViewMode->IsPerspective();
 		if (bIs2D && bFirstValid2DMesh)
 		{
 			bFirstValid2DMesh = false;
@@ -582,9 +585,8 @@ void UDataflowEditorMode::FirstTimeFocusConstructionViewport()
 void UDataflowEditorMode::FirstTimeFocusSimulationViewport()
 {
 	// If this is the first time seeing a valid 2D or 3D mesh, refocus the camera on it.
-	const bool bIsValid = (SimulationScene->HasRenderableGeometry());
-
-	if (bIsValid)
+	const TSharedPtr<FDataflowSimulationScene> SimulationScenePtr = SimulationScene.Pin();
+	if (SimulationScenePtr && SimulationScenePtr->HasRenderableGeometry())
 	{
 		RefocusSimulationViewportClient();
 	}
@@ -603,7 +605,10 @@ void UDataflowEditorMode::InitializeTargets(const TArray<TObjectPtr<UObject>>& O
 	// ... }
 
 	// @todo(michael) : do we need to update the construction scene?
-	ConstructionScene->UpdateConstructionScene();
+	if (const TSharedPtr<FDataflowConstructionScene> ConstructionScenePtr = ConstructionScene.Pin())
+	{
+		ConstructionScenePtr->UpdateConstructionScene();
+	}
 }
 
 void UDataflowEditorMode::ModeTick(float DeltaTime)
@@ -624,14 +629,12 @@ void UDataflowEditorMode::ModeTick(float DeltaTime)
 		}
 	}
 
-	checkf(ConstructionScene, TEXT("Null ConstructionScene in DataflowEditorMode"));
-	if (const UDataflowBaseContent* const EditorContent = ConstructionScene->GetEditorContent())
+	const TSharedPtr<FDataflowConstructionScene> ConstructionScenePtr = ConstructionScene.Pin();
+	const UDataflowBaseContent* const EditorContent = ConstructionScenePtr ? ConstructionScenePtr->GetEditorContent() : nullptr;
+	if (EditorContent && EditorContent->IsConstructionDirty())
 	{
-		if (EditorContent->IsConstructionDirty())
-		{
-			ConstructionScene->UpdateConstructionScene();
-			FirstTimeFocusConstructionViewport();
-		}
+		ConstructionScenePtr->UpdateConstructionScene();
+		FirstTimeFocusConstructionViewport();
 	}
 
 	if (!NodeTypeForPendingToolStart.IsNone() && !GetToolManager()->HasActiveTool(EToolSide::Left) && !bPendingNodeSelectionChanged)
@@ -651,19 +654,15 @@ void UDataflowEditorMode::ModeTick(float DeltaTime)
 	if (bShouldRestartToolNextTick)
 	{
 		// If we ended the active tool in order to change view mode, restart it now
-
-		const UDataflowBaseContent* const EditorContent = ConstructionScene->GetEditorContent();
-		checkf(EditorContent, TEXT("Expected EditorContent in ConstructionScene"));
-
-		if (!EditorContent->IsConstructionDirty())		// hold off restarting the tool until the scene finishes rebuilding
+		if (EditorContent && !EditorContent->IsConstructionDirty())		// hold off restarting the tool until the scene finishes rebuilding
 		{
 			// First select the lone mesh in the construction scene if there is one
 			if (bHadSingleSelectionBeforeToolShutdown)
 			{
-				const TArray<UDynamicMeshComponent*> DynamicMeshComponents = ConstructionScene->GetDynamicMeshComponents();
+				const TArray<TObjectPtr<UDynamicMeshComponent>> DynamicMeshComponents = ConstructionScenePtr ? ConstructionScenePtr->GetDynamicMeshComponents() : TArray<TObjectPtr<UDynamicMeshComponent>>();
 				if (DynamicMeshComponents.Num() == 1)
 				{
-					if (USelection* const SelectedComponents = FDataflowPreviewSceneBase::GetSelectedComponents(ConstructionScene))
+					if (USelection* const SelectedComponents = ConstructionScenePtr->GetSelectedComponents())
 					{
 						SelectedComponents->Select(DynamicMeshComponents[0]);
 					}
@@ -723,7 +722,8 @@ void UDataflowEditorMode::ViewportResized(FViewport* Viewport, uint32 /*Unused*/
 
 FBox UDataflowEditorMode::SceneBoundingBox() const
 {
-	return ConstructionScene->GetBoundingBox();
+	const TSharedPtr<FDataflowConstructionScene> ConstructionScenePtr = ConstructionScene.Pin();
+	return ConstructionScenePtr ? ConstructionScenePtr->GetBoundingBox() : FBox(FVector(-100.), FVector(100.));
 }
 
 FBox UDataflowEditorMode::SelectionBoundingBox() const
@@ -739,14 +739,8 @@ FBox UDataflowEditorMode::SelectionBoundingBox() const
 		}
 	}
 
-	// If the selection is on the GetBoundingBox is automatically computing the selection one
-	FBox Bounds = ConstructionScene->GetBoundingBox();
-	if (Bounds.IsValid)
-	{
-		return Bounds;
-	}
-
-	// Nothing selected, return the whole scene
+	// If the selection is on, the GetBoundingBox is automatically computing the selection one
+	// Otherwise if nothing selected, it will return the whole scene
 	return SceneBoundingBox();
 }
 
@@ -760,7 +754,7 @@ void UDataflowEditorMode::SetConstructionViewMode(const FName& NewViewModeName)
 	const TObjectPtr<UInteractiveToolManager> ToolManager = GetInteractiveToolsContext()->ToolManager;
 	checkf(ToolManager, TEXT("No valid ToolManager found for UDataflowEditorMode"));
 
-	checkf(ConstructionScene, TEXT("Expected UDataflowEditorMode::ConstructionScene to have been initialized"));
+	const TSharedPtr<FDataflowConstructionScene> ConstructionScenePtr = ConstructionScene.Pin();
 
 	// Check if we have a single component selected. If we do, we will attempt to re-select it once the Construction Scene is rebuilt (if we have a tool running)
 	bHadSingleSelectionBeforeToolShutdown = false;
@@ -770,13 +764,13 @@ void UDataflowEditorMode::SetConstructionViewMode(const FName& NewViewModeName)
 
 	if (UInteractiveTool* const ActiveTool = ToolManager->GetActiveTool(EToolSide::Left))
 	{
-		if (const USelection* const SelectedComponents = FDataflowPreviewSceneBase::GetSelectedComponents(ConstructionScene))
+		if (const USelection* const SelectedComponents = ConstructionScenePtr ? ConstructionScenePtr->GetSelectedComponents() : nullptr)
 		{
 			if (SelectedComponents->Num() == 1)			// TODO: Extend this to handle multiple selected components
 			{
 				if (const UDynamicMeshComponent* const SelectedDynamicMeshComponent = Cast<UDynamicMeshComponent>(SelectedComponents->GetSelectedObject(0)))
 				{
-					if (ConstructionScene->GetDynamicMeshComponents().Contains(SelectedDynamicMeshComponent))
+					if (ConstructionScenePtr->GetDynamicMeshComponents().Contains(SelectedDynamicMeshComponent))
 					{
 						bHadSingleSelectionBeforeToolShutdown = true;
 					}
@@ -814,8 +808,11 @@ void UDataflowEditorMode::SetConstructionViewMode(const FName& NewViewModeName)
 	// Do the actual view mode updates
 
 	ConstructionViewMode = NewMode;
-	ConstructionScene->GetEditorContent()->SetConstructionViewMode(ConstructionViewMode);
-	ConstructionScene->UpdateConstructionScene();
+	if (ConstructionScenePtr)
+	{
+		ConstructionScenePtr->GetEditorContent()->SetConstructionViewMode(ConstructionViewMode);
+		ConstructionScenePtr->UpdateConstructionScene();
+	}
 
 	const TSharedPtr<FDataflowConstructionViewportClient> VC = ConstructionViewportClient.Pin();
 	if (VC.IsValid())
@@ -841,10 +838,10 @@ void UDataflowEditorMode::SetConstructionViewMode(const FName& NewViewModeName)
 
 		if (bHadSingleSelectionBeforeToolShutdown)
 		{
-			const TArray<UDynamicMeshComponent*> DynamicMeshComponents = ConstructionScene->GetDynamicMeshComponents();
+			const TArray<TObjectPtr<UDynamicMeshComponent>> DynamicMeshComponents = ConstructionScenePtr ? ConstructionScenePtr->GetDynamicMeshComponents() : TArray<TObjectPtr<UDynamicMeshComponent>>();
 			if (DynamicMeshComponents.Num() == 1)
 			{
-				if (USelection* const SelectedComponents = FDataflowPreviewSceneBase::GetSelectedComponents(ConstructionScene))
+				if (USelection* const SelectedComponents = ConstructionScenePtr->GetSelectedComponents())
 				{
 					SelectedComponents->Select(DynamicMeshComponents[0]);
 				}
@@ -870,7 +867,7 @@ void UDataflowEditorMode::SetConstructionViewMode(const FName& NewViewModeName)
 		{
 			if (const UDataflowEdNode* const SelectedDataflowEdNode = Cast<UDataflowEdNode>(SelectedNode))
 			{
-				if (const TObjectPtr<UDataflowBaseContent>& EditorContent = ConstructionScene->GetEditorContent())
+				if (const TObjectPtr<UDataflowBaseContent>& EditorContent = ConstructionScenePtr ? ConstructionScenePtr->GetEditorContent() : nullptr)
 				{
 					TArray<FName> ViewModesForNode;
 					UE::Dataflow::GetViewModesForNode(*SelectedDataflowEdNode, *EditorContent, ViewModesForNode);
@@ -893,7 +890,8 @@ void UDataflowEditorMode::SetConstructionViewMode(const FName& NewViewModeName)
 
 void UDataflowEditorMode::SetConstructionViewModeForNode(const UDataflowEdNode* Node)
 {
-	const TObjectPtr<UDataflowBaseContent> EditorContent = ConstructionScene ? ConstructionScene->GetEditorContent() : nullptr;
+	const TSharedPtr<FDataflowConstructionScene> ConstructionScenePtr = ConstructionScene.Pin();
+	const TObjectPtr<UDataflowBaseContent> EditorContent = ConstructionScenePtr ? ConstructionScenePtr->GetEditorContent() : nullptr;
 	if (Node && EditorContent)
 	{
 		// Check the most recently used view mode for this "kind" of node -- where "kind" means nodes that have the same set of valid view modes
@@ -959,30 +957,31 @@ bool UDataflowEditorMode::CanChangeConstructionViewModeTo(const FName& NewViewMo
 {
 	if (!GetToolManager()->HasActiveTool(EToolSide::Left))
 	{
-		if (const TObjectPtr<UDataflowBaseContent>& EditorContent = ConstructionScene->GetEditorContent())
+		if (const TSharedPtr<FDataflowConstructionScene> ConstructionScenePtr = ConstructionScene.Pin())
 		{
-			if (const TSharedPtr<const SDataflowGraphEditor> PinnedDataflowGraphEditor = DataflowGraphEditor.Pin())
+			if (const TObjectPtr<UDataflowBaseContent>& EditorContent = ConstructionScenePtr->GetEditorContent())
 			{
-				if (const UEdGraphNode* const SelectedNode = PinnedDataflowGraphEditor->GetSingleSelectedNode())
+				if (const TSharedPtr<const SDataflowGraphEditor> PinnedDataflowGraphEditor = DataflowGraphEditor.Pin())
 				{
-					if (const UDataflowEdNode* const SelectedDataflowEdNode = Cast<UDataflowEdNode>(SelectedNode))
+					if (const UEdGraphNode* const SelectedNode = PinnedDataflowGraphEditor->GetSingleSelectedNode())
 					{
-						if (const UE::Dataflow::IDataflowConstructionViewMode* const ViewMode = UE::Dataflow::FRenderingViewModeFactory::GetInstance().GetViewMode(NewViewModeName))
+						if (const UDataflowEdNode* const SelectedDataflowEdNode = Cast<UDataflowEdNode>(SelectedNode))
 						{
-							if (UE::Dataflow::CanRenderNodeOutput(*SelectedDataflowEdNode, *EditorContent, *ViewMode))
+							if (const UE::Dataflow::IDataflowConstructionViewMode* const ViewMode = UE::Dataflow::FRenderingViewModeFactory::GetInstance().GetViewMode(NewViewModeName))
 							{
-								return true;
+								if (UE::Dataflow::CanRenderNodeOutput(*SelectedDataflowEdNode, *EditorContent, *ViewMode))
+								{
+									return true;
+								}
 							}
 						}
 					}
 				}
 			}
 		}
-
 		return false;
 	}
 
-
 	// Check active tool to see if we can switch modes while the tool is running
 
 	const UInteractiveToolBuilder* const ActiveToolBuilder = GetToolManager()->GetActiveToolBuilder(EToolSide::Left);
@@ -1013,7 +1012,10 @@ void UDataflowEditorMode::ToggleConstructionViewWireframe()
 {
 	check(false);
 	bConstructionViewWireframe = !bConstructionViewWireframe;
-	ConstructionScene->UpdateConstructionScene();
+	if (const TSharedPtr<FDataflowConstructionScene> ConstructionScenePtr = ConstructionScene.Pin())
+	{
+		ConstructionScenePtr->UpdateConstructionScene();
+	}
 }
 
 bool UDataflowEditorMode::CanSetConstructionViewWireframeActive() const
@@ -1050,9 +1052,8 @@ void UDataflowEditorMode::SetSimulationViewportClient(TWeakPtr<FDataflowSimulati
 
 void UDataflowEditorMode::InitializeContextObject()
 {
-	check(ConstructionScene);
-
-	if (const TObjectPtr<UDataflowBaseContent>& EditorContent = ConstructionScene->GetEditorContent())
+	const TSharedPtr<FDataflowConstructionScene> ConstructionScenePtr = ConstructionScene.Pin();
+	if (const TObjectPtr<UDataflowBaseContent>& EditorContent = ConstructionScenePtr ? ConstructionScenePtr->GetEditorContent() : nullptr)
 	{
 		const UEditorInteractiveToolsContext* const ConstructionToolsContext = GetInteractiveToolsContext();
 
@@ -1162,8 +1163,8 @@ UEdGraphNode* UDataflowEditorMode::CreateNewNode(const FName& NewNodeTypeName)
 	{
 		return nullptr;
 	}
-
-	if (const TObjectPtr<UDataflowBaseContent>& EditorContent = ConstructionScene->GetEditorContent())
+	const TSharedPtr<FDataflowConstructionScene> ConstructionScenePtr = ConstructionScene.Pin();
+	if (const TObjectPtr<UDataflowBaseContent>& EditorContent = ConstructionScenePtr ? ConstructionScenePtr->GetEditorContent() : nullptr)
 	{
 		if (const TObjectPtr<UDataflow>& DataflowGraph = EditorContent->GetDataflowAsset())
 		{
@@ -1180,7 +1181,8 @@ UEdGraphNode* UDataflowEditorMode::CreateNewNode(const FName& NewNodeTypeName)
 
 UEdGraphNode* UDataflowEditorMode::CreateAndConnectNewNode(const FName& NewNodeTypeName, UEdGraphNode& UpstreamNode, const FName& ConnectionTypeName, const FName& NewNodeConnectionName)
 {
-	if (const TObjectPtr<UDataflowBaseContent>& EditorContent = ConstructionScene->GetEditorContent())
+	const TSharedPtr<FDataflowConstructionScene> ConstructionScenePtr = ConstructionScene.Pin();
+	if (const TObjectPtr<UDataflowBaseContent>& EditorContent = ConstructionScenePtr ? ConstructionScenePtr->GetEditorContent() : nullptr)
 	{
 		if (const TObjectPtr<UDataflow>& DataflowGraph = EditorContent->GetDataflowAsset())
 		{
diff --git a/Engine/Plugins/Experimental/Dataflow/Source/DataflowEditor/Private/Dataflow/DataflowEditorModeToolkit.cpp b/Engine/Plugins/Experimental/Dataflow/Source/DataflowEditor/Private/Dataflow/DataflowEditorModeToolkit.cpp
--- a/Engine/Plugins/Experimental/Dataflow/Source/DataflowEditor/Private/Dataflow/DataflowEditorModeToolkit.cpp
+++ b/Engine/Plugins/Experimental/Dataflow/Source/DataflowEditor/Private/Dataflow/DataflowEditorModeToolkit.cpp
@@ -147,7 +147,7 @@ SBaseCharacterFXEditorViewport* FDataflowEditorModeToolkit::GetViewportWidgetFor
 
 		if (const UDataflowEditorMode* const DataflowEdMode = Cast<UDataflowEditorMode>(Mode))
 		{
-			if (const FDataflowPreviewSceneBase* const PreviewScene = DataflowEdMode->GetDataflowConstructionScene())
+			if (const TSharedPtr<const FDataflowPreviewSceneBase> PreviewScene = DataflowEdMode->GetDataflowConstructionScene().Pin())
 			{
 				if (const UEditorInteractiveToolsContext* const PreviewToolsContext = PreviewScene->GetDataflowModeManager()->GetInteractiveToolsContext())
 				{
@@ -231,18 +231,19 @@ void FDataflowEditorModeToolkit::OnToolStarted(UInteractiveToolManager* Manager,
 
 void FDataflowEditorModeToolkit::OnToolEnded(UInteractiveToolManager* Manager, UInteractiveTool* Tool)
 {
-	auto GetConstructionScene = [&]() 
-	{
-		if (OwningEditorMode.IsValid(/*bEvenIfPendingKill*/ false))
+	auto GetConstructionScene = [this]() -> TWeakPtr<FDataflowConstructionScene>
 		{
-			UEdMode* const Mode = OwningEditorMode.Get();
-			if (UDataflowEditorMode* DataflowEdMode = Cast<UDataflowEditorMode>(Mode))
+			constexpr bool bEvenIfPendingKill = false;
+			if (OwningEditorMode.IsValid(bEvenIfPendingKill))
 			{
-				return DataflowEdMode->GetDataflowConstructionScene();
+				UEdMode* const Mode = OwningEditorMode.Get();
+				if (UDataflowEditorMode* const DataflowEdMode = Cast<UDataflowEditorMode>(Mode))
+				{
+					return DataflowEdMode->GetDataflowConstructionScene();
+				}
 			}
-		}
-		return (FDataflowConstructionScene* )nullptr;
-	};
+			return nullptr;
+		};
 
 	FModeToolkit::OnToolEnded(Manager, Tool);
 
@@ -261,15 +262,13 @@ void FDataflowEditorModeToolkit::OnToolEnded(UInteractiveToolManager* Manager, U
 		CurTool->OnPropertyModifiedDirectlyByTool.RemoveAll(this);
 	}
 
-
-	if (FDataflowConstructionScene* ConstructionScene = GetConstructionScene())
+	if (const TSharedPtr<FDataflowConstructionScene> ConstructionScene = GetConstructionScene().Pin())
 	{
 		if (const TObjectPtr<UDataflowBaseContent>& EditorContent = ConstructionScene->GetEditorContent())
 		{
 			EditorContent->SetConstructionDirty(true);
 		}
 	}
-
 }
 
 void FDataflowEditorModeToolkit::SetConstructionViewportWidget(TWeakPtr<SDataflowConstructionViewport> InConstructionViewportWidget)
diff --git a/Engine/Plugins/Experimental/Dataflow/Source/DataflowEditor/Private/Dataflow/DataflowEditorPreviewSceneBase.cpp b/Engine/Plugins/Experimental/Dataflow/Source/DataflowEditor/Private/Dataflow/DataflowEditorPreviewSceneBase.cpp
--- a/Engine/Plugins/Experimental/Dataflow/Source/DataflowEditor/Private/Dataflow/DataflowEditorPreviewSceneBase.cpp
+++ b/Engine/Plugins/Experimental/Dataflow/Source/DataflowEditor/Private/Dataflow/DataflowEditorPreviewSceneBase.cpp
@@ -67,13 +67,6 @@ FDataflowPreviewSceneBase::FDataflowPreviewSceneBase(FPreviewScene::Construction
 
 FDataflowPreviewSceneBase::~FDataflowPreviewSceneBase()
 {
-	// There is a crash in 5.6, but the fix requires the preview scene to become a weakptr and this isn't hotfixable for 5.6.1.
-	// Instead sending the following guid to notify the destruction of the preview scene, so the pointer can be safely cleared up.
-	// TODO: Remove this hack for 5.7.
-	static const FGuid TerminationGuid(0x8A4F3813, 0x5D7E4932, 0x8F9F92AF, 0x74341FAA);
-	static const FBox TerminationBox(FVector((double)TerminationGuid.A, (double)TerminationGuid.B, 0.0), FVector((double)TerminationGuid.C, (double)TerminationGuid.D, 0.0));
-	FocusRequestDelegate.Broadcast(TerminationBox);
-
 	UnregisterSceneElements();
 	ModifySceneElements().Reset();
 }
@@ -276,16 +269,9 @@ void FDataflowPreviewSceneBase::RegisterSceneElements(const bool bIsConstruction
 	}
 }
 
-USelection* FDataflowPreviewSceneBase::GetSelectedComponents(FDataflowPreviewSceneBase* DataflowScene)
+USelection* FDataflowPreviewSceneBase::GetSelectedComponents() const
 {
-	if( DataflowScene)
-	{
-		if(TSharedPtr<FAssetEditorModeManager> ModeManager = DataflowScene->GetDataflowModeManager())
-		{
-			return ModeManager->GetSelectedComponents();
-		}
-	}
-	return nullptr; 
+	return DataflowModeManager ? DataflowModeManager->GetSelectedComponents() : nullptr;
 }
 
 #undef LOCTEXT_NAMESPACE
diff --git a/Engine/Plugins/Experimental/Dataflow/Source/DataflowEditor/Private/Dataflow/DataflowEditorToolkit.cpp b/Engine/Plugins/Experimental/Dataflow/Source/DataflowEditor/Private/Dataflow/DataflowEditorToolkit.cpp
--- a/Engine/Plugins/Experimental/Dataflow/Source/DataflowEditor/Private/Dataflow/DataflowEditorToolkit.cpp
+++ b/Engine/Plugins/Experimental/Dataflow/Source/DataflowEditor/Private/Dataflow/DataflowEditorToolkit.cpp
@@ -7,18 +7,16 @@
 #include "AssetEditorModeManager.h"
 #include "AssetViewerSettings.h"
 #include "DetailCategoryBuilder.h"
-#include "Animation/Skeleton.h"
 #include "Dataflow/DataflowAdvancedPreviewDetailsTab.h"
 #include "Dataflow/DataflowAssetViewerSettingsCustomization.h"
 #include "Dataflow/DataflowDebugDrawComponent.h"
 #include "Dataflow/DataflowDebugDraw.h"
-#include "Dataflow/DataflowCore.h"
 #include "Dataflow/DataflowEditor.h"
 #include "Dataflow/DataflowContent.h"
+#include "Dataflow/DataflowCollectionSpreadSheet.h"
 #include "Dataflow/DataflowCollectionSpreadSheetWidget.h"
 #include "Dataflow/DataflowConstructionScene.h"
 #include "Dataflow/DataflowConstructionViewportClient.h"
-#include "Dataflow/DataflowEditorCollectionComponent.h"
 #include "Dataflow/DataflowEditorCommands.h"
 #include "Dataflow/DataflowEditorMode.h"
 #include "Dataflow/DataflowEditorModeToolkit.h"
@@ -30,48 +28,35 @@
 #include "Dataflow/DataflowEditorOptions.h"
 #include "Dataflow/DataflowEdNode.h"
 #include "Dataflow/DataflowGraphEditor.h"
-#include "Dataflow/DataflowNodeFactory.h"
 #include "Dataflow/DataflowNodeDetailExtension.h"
-#include "Dataflow/DataflowObject.h"
 #include "Dataflow/DataflowInstance.h"
-#include "Dataflow/DataflowObjectInterface.h"
-#include "Dataflow/DataflowEditorPreviewSceneBase.h"
-#include "Dataflow/DataflowRenderingFactory.h"
+#include "Dataflow/DataflowRenderingViewMode.h"
 #include "Dataflow/DataflowSceneProfileIndexStorage.h"
 #include "Dataflow/DataflowSimulationScene.h"
 #include "Dataflow/DataflowSimulationVisualization.h"
 #include "Dataflow/DataflowSchema.h"
 #include "Dataflow/DataflowSkeletonView.h"
 #include "Dataflow/DataflowOutlinerView.h"
+#include "Dataflow/DataflowSimulationViewport.h"
 #include "Dataflow/DataflowSimulationViewportClient.h"
 #include "Dataflow/DataflowSubGraphNodes.h"
 #include "Dataflow/DataflowMembersWidget.h"
-#include "DynamicMeshBuilder.h"
+#include "Dataflow/DataflowSelectionView.h"
 #include "EditorModeManager.h"
-#include "EditorStyleSet.h"
 #include "EditorViewportTabContent.h"
-#include "EditorViewportLayout.h"
-#include "EditorViewportCommands.h"
 #include "Tools/EdModeInteractiveToolsContext.h"
-#include "Engine/SkeletalMesh.h"
-#include "Framework/Commands/GenericCommands.h"
 #include "GameFramework/Actor.h"
-#include "GraphEditorActions.h"
-#include "IDetailCustomization.h"
 #include "ISkeletonTree.h"
 #include "ISceneOutliner.h"
 #include "IStructureDetailsView.h"
-#include "Kismet/GameplayStatics.h"
-#include "Misc/MessageDialog.h"
 #include "Modules/ModuleManager.h"
 #include "PropertyEditorModule.h"
 #include "Selection.h"
-#include "Styling/SlateStyleRegistry.h"
 #include "Widgets/Docking/SDockTab.h"
+#include "Widgets/Layout/SSpacer.h"
 #include "GeometryCache.h"
 #include "ToolMenus.h"
 #include "WorkflowOrientedApp/WorkflowTabManager.h"
-#include "Dataflow/DataflowNodeParameters.h"
 #include "Framework/Notifications/NotificationManager.h"
 #include "Widgets/Notifications/SNotificationList.h"
 #include "Dataflow/DataflowOutputLog.h"
@@ -117,7 +102,7 @@ namespace UE::Dataflow::Private
 
 	/** Create a debug draw component to visualize per node informations */
 	TUniquePtr<IDataflowDebugDrawInterface> CreateDebugDrawComponent(
-		FDataflowPreviewSceneBase* DataflowScene,
+		TSharedPtr<FDataflowPreviewSceneBase>& DataflowScene,
 		TObjectPtr<class UDataflowDebugDrawComponent>& DebugDrawComponent)
 	{
 		if (DataflowScene && !DebugDrawComponent)
@@ -147,7 +132,7 @@ namespace UE::Dataflow::Private
 	/** Update the debug draw component from the selected/pinned nodes matching the NodeType */
 	template<typename NodeType>
 	void UpdateDebugDrawComponent(const TObjectPtr<UDataflowBaseContent>& EditorContent,
-		const FString& RootName, FDataflowPreviewSceneBase* DataflowScene, const bool bIsConstruction, FString& DebugDrawOverlay,
+		const FString& RootName, TSharedPtr<FDataflowPreviewSceneBase> DataflowScene, const bool bIsConstruction, FString& DebugDrawOverlay,
 		const TSharedPtr<FEditorModeTools> EditorModeManager)
 	{
 		if(!DataflowScene || !EditorContent)
@@ -458,13 +443,12 @@ FDataflowEditorToolkit::FDataflowEditorToolkit(UAssetEditor* InOwningAssetEditor
 	PreviewSceneArgs.bShouldSimulatePhysics = 1;
 	PreviewSceneArgs.bCreatePhysicsScene = 1;
 	
-	TUniquePtr<FDataflowConstructionScene> NewConstructionScene = MakeUnique<FDataflowConstructionScene>(PreviewSceneArgs, DataflowEditor);
-	ConstructionScene = NewConstructionScene.Get();
-	ObjectScene = MoveTemp(NewConstructionScene);
+	ConstructionScene = MakeShared<FDataflowConstructionScene>(PreviewSceneArgs, DataflowEditor);
+	ObjectScene = ConstructionScene;
 
 	SimulationScene = MakeShared<FDataflowSimulationScene>(PreviewSceneArgs, DataflowEditor);
 
-	SimulationSceneProfileIndexStorage = MakeShared<FDataflowSimulationSceneProfileIndexStorage>(SimulationScene.Get());
+	SimulationSceneProfileIndexStorage = MakeShared<FDataflowSimulationSceneProfileIndexStorage>(SimulationScene);
 
 	IConsoleVariable* const ConsoleVar = IConsoleManager::Get().FindConsoleVariable(TEXT("p.Dataflow.EnableGraphEval"));
 	bDataflowEnableGraphEval = ConsoleVar ? ConsoleVar->GetBool() : true;
@@ -1011,7 +995,7 @@ void FDataflowEditorToolkit::InitializeEdMode(UBaseCharacterFXEditorMode* EdMode
 	DataflowMode->SetDataflowConstructionScene(ConstructionScene);
 
 	checkf(SimulationScene, TEXT("Expected SimulationScene to have been created in FDataflowEditorToolkit constructor"));
-	DataflowMode->SetDataflowSimlationScene(SimulationScene.Get());
+	DataflowMode->SetDataflowSimlationScene(SimulationScene);
 
 	// Set of the graph editor to be able to add nodes
 	DataflowMode->SetDataflowGraphEditor(GraphEditor);
@@ -1199,7 +1183,7 @@ TSharedPtr<FEditorViewportClient> FDataflowEditorToolkit::CreateEditorViewportCl
 	// Instead, we do viewport client adjustment in PostInitAssetEditor().
 	check(EditorModeManager.IsValid());
 	TSharedPtr<FDataflowConstructionViewportClient> LocalConstructionClient = MakeShared<FDataflowConstructionViewportClient>(
-	EditorModeManager.Get(), ObjectScene.Get(), true);
+		EditorModeManager.Get(), ConstructionScene, true);
 	LocalConstructionClient->SetDataflowEditorToolkit(StaticCastSharedRef<FDataflowEditorToolkit>(
 		const_cast<FDataflowEditorToolkit*>(this)->AsShared()));
 
@@ -1210,7 +1194,7 @@ void FDataflowEditorToolkit::CreateSimulationViewportClient()
 {
 	SimulationTabContent = MakeShareable(new FEditorViewportTabContent());
 	SimulationViewportClient = MakeShared<FDataflowSimulationViewportClient>(SimulationModeManager.Get(),
-		SimulationScene.Get(), false);
+		SimulationScene, false);
 	
 	SimulationViewportClient->SetDataflowEditorToolkit(StaticCastSharedRef<FDataflowEditorToolkit>(this->AsShared()));
 
@@ -1250,7 +1234,7 @@ void FDataflowEditorToolkit::UpdateDebugDraw()
 		{
 			static const FString RootElementsName = "Simulation Elements";
 			UE::Dataflow::Private::UpdateDebugDrawComponent<FDataflowSimulationNode>(
-				EditorContent, RootElementsName, SimulationScene.Get(), false, DebugDrawOverlayString, SimulationModeManager);
+				EditorContent, RootElementsName, SimulationScene, false, DebugDrawOverlayString, SimulationModeManager);
 		}
 	}
 }
@@ -2190,7 +2174,7 @@ TSharedRef<SDockTab> FDataflowEditorToolkit::SpawnTab_OutlinerView(const FSpawnT
 	check(DataflowEditor);
 	check(DataflowEditor->GetEditorContent());
 
-	DataflowOutlinerView = MakeShared<FDataflowOutlinerView>(ConstructionScene, SimulationScene.Get(), DataflowEditor->GetEditorContent());
+	DataflowOutlinerView = MakeShared<FDataflowOutlinerView>(ConstructionScene, SimulationScene, DataflowEditor->GetEditorContent());
 	ViewListeners.Add(DataflowOutlinerView.Get());
 
 	TSharedPtr<ISceneOutliner> DataflowOutliner = DataflowOutlinerView->CreateWidget();
@@ -2730,11 +2714,6 @@ FLinearColor FDataflowEditorToolkit::GetWorldCentricTabColorScale() const
 	return FLinearColor(0.3f, 0.2f, 0.5f, 0.5f);
 }
 
-FDataflowConstructionScene* FDataflowEditorToolkit::GetConstructionScene() const
-{ 
-	return ConstructionScene;
-}
-
 void FDataflowEditorToolkit::AddReferencedObjects(FReferenceCollector& Collector)
 {
 	Collector.AddReferencedObjects(SelectedDataflowNodes);
diff --git a/Engine/Plugins/Experimental/Dataflow/Source/DataflowEditor/Private/Dataflow/DataflowEditorViewportClientBase.cpp b/Engine/Plugins/Experimental/Dataflow/Source/DataflowEditor/Private/Dataflow/DataflowEditorViewportClientBase.cpp
--- a/Engine/Plugins/Experimental/Dataflow/Source/DataflowEditor/Private/Dataflow/DataflowEditorViewportClientBase.cpp
+++ b/Engine/Plugins/Experimental/Dataflow/Source/DataflowEditor/Private/Dataflow/DataflowEditorViewportClientBase.cpp
@@ -12,9 +12,10 @@
 #define LOCTEXT_NAMESPACE "DataflowEditorViewportClientBase"
 
 FDataflowEditorViewportClientBase::FDataflowEditorViewportClientBase(FEditorModeTools* InModeTools,
-                                                             FPreviewScene* InPreviewScene,  const bool bCouldTickScene,
+                                                             const TWeakPtr<FDataflowPreviewSceneBase> InPreviewScene, const bool bCouldTickScene,
                                                              const TWeakPtr<SEditorViewport> InEditorViewportWidget)
-	: FEditorViewportClient(InModeTools, InPreviewScene, InEditorViewportWidget)
+	: FEditorViewportClient(InModeTools, InPreviewScene.Pin().Get(), InEditorViewportWidget)
+	, DataflowPreviewScene(InPreviewScene)
 {
 	BehaviorSet = NewObject<UInputBehaviorSet>();
 
@@ -47,8 +48,6 @@ FDataflowEditorViewportClientBase::FDataflowEditorViewportClientBase(FEditorMode
 
 	ModeTools->GetInteractiveToolsContext()->InputRouter->RegisterSource(this);
 
-	DataflowPreviewScene = static_cast<FDataflowPreviewSceneBase*>(InPreviewScene);
-
 	RegisterDelegates();
 }
 
@@ -76,9 +75,11 @@ void FDataflowEditorViewportClientBase::RegisterDelegates()
 
 	auto SyncEngineShowFlags = [this]()
 		{
-			if (ensureMsgf(DataflowPreviewScene, TEXT("Expected valid DataflowPreviewScene pointer")))
+			const TSharedPtr<FDataflowPreviewSceneBase> DataflowPreviewScenePtr = DataflowPreviewScene.Pin();
+
+			if (ensureMsgf(DataflowPreviewScenePtr, TEXT("Expected valid DataflowPreviewScene pointer")))
 			{
-				const int32 CurrentProfileIndex = DataflowPreviewScene->GetCurrentProfileIndex();
+				const int32 CurrentProfileIndex = DataflowPreviewScenePtr->GetCurrentProfileIndex();
 
 				const UAssetViewerSettings* const DefaultSettings = UAssetViewerSettings::Get();
 
@@ -112,9 +113,9 @@ void FDataflowEditorViewportClientBase::RegisterDelegates()
 		}
 	});
 
-	if(DataflowPreviewScene)
+	if (const TSharedPtr<FDataflowPreviewSceneBase> DataflowPreviewScenePtr = DataflowPreviewScene.Pin())
 	{
-		OnFocusRequestDelegateHandle = DataflowPreviewScene->OnFocusRequest().AddRaw(this, &FDataflowEditorViewportClientBase::HandleFocusRequest);
+		OnFocusRequestDelegateHandle = DataflowPreviewScenePtr->OnFocusRequest().AddRaw(this, &FDataflowEditorViewportClientBase::HandleFocusRequest);
 	}
 }
 
@@ -122,24 +123,14 @@ void FDataflowEditorViewportClientBase::DeregisterDelegates()
 {
 	UAssetViewerSettings::Get()->OnAssetViewerSettingsChanged().Remove(OnAssetViewerSettingsChangedDelegateHandle);
 
-	if(DataflowPreviewScene)
+	if (const TSharedPtr<FDataflowPreviewSceneBase> DataflowPreviewScenePtr = DataflowPreviewScene.Pin())
 	{
-		DataflowPreviewScene->OnFocusRequest().Remove(OnFocusRequestDelegateHandle);
+		DataflowPreviewScenePtr->OnFocusRequest().Remove(OnFocusRequestDelegateHandle);
 	}
 }
 
 void FDataflowEditorViewportClientBase::HandleFocusRequest(const FBox& BoundingBox)
 {
-	// There is a crash in 5.6, but the fix requires the preview scene to become a weakptr and this isn't hotfixable for 5.6.1.
-	// Instead sending the following guid to notify the destruction of the preview scene, so the pointer can be safely cleared up.
-	// TODO: Remove this hack for 5.7.
-	static const FGuid TerminationGuid(0x8A4F3813, 0x5D7E4932, 0x8F9F92AF, 0x74341FAA);
-	static const FBox TerminationBox(FVector((double)TerminationGuid.A, (double)TerminationGuid.B, 0.0), FVector((double)TerminationGuid.C, (double)TerminationGuid.D, 0.0));
-	if (BoundingBox == TerminationBox)
-	{
-		DataflowPreviewScene = nullptr;
-		return;
-	}
 	FocusViewportOnBox(BoundingBox);
 }
 
@@ -151,9 +142,9 @@ void FDataflowEditorViewportClientBase::Draw(const FSceneView* View, FPrimitiveD
 void FDataflowEditorViewportClientBase::GetSelectedElements(HHitProxy* HitProxy, TArray<FDataflowBaseElement*>& SelectedElements) const
 {
 	TArray<FDataflowBaseElement*> DataflowElements;
-	if(DataflowPreviewScene)
+	if (const TSharedPtr<FDataflowPreviewSceneBase> DataflowPreviewScenePtr = DataflowPreviewScene.Pin())
 	{
-		for(IDataflowDebugDrawInterface::FDataflowElementsType::ElementType& DataflowElement : DataflowPreviewScene->ModifySceneElements())
+		for(IDataflowDebugDrawInterface::FDataflowElementsType::ElementType& DataflowElement : DataflowPreviewScenePtr->ModifySceneElements())
 		{
 			if(DataflowElement.IsValid())
 			{
@@ -164,9 +155,9 @@ void FDataflowEditorViewportClientBase::GetSelectedElements(HHitProxy* HitProxy,
 		{
 			if(const HDataflowElementHitProxy* DataflowElementProxy = static_cast<HDataflowElementHitProxy*>(HitProxy))
 			{
-				if(DataflowElementProxy->ElementIndex < DataflowPreviewScene->ModifySceneElements().Num())
+				if(DataflowElementProxy->ElementIndex < DataflowPreviewScenePtr->ModifySceneElements().Num())
 				{
-					IDataflowDebugDrawInterface::FDataflowElementsType::ElementType& DataflowElement = DataflowPreviewScene->ModifySceneElements()[DataflowElementProxy->ElementIndex];
+					IDataflowDebugDrawInterface::FDataflowElementsType::ElementType& DataflowElement = DataflowPreviewScenePtr->ModifySceneElements()[DataflowElementProxy->ElementIndex];
 
 					if(DataflowElement.IsValid())
 					{
diff --git a/Engine/Plugins/Experimental/Dataflow/Source/DataflowEditor/Private/Dataflow/DataflowOutlinerMode.cpp b/Engine/Plugins/Experimental/Dataflow/Source/DataflowEditor/Private/Dataflow/DataflowOutlinerMode.cpp
--- a/Engine/Plugins/Experimental/Dataflow/Source/DataflowEditor/Private/Dataflow/DataflowOutlinerMode.cpp
+++ b/Engine/Plugins/Experimental/Dataflow/Source/DataflowEditor/Private/Dataflow/DataflowOutlinerMode.cpp
@@ -1,7 +1,8 @@
 // Copyright Epic Games, Inc. All Rights Reserved.
 
 #include "Dataflow/DataflowOutlinerMode.h"
-#include "Dataflow/DataflowEditorPreviewSceneBase.h"
+#include "Dataflow/DataflowConstructionScene.h"
+#include "Dataflow/DataflowSimulationScene.h"
 #include "AssetEditorModeManager.h"
 #include "EditorViewportCommands.h"
 #include "TedsOutlinerItem.h"
@@ -184,7 +185,7 @@ namespace UE::Dataflow::Private
 }
 
 FDataflowOutlinerMode::FDataflowOutlinerMode(const UE::Editor::Outliner::FTedsOutlinerParams& InModeParams,
-		FDataflowPreviewSceneBase* InConstructionScene, FDataflowPreviewSceneBase* InSimulationScene)
+		TWeakPtr<FDataflowConstructionScene> InConstructionScene, TWeakPtr<FDataflowSimulationScene> InSimulationScene)
 	: FTedsOutlinerMode(InModeParams)
 	, ConstructionScene(InConstructionScene)
 	, SimulationScene(InSimulationScene)
@@ -209,17 +210,19 @@ FDataflowOutlinerMode::~FDataflowOutlinerMode()
 
 void FDataflowOutlinerMode::OnItemSelectionChanged(FSceneOutlinerTreeItemPtr Item, ESelectInfo::Type SelectionType, const FSceneOutlinerItemSelection& Selection)
 {
-	if ((SelectionType == ESelectInfo::Direct) || !ConstructionScene || !SimulationScene)
+	const TSharedPtr<FDataflowConstructionScene> ConstructionScenePtr = ConstructionScene.Pin();
+	const TSharedPtr<FDataflowSimulationScene> SimulationScenePtr = SimulationScene.Pin();
+	if ((SelectionType == ESelectInfo::Direct) || !ConstructionScenePtr || !SimulationScenePtr)
 	{
 		return;
 	}
-	if(TSharedPtr<FAssetEditorModeManager> ConstructionManager = ConstructionScene->GetDataflowModeManager())
+	if (const TSharedPtr<FAssetEditorModeManager> ConstructionManager = ConstructionScenePtr->GetDataflowModeManager())
 	{
-		UE::Dataflow::Private::UpdateSceneSelection(Selection, ConstructionScene, ConstructionManager->GetSelectedComponents());
+		UE::Dataflow::Private::UpdateSceneSelection(Selection, ConstructionScenePtr.Get(), ConstructionManager->GetSelectedComponents());
 	}
-	if(TSharedPtr<FAssetEditorModeManager> SimulationManager = SimulationScene->GetDataflowModeManager())
+	if (const TSharedPtr<FAssetEditorModeManager> SimulationManager = SimulationScenePtr->GetDataflowModeManager())
 	{
-		UE::Dataflow::Private::UpdateSceneSelection(Selection, SimulationScene, SimulationManager->GetSelectedComponents());
+		UE::Dataflow::Private::UpdateSceneSelection(Selection, SimulationScenePtr.Get(), SimulationManager->GetSelectedComponents());
 	}
 }
 
@@ -231,38 +234,42 @@ void FDataflowOutlinerMode::OnItemAdded(FSceneOutlinerTreeItemPtr Item)
 
 void FDataflowOutlinerMode::OnItemDoubleClick(FSceneOutlinerTreeItemPtr SelectedItem)
 {
-	if (!ConstructionScene || !SimulationScene)
+	const TSharedPtr<FDataflowConstructionScene> ConstructionScenePtr = ConstructionScene.Pin();
+	const TSharedPtr<FDataflowSimulationScene> SimulationScenePtr = SimulationScene.Pin();
+	if (!ConstructionScenePtr || !SimulationScenePtr)
 	{
 		return;
 	}
 
 	if(const UPrimitiveComponent* SelectedComponent = UE::Dataflow::Private::GetOutlinerItemObject<UPrimitiveComponent>(SelectedItem))
 	{
-		if(SelectedComponent->GetWorld() == ConstructionScene->GetWorld())
+		if(SelectedComponent->GetWorld() == ConstructionScenePtr->GetWorld())
 		{
-			ConstructionScene->OnFocusRequest().Broadcast(SelectedComponent->Bounds.GetBox());
+			ConstructionScenePtr->OnFocusRequest().Broadcast(SelectedComponent->Bounds.GetBox());
 		}
-		else if(SelectedComponent->GetWorld() == SimulationScene->GetWorld())
+		else if(SelectedComponent->GetWorld() == SimulationScenePtr->GetWorld())
 		{
-			SimulationScene->OnFocusRequest().Broadcast(SelectedComponent->Bounds.GetBox());
+			SimulationScenePtr->OnFocusRequest().Broadcast(SelectedComponent->Bounds.GetBox());
 		}
 	}
 	else if(FDataflowBaseElement* SelectedElement = UE::Dataflow::Private::GetOutlinerItemStruct<FDataflowBaseElement>(SelectedItem))
 	{
 		if(SelectedElement->bIsConstruction)
 		{
-			ConstructionScene->OnFocusRequest().Broadcast(SelectedElement->BoundingBox);
+			ConstructionScenePtr->OnFocusRequest().Broadcast(SelectedElement->BoundingBox);
 		}
 		else
 		{
-			SimulationScene->OnFocusRequest().Broadcast(SelectedElement->BoundingBox);
+			SimulationScenePtr->OnFocusRequest().Broadcast(SelectedElement->BoundingBox);
 		}
 	}
 }
 
 FReply FDataflowOutlinerMode::OnKeyDown(const FKeyEvent& InKeyEvent)
 {
-	if(SceneOutliner)
+	const TSharedPtr<FDataflowConstructionScene> ConstructionScenePtr = ConstructionScene.Pin();
+	const TSharedPtr<FDataflowSimulationScene> SimulationScenePtr = SimulationScene.Pin();
+	if (SceneOutliner && ConstructionScenePtr && SimulationScenePtr)
 	{
 		const FModifierKeysState ModifierKeys = FSlateApplication::Get().GetModifierKeys();
 		const FInputChord CheckChord( InKeyEvent.GetKey(), EModifierKey::FromBools(ModifierKeys.IsControlDown(), ModifierKeys.IsAltDown(), ModifierKeys.IsShiftDown(), ModifierKeys.IsCommandDown()) );
@@ -277,11 +284,11 @@ FReply FDataflowOutlinerMode::OnKeyDown(const FKeyEvent& InKeyEvent)
 			{
 				if(const UPrimitiveComponent* SelectedComponent = UE::Dataflow::Private::GetOutlinerItemObject<UPrimitiveComponent>(SelectedItem))
 				{
-					if(SelectedComponent->GetWorld() == ConstructionScene->GetWorld())
+					if(SelectedComponent->GetWorld() == ConstructionScenePtr->GetWorld())
 					{
 						ConstructionBox += SelectedComponent->Bounds.GetBox();
 					}
-					else if(SelectedComponent->GetWorld() == SimulationScene->GetWorld())
+					else if(SelectedComponent->GetWorld() == SimulationScenePtr->GetWorld())
 					{
 						SimulationBox += SelectedComponent->Bounds.GetBox();
 					}
@@ -300,11 +307,11 @@ FReply FDataflowOutlinerMode::OnKeyDown(const FKeyEvent& InKeyEvent)
 			}
 			if(ConstructionBox.IsValid)
 			{
-				ConstructionScene->OnFocusRequest().Broadcast(ConstructionBox);
+				ConstructionScenePtr->OnFocusRequest().Broadcast(ConstructionBox);
 			}
 			if(SimulationBox.IsValid)
 			{
-				SimulationScene->OnFocusRequest().Broadcast(SimulationBox);
+				SimulationScenePtr->OnFocusRequest().Broadcast(SimulationBox);
 			}
 		}
 		else if (InKeyEvent.GetKey() == EKeys::H)
diff --git a/Engine/Plugins/Experimental/Dataflow/Source/DataflowEditor/Private/Dataflow/DataflowOutlinerMode.h b/Engine/Plugins/Experimental/Dataflow/Source/DataflowEditor/Private/Dataflow/DataflowOutlinerMode.h
--- a/Engine/Plugins/Experimental/Dataflow/Source/DataflowEditor/Private/Dataflow/DataflowOutlinerMode.h
+++ b/Engine/Plugins/Experimental/Dataflow/Source/DataflowEditor/Private/Dataflow/DataflowOutlinerMode.h
@@ -3,7 +3,8 @@
 
 #include "TedsOutlinerMode.h"
 
-class FDataflowPreviewSceneBase;
+class FDataflowConstructionScene;
+class FDataflowSimulationScene;
 
 /**
  * Dataflow outliner mode used to represent the dataflow scenes (construction/simulation)
@@ -14,7 +15,7 @@ class FDataflowOutlinerMode : public UE::Editor::Outliner::FTedsOutlinerMode
 public:
 
 	FDataflowOutlinerMode(const UE::Editor::Outliner::FTedsOutlinerParams& InModeParams,
-		FDataflowPreviewSceneBase* InConstructionScene, FDataflowPreviewSceneBase* InSimulationScene);
+		TWeakPtr<FDataflowConstructionScene> InConstructionScene, TWeakPtr<FDataflowSimulationScene> InSimulationScene);
 
 	virtual ~FDataflowOutlinerMode() override;
 
@@ -31,8 +32,8 @@ public:
 
 private :
 	/** Construction dataflow scene */
-	FDataflowPreviewSceneBase* ConstructionScene;
+	TWeakPtr<FDataflowConstructionScene> ConstructionScene;
 
 	/** Simulation dataflow scene */
-	FDataflowPreviewSceneBase* SimulationScene;
+	TWeakPtr<FDataflowSimulationScene> SimulationScene;
 };
diff --git a/Engine/Plugins/Experimental/Dataflow/Source/DataflowEditor/Private/Dataflow/DataflowOutlinerView.cpp b/Engine/Plugins/Experimental/Dataflow/Source/DataflowEditor/Private/Dataflow/DataflowOutlinerView.cpp
--- a/Engine/Plugins/Experimental/Dataflow/Source/DataflowEditor/Private/Dataflow/DataflowOutlinerView.cpp
+++ b/Engine/Plugins/Experimental/Dataflow/Source/DataflowEditor/Private/Dataflow/DataflowOutlinerView.cpp
@@ -12,7 +12,7 @@
 #include "Elements/Columns/TypedElementVisibilityColumns.h"
 #include "Elements/Framework/TypedElementQueryBuilder.h"
 
-FDataflowOutlinerView::FDataflowOutlinerView(FDataflowPreviewSceneBase* InConstructionScene, FDataflowPreviewSceneBase* InSimulationScene, TObjectPtr<UDataflowBaseContent> InContent)
+FDataflowOutlinerView::FDataflowOutlinerView(TWeakPtr<FDataflowConstructionScene> InConstructionScene, TWeakPtr<FDataflowSimulationScene> InSimulationScene, TObjectPtr<UDataflowBaseContent> InContent)
 	: FDataflowNodeView(InContent)
 	, OutlinerWidget(nullptr)
 	, ConstructionScene(InConstructionScene)
diff --git a/Engine/Plugins/Experimental/Dataflow/Source/DataflowEditor/Private/Dataflow/DataflowOutlinerView.h b/Engine/Plugins/Experimental/Dataflow/Source/DataflowEditor/Private/Dataflow/DataflowOutlinerView.h
--- a/Engine/Plugins/Experimental/Dataflow/Source/DataflowEditor/Private/Dataflow/DataflowOutlinerView.h
+++ b/Engine/Plugins/Experimental/Dataflow/Source/DataflowEditor/Private/Dataflow/DataflowOutlinerView.h
@@ -5,14 +5,15 @@
 
 class SSceneOutliner;
 class ISceneOutliner;
-class FDataflowPreviewSceneBase;
+class FDataflowConstructionScene;
+class FDataflowSimulationScene;
 struct FDataflowBaseElement;
 
 /** Class to handle the dataflow outliner widget */
 class FDataflowOutlinerView : public FDataflowNodeView
 {
 public:
-	FDataflowOutlinerView(FDataflowPreviewSceneBase* InConstructionScene, FDataflowPreviewSceneBase* InSimulationScene, TObjectPtr<UDataflowBaseContent> InContent = nullptr);
+	FDataflowOutlinerView(TWeakPtr<FDataflowConstructionScene> InConstructionScene, TWeakPtr<FDataflowSimulationScene> InSimulationScene, TObjectPtr<UDataflowBaseContent> InContent = nullptr);
 	~FDataflowOutlinerView();
 
 	/** Create the outliner widget */
@@ -41,9 +42,9 @@ private:
 	TSharedPtr<SSceneOutliner> OutlinerWidget;
 
 	/** Construction scene the outliner could refer to */
-	FDataflowPreviewSceneBase* ConstructionScene;
+	TWeakPtr<FDataflowConstructionScene> ConstructionScene;
 
 	/** Simulation scene the outliner could refer to */
-	FDataflowPreviewSceneBase* SimulationScene;
+	TWeakPtr<FDataflowSimulationScene> SimulationScene;
 };
 
diff --git a/Engine/Plugins/Experimental/Dataflow/Source/DataflowEditor/Private/Dataflow/DataflowPreviewProfileController.h b/Engine/Plugins/Experimental/Dataflow/Source/DataflowEditor/Private/Dataflow/DataflowPreviewProfileController.h
--- a/Engine/Plugins/Experimental/Dataflow/Source/DataflowEditor/Private/Dataflow/DataflowPreviewProfileController.h
+++ b/Engine/Plugins/Experimental/Dataflow/Source/DataflowEditor/Private/Dataflow/DataflowPreviewProfileController.h
@@ -2,7 +2,6 @@
 
 #pragma once
 
-#include "CoreMinimal.h"
 #include "IPreviewProfileController.h"
 
 class UAssetViewerSettings;
@@ -21,6 +20,9 @@ public:
 	public:
 		virtual void StoreProfileIndex(int32) = 0;
 		virtual int32 RetrieveProfileIndex() = 0;
+
+	protected:
+		virtual ~IProfileIndexStorage() = default;
 	};
 
 	FDataflowPreviewProfileController(TSharedPtr<IProfileIndexStorage> ProfileIndexStorage);
diff --git a/Engine/Plugins/Experimental/Dataflow/Source/DataflowEditor/Private/Dataflow/DataflowSceneProfileIndexStorage.cpp b/Engine/Plugins/Experimental/Dataflow/Source/DataflowEditor/Private/Dataflow/DataflowSceneProfileIndexStorage.cpp
--- a/Engine/Plugins/Experimental/Dataflow/Source/DataflowEditor/Private/Dataflow/DataflowSceneProfileIndexStorage.cpp
+++ b/Engine/Plugins/Experimental/Dataflow/Source/DataflowEditor/Private/Dataflow/DataflowSceneProfileIndexStorage.cpp
@@ -25,15 +25,14 @@ namespace UE::Dataflow::Private
 	}
 }
 
-FDataflowConstructionSceneProfileIndexStorage::FDataflowConstructionSceneProfileIndexStorage(FDataflowConstructionScene* ConstructionScene) :
-	ConstructionScene(ConstructionScene)
+FDataflowConstructionSceneProfileIndexStorage::FDataflowConstructionSceneProfileIndexStorage(TWeakPtr<FDataflowConstructionScene> InConstructionScene) :
+	ConstructionScene(InConstructionScene)
 {
-	checkf(ConstructionScene, TEXT("Received invalid construction scene pointer"));
 }
 
 void FDataflowConstructionSceneProfileIndexStorage::StoreProfileIndex(int32 Index)
 {
-	if (UAssetViewerSettings* AssetViewerSettings = UAssetViewerSettings::Get())
+	if (const UAssetViewerSettings* const AssetViewerSettings = UAssetViewerSettings::Get())
 	{
 		if (AssetViewerSettings->Profiles.IsValidIndex(Index))
 		{
@@ -43,7 +42,10 @@ void FDataflowConstructionSceneProfileIndexStorage::StoreProfileIndex(int32 Inde
 				Options->SaveConfig();
 			}
 
-			ConstructionScene->SetCurrentProfileIndex(Index);
+			if (const TSharedPtr<FDataflowConstructionScene> ConstructionScenePtr = ConstructionScene.Pin())
+			{
+				ConstructionScenePtr->SetCurrentProfileIndex(Index);
+			}
 		}
 	}
 }
@@ -53,25 +55,28 @@ int32 FDataflowConstructionSceneProfileIndexStorage::RetrieveProfileIndex()
 	if (const UDataflowEditorOptions* const Options = UDataflowEditorOptions::StaticClass()->GetDefaultObject<UDataflowEditorOptions>())
 	{
 		const int32 FoundIndex = UE::Dataflow::Private::GetProfileIndex(Options->ConstructionProfileName);
-		if (FoundIndex != ConstructionScene->GetCurrentProfileIndex() && (FoundIndex != INDEX_NONE))
+		if (FoundIndex != INDEX_NONE)
 		{
-			ConstructionScene->SetCurrentProfileIndex(FoundIndex);
+			const TSharedPtr<FDataflowConstructionScene> ConstructionScenePtr = ConstructionScene.Pin();
+			if (ConstructionScenePtr && FoundIndex != ConstructionScenePtr->GetCurrentProfileIndex())
+			{
+				ConstructionScenePtr->SetCurrentProfileIndex(FoundIndex);
+			}
+			return FoundIndex;
 		}
-		return FoundIndex;
 	}
 	return INDEX_NONE;
 }
 
 
-FDataflowSimulationSceneProfileIndexStorage::FDataflowSimulationSceneProfileIndexStorage(FDataflowSimulationScene* SimulationScene) :
-	SimulationScene(SimulationScene)
+FDataflowSimulationSceneProfileIndexStorage::FDataflowSimulationSceneProfileIndexStorage(TWeakPtr<FDataflowSimulationScene> InSimulationScene) :
+	SimulationScene(InSimulationScene)
 {
-	checkf(SimulationScene, TEXT("Received invalid simulation scene pointer"));
 }
 
 void FDataflowSimulationSceneProfileIndexStorage::StoreProfileIndex(int32 Index)
 {
-	if (UAssetViewerSettings* AssetViewerSettings = UAssetViewerSettings::Get())
+	if (const UAssetViewerSettings* const AssetViewerSettings = UAssetViewerSettings::Get())
 	{
 		if (AssetViewerSettings->Profiles.IsValidIndex(Index))
 		{
@@ -81,7 +86,10 @@ void FDataflowSimulationSceneProfileIndexStorage::StoreProfileIndex(int32 Index)
 				Options->SaveConfig();
 			}
 
-			SimulationScene->SetCurrentProfileIndex(Index);
+			if (const TSharedPtr<FDataflowSimulationScene> SimulationScenePtr = SimulationScene.Pin())
+			{
+				SimulationScenePtr->SetCurrentProfileIndex(Index);
+			}
 		}
 	}
 }
@@ -91,12 +99,15 @@ int32 FDataflowSimulationSceneProfileIndexStorage::RetrieveProfileIndex()
 	if (const UDataflowEditorOptions* const Options = UDataflowEditorOptions::StaticClass()->GetDefaultObject<UDataflowEditorOptions>())
 	{
 		const int32 FoundIndex = UE::Dataflow::Private::GetProfileIndex(Options->SimulationProfileName);
-		if (FoundIndex != SimulationScene->GetCurrentProfileIndex() && (FoundIndex != INDEX_NONE))
+		if (FoundIndex != INDEX_NONE)
 		{
-			SimulationScene->SetCurrentProfileIndex(FoundIndex);
+			const TSharedPtr<FDataflowSimulationScene> SimulationScenePtr = SimulationScene.Pin();
+			if (SimulationScenePtr && FoundIndex != SimulationScenePtr->GetCurrentProfileIndex())
+			{
+				SimulationScenePtr->SetCurrentProfileIndex(FoundIndex);
+			}
+			return FoundIndex;
 		}
-		return FoundIndex;
 	}
-
 	return INDEX_NONE;
 }
diff --git a/Engine/Plugins/Experimental/Dataflow/Source/DataflowEditor/Private/Dataflow/DataflowSceneProfileIndexStorage.h b/Engine/Plugins/Experimental/Dataflow/Source/DataflowEditor/Private/Dataflow/DataflowSceneProfileIndexStorage.h
--- a/Engine/Plugins/Experimental/Dataflow/Source/DataflowEditor/Private/Dataflow/DataflowSceneProfileIndexStorage.h
+++ b/Engine/Plugins/Experimental/Dataflow/Source/DataflowEditor/Private/Dataflow/DataflowSceneProfileIndexStorage.h
@@ -11,29 +11,29 @@ class FDataflowConstructionSceneProfileIndexStorage : public FDataflowPreviewPro
 {
 public:
 
-	FDataflowConstructionSceneProfileIndexStorage(FDataflowConstructionScene* ConstructionScene);
-	virtual ~FDataflowConstructionSceneProfileIndexStorage() = default;
+	explicit FDataflowConstructionSceneProfileIndexStorage(TWeakPtr<FDataflowConstructionScene> InConstructionScene);
+	virtual ~FDataflowConstructionSceneProfileIndexStorage() override = default;
 
 	virtual void StoreProfileIndex(int32 Index) override;
 	virtual int32 RetrieveProfileIndex() override;
 
 private:
 
-	FDataflowConstructionScene* ConstructionScene;
+	TWeakPtr<FDataflowConstructionScene> ConstructionScene;
 };
 
 class FDataflowSimulationSceneProfileIndexStorage : public FDataflowPreviewProfileController::IProfileIndexStorage
 {
 public:
 
-	FDataflowSimulationSceneProfileIndexStorage(FDataflowSimulationScene* SimulationScene);
-	virtual ~FDataflowSimulationSceneProfileIndexStorage() = default;
+	explicit FDataflowSimulationSceneProfileIndexStorage(TWeakPtr<FDataflowSimulationScene> InSimulationScene);
+	virtual ~FDataflowSimulationSceneProfileIndexStorage() override = default;
 
 	virtual void StoreProfileIndex(int32 Index) override;
 	virtual int32 RetrieveProfileIndex() override;
 
 private:
 
-	FDataflowSimulationScene* SimulationScene;
+	TWeakPtr<FDataflowSimulationScene> SimulationScene;
 };
 
diff --git a/Engine/Plugins/Experimental/Dataflow/Source/DataflowEditor/Private/Dataflow/DataflowSimulationViewport.cpp b/Engine/Plugins/Experimental/Dataflow/Source/DataflowEditor/Private/Dataflow/DataflowSimulationViewport.cpp
--- a/Engine/Plugins/Experimental/Dataflow/Source/DataflowEditor/Private/Dataflow/DataflowSimulationViewport.cpp
+++ b/Engine/Plugins/Experimental/Dataflow/Source/DataflowEditor/Private/Dataflow/DataflowSimulationViewport.cpp
@@ -255,7 +255,7 @@ TSharedPtr<IPreviewProfileController> SDataflowSimulationViewport::CreatePreview
 	
 	if (!ProfileIndexStorage)
 	{
-		ProfileIndexStorage = MakeShared<FDataflowSimulationSceneProfileIndexStorage>(GetSimulationScene().Get());
+		ProfileIndexStorage = MakeShared<FDataflowSimulationSceneProfileIndexStorage>(GetSimulationScene());
 	}
 	
 	return MakeShared<FDataflowPreviewProfileController>(ProfileIndexStorage);
diff --git a/Engine/Plugins/Experimental/Dataflow/Source/DataflowEditor/Private/Dataflow/DataflowSimulationViewportClient.cpp b/Engine/Plugins/Experimental/Dataflow/Source/DataflowEditor/Private/Dataflow/DataflowSimulationViewportClient.cpp
--- a/Engine/Plugins/Experimental/Dataflow/Source/DataflowEditor/Private/Dataflow/DataflowSimulationViewportClient.cpp
+++ b/Engine/Plugins/Experimental/Dataflow/Source/DataflowEditor/Private/Dataflow/DataflowSimulationViewportClient.cpp
@@ -21,9 +21,9 @@
 #include "SNodePanel.h"
 
 FDataflowSimulationViewportClient::FDataflowSimulationViewportClient(FEditorModeTools* InModeTools,
-                                                             FPreviewScene* InPreviewScene,  const bool bCouldTickScene,
+                                                             TWeakPtr<FDataflowSimulationScene> InDataflowSimulationScene,  const bool bCouldTickScene,
                                                              const TWeakPtr<SEditorViewport> InEditorViewportWidget)
-	: FDataflowEditorViewportClientBase(InModeTools, InPreviewScene, bCouldTickScene, InEditorViewportWidget)
+	: FDataflowEditorViewportClientBase(InModeTools, InDataflowSimulationScene, bCouldTickScene, InEditorViewportWidget)
 {
 	// We want our near clip plane to be quite close so that we can zoom in further.
 	OverrideNearClipPlane(KINDA_SMALL_NUMBER);
@@ -64,9 +64,9 @@ void FDataflowSimulationViewportClient::SetToolCommandList(TWeakPtr<FUICommandLi
 void FDataflowSimulationViewportClient::Tick(float DeltaSeconds)
 {
 	Super::Tick(DeltaSeconds);
-	if (DataflowPreviewScene)
+	if (const TSharedPtr<FDataflowPreviewSceneBase> DataflowPreviewScenePtr = DataflowPreviewScene.Pin())
 	{
-		DataflowPreviewScene->TickDataflowScene(DeltaSeconds);
+		DataflowPreviewScenePtr->TickDataflowScene(DeltaSeconds);
 	}
 }
 
diff --git a/Engine/Plugins/Experimental/Dataflow/Source/DataflowEditor/Private/DataflowEditorTools/DataflowEditorWeightMapPaintTool.cpp b/Engine/Plugins/Experimental/Dataflow/Source/DataflowEditor/Private/DataflowEditorTools/DataflowEditorWeightMapPaintTool.cpp
--- a/Engine/Plugins/Experimental/Dataflow/Source/DataflowEditor/Private/DataflowEditorTools/DataflowEditorWeightMapPaintTool.cpp
+++ b/Engine/Plugins/Experimental/Dataflow/Source/DataflowEditor/Private/DataflowEditorTools/DataflowEditorWeightMapPaintTool.cpp
@@ -173,7 +173,7 @@ void UDataflowEditorWeightMapPaintTool::Setup()
 	SetToolDisplayName(LOCTEXT("ToolName", "Paint Weight Maps"));
 
 	// Hide all meshes in the DataflowConstructionScene, as we will be painting onto our own Preview mesh
-	if (FDataflowConstructionScene* Scene = Mode->GetDataflowConstructionScene())
+	if (const TSharedPtr<FDataflowConstructionScene> Scene = Mode->GetDataflowConstructionScene().Pin())
 	{
 		Scene->SetVisibility(false);
 	}
diff --git a/Engine/Plugins/Experimental/Dataflow/Source/DataflowEditor/Public/Dataflow/DataflowConstructionViewport.h b/Engine/Plugins/Experimental/Dataflow/Source/DataflowEditor/Public/Dataflow/DataflowConstructionViewport.h
--- a/Engine/Plugins/Experimental/Dataflow/Source/DataflowEditor/Public/Dataflow/DataflowConstructionViewport.h
+++ b/Engine/Plugins/Experimental/Dataflow/Source/DataflowEditor/Public/Dataflow/DataflowConstructionViewport.h
@@ -1,14 +1,10 @@
 // Copyright Epic Games, Inc. All Rights Reserved.
 #pragma once
 
-#include "CoreMinimal.h"
 #include "SBaseCharacterFXEditorViewport.h"
 #include "SCommonEditorViewportToolbarBase.h"
 
-class ADataflowActor;
-class FAdvancedPreviewScene;
 class UDataflowEditorMode;
-class FDataflowConstructionViewportClient;
 class FDataflowConstructionScene;
 class SRichTextBlock;
 
@@ -25,8 +21,6 @@ public:
 
 	void Construct(const FArguments& InArgs, const FAssetEditorViewportConstructionArgs& InViewportConstructionArgs);
 
-	FDataflowConstructionScene* GetConstructionScene() const;
-
 	// SEditorViewport
 	virtual void BindCommands() override;
 	virtual TSharedPtr<SWidget> BuildViewportToolbar() override;
@@ -43,6 +37,8 @@ public:
 	FText GetOverlayText() const;
 
 private:
+	TWeakPtr<FDataflowConstructionScene> GetConstructionScene() const;
+
 	/** Pointer to the box into which the overlay text items are added */
 	TSharedPtr<SRichTextBlock> OverlayText;
 
diff --git a/Engine/Plugins/Experimental/Dataflow/Source/DataflowEditor/Public/Dataflow/DataflowConstructionViewportClient.h b/Engine/Plugins/Experimental/Dataflow/Source/DataflowEditor/Public/Dataflow/DataflowConstructionViewportClient.h
--- a/Engine/Plugins/Experimental/Dataflow/Source/DataflowEditor/Public/Dataflow/DataflowConstructionViewportClient.h
+++ b/Engine/Plugins/Experimental/Dataflow/Source/DataflowEditor/Public/Dataflow/DataflowConstructionViewportClient.h
@@ -1,26 +1,16 @@
 // Copyright Epic Games, Inc. All Rights Reserved.
 #pragma once
 
-#include "CoreMinimal.h"
 #include "Dataflow/DataflowEditorViewportClientBase.h"
-#include "Dataflow/DataflowNodeParameters.h"
-#include "Dataflow/DataflowComponentSelectionState.h"
 #include "Dataflow/DataflowContent.h"
 #include "Delegates/Delegate.h"
-#include "InputBehaviorSet.h"
 
 #define UE_API DATAFLOWEDITOR_API
 
 
 class FDataflowEditorToolkit;
-class ADataflowActor;
-class UTransformProxy;
-class UCombinedTransformGizmo;
-class FTransformGizmoDataBinder;
-class FDataflowPreviewSceneBase;
-class UInputBehaviorSet;
+class FDataflowConstructionScene;
 class USelection;
-class FDataflowConstructionViewportClient;
 class IClickDragBehaviorTarget;
 namespace UE::Dataflow
 {
@@ -32,15 +22,13 @@ class FDataflowConstructionViewportClient : public FDataflowEditorViewportClient
 public:
 	using Super = FDataflowEditorViewportClientBase;
 
-	UE_API FDataflowConstructionViewportClient(FEditorModeTools* InModeTools, FPreviewScene* InPreviewScene,  const bool bCouldTickScene,
+	UE_API FDataflowConstructionViewportClient(FEditorModeTools* InModeTools, TWeakPtr<FDataflowConstructionScene> InConstructionScenePtr,  const bool bCouldTickScene,
 								  const TWeakPtr<SEditorViewport> InEditorViewportWidget = nullptr);
 
 	UE_API virtual ~FDataflowConstructionViewportClient();
 
 	UE_API void SetConstructionViewMode(const UE::Dataflow::IDataflowConstructionViewMode* InViewMode);
 
-	// IInputBehaviorSource
-	// virtual const UInputBehaviorSet* GetInputBehaviors() const override;
 	UE_API USelection* GetSelectedComponents();
 
 	/** Set the data flow toolkit used to create the client*/
@@ -85,7 +73,7 @@ private:
 	TWeakPtr<FDataflowEditorToolkit> DataflowEditorToolkitPtr = nullptr;
 
 	/** Dataflow preview scene from the toolkit */
-	FDataflowPreviewSceneBase* PreviewScene = nullptr;
+	TWeakPtr<FDataflowConstructionScene> ConstructionScene;
 
 	// @todo(brice) : Is this needed?
 	TWeakPtr<FUICommandList> ToolCommandList;
diff --git a/Engine/Plugins/Experimental/Dataflow/Source/DataflowEditor/Public/Dataflow/DataflowEditorMode.h b/Engine/Plugins/Experimental/Dataflow/Source/DataflowEditor/Public/Dataflow/DataflowEditorMode.h
--- a/Engine/Plugins/Experimental/Dataflow/Source/DataflowEditor/Public/Dataflow/DataflowEditorMode.h
+++ b/Engine/Plugins/Experimental/Dataflow/Source/DataflowEditor/Public/Dataflow/DataflowEditorMode.h
@@ -3,19 +3,16 @@
 #pragma once
 
 #include "BaseCharacterFXEditorMode.h"
-#include "Dataflow/DataflowContent.h"
 #include "DataflowEditorMode.generated.h"
 
 #define UE_API DATAFLOWEDITOR_API
 
-class UDataflowComponent;
 class FDataflowConstructionScene;
 class SDataflowGraphEditor;
 class FDataflowConstructionViewportClient;
 class FDataflowSimulationScene;
 class FDataflowSimulationViewportClient;
-class IDataflowEditorToolBuilder;
-struct FManagedArrayCollection;
+class UDataflowEdNode;
 class UEdGraphNode;
 class FViewport;
 class UDataflowEditor;
@@ -74,14 +71,14 @@ public:
 	}
 
 	/** Set the construction scene to store the dynamic meshes generated by the tools */
-	UE_API void SetDataflowConstructionScene(FDataflowConstructionScene* DataflowPreviewScene);
-	FDataflowConstructionScene* GetDataflowConstructionScene() { return ConstructionScene; }
-	const FDataflowConstructionScene* GetDataflowConstructionScene() const { return ConstructionScene; }
+	UE_API void SetDataflowConstructionScene(TWeakPtr<FDataflowConstructionScene> InConstructionScene);
+	const TWeakPtr<FDataflowConstructionScene>& GetDataflowConstructionScene() { return ConstructionScene; }
+	TWeakPtr<const FDataflowConstructionScene> GetDataflowConstructionScene() const { return ConstructionScene; }
 
 	/** Set the simulation scene to store the simulation components */
-	UE_API void SetDataflowSimlationScene(FDataflowSimulationScene* InSimulationScene);
-	FDataflowSimulationScene* GetDataflowSimulationScene() { return SimulationScene; }
-	const FDataflowSimulationScene* GetDataflowSimulationScene() const { return SimulationScene; }
+	void SetDataflowSimlationScene(TWeakPtr<FDataflowSimulationScene> InSimulationScene) { SimulationScene = InSimulationScene; }
+	const TWeakPtr<FDataflowSimulationScene>& GetDataflowSimulationScene() { return SimulationScene; }
+	TWeakPtr<const FDataflowSimulationScene> GetDataflowSimulationScene() const { return SimulationScene; }
 
 	/** Set the data flow graph editor to create nodes once the tools have ended*/
 	UE_API void SetDataflowGraphEditor(TSharedPtr<SDataflowGraphEditor> DataflowGraphEditor);
@@ -189,10 +186,10 @@ private:
 	UE_API const TArray<FName>& GetToolCategories();
 
 	/** Dataflow construction scene from the toolkit */
-	FDataflowConstructionScene* ConstructionScene = nullptr;
+	TWeakPtr<FDataflowConstructionScene> ConstructionScene;
 
 	/** Dataflow simulation scene from the toolkit */
-	FDataflowSimulationScene* SimulationScene = nullptr;
+	TWeakPtr<FDataflowSimulationScene> SimulationScene;
 
 	/** Correspondence between node types and commands to add the node to the graph */
 	TMap<FName, TSharedPtr<const FUICommandInfo>> NodeTypeToAddNodeCommandMap;
diff --git a/Engine/Plugins/Experimental/Dataflow/Source/DataflowEditor/Public/Dataflow/DataflowEditorPreviewSceneBase.h b/Engine/Plugins/Experimental/Dataflow/Source/DataflowEditor/Public/Dataflow/DataflowEditorPreviewSceneBase.h
--- a/Engine/Plugins/Experimental/Dataflow/Source/DataflowEditor/Public/Dataflow/DataflowEditorPreviewSceneBase.h
+++ b/Engine/Plugins/Experimental/Dataflow/Source/DataflowEditor/Public/Dataflow/DataflowEditorPreviewSceneBase.h
@@ -144,7 +144,7 @@ public:
 	UE_API void UnregisterSceneElements();
 
 	/** Return the scene selected components */
-	static UE_API USelection* GetSelectedComponents(FDataflowPreviewSceneBase* DataflowScene);
+	UE_API USelection* GetSelectedComponents() const;
 
 	/** Get the dataflow editor */
 	const UDataflowEditor* GetDataflowEditor() const { return DataflowEditor; }
diff --git a/Engine/Plugins/Experimental/Dataflow/Source/DataflowEditor/Public/Dataflow/DataflowEditorToolkit.h b/Engine/Plugins/Experimental/Dataflow/Source/DataflowEditor/Public/Dataflow/DataflowEditorToolkit.h
--- a/Engine/Plugins/Experimental/Dataflow/Source/DataflowEditor/Public/Dataflow/DataflowEditorToolkit.h
+++ b/Engine/Plugins/Experimental/Dataflow/Source/DataflowEditor/Public/Dataflow/DataflowEditorToolkit.h
@@ -3,30 +3,22 @@
 #pragma once
 
 #include "BaseCharacterFXEditorToolkit.h"
-#include "CoreMinimal.h"
-#include "Dataflow/DataflowObjectInterface.h"
 #include "Misc/NotifyHook.h"
-#include "GraphEditor.h"
 #include "TickableEditorObject.h"
-#include "Dataflow/DataflowSelectionView.h"
-#include "Dataflow/DataflowCollectionSpreadSheet.h"
-#include "Dataflow/DataflowConstructionViewport.h"
-#include "Dataflow/DataflowSimulationViewport.h"
 #include "Dataflow/DataflowContent.h"
 
 #define UE_API DATAFLOWEDITOR_API
 
 class FDocumentTracker;
-class FEditorViewportTabContent;
 class IDetailsView;
 class FTabManager;
 class FTabInfo;
+class IDataflowViewListener;
 class IStructureDetailsView;
-class IToolkitHost;
 class UDataflow;
 class UDataflowSubGraph;
-class USkeletalMesh;
 class SDataflowGraphEditor;
+class FDataflowCollectionSpreadSheet;
 class FDataflowConstructionScene;
 class FDataflowSimulationViewportClient;
 class UDataflowBaseContent;
@@ -34,9 +26,13 @@ class FDataflowSimulationScene;
 class FDataflowSkeletonView;
 class FDataflowOutlinerView;
 class UDataflowEditor;
-class SDataflowMembersWidget;
 class FDataflowSimulationSceneProfileIndexStorage;
 class FDataflowOutputLog;
+class FDataflowSelectionView;
+class SDataflowMembersWidget;
+class SDataflowConstructionViewport;
+class SDataflowSimulationViewport;
+class SGraphEditor;
 struct FDataflowPath;
 
 enum class EDataflowEditorEvaluationMode : uint8;
@@ -86,8 +82,8 @@ public:
 	UE_API virtual void UnregisterTabSpawners(const TSharedRef<class FTabManager>& TabManager) override;
 
 	/** Dataflow preview scenes accessor */
-	const TSharedPtr<FDataflowSimulationScene>& GetSimulationScene() const {return SimulationScene;}
-	UE_API FDataflowConstructionScene* GetConstructionScene() const;
+	const TSharedPtr<FDataflowSimulationScene>& GetSimulationScene() const { return SimulationScene; }
+	const TSharedPtr<FDataflowConstructionScene>& GetConstructionScene() const { return ConstructionScene; }
 	const TSharedPtr<FDataflowSimulationSceneProfileIndexStorage>& GetSimulationSceneProfileIndexStorage() const { return SimulationSceneProfileIndexStorage; }
 
 	// FSerializableObject interface
@@ -301,7 +297,7 @@ private:
 	TSet<TObjectPtr<UDataflowEdNode>> SelectedDataflowNodes;
 
 	/** Pointer to the construction viewport scene. Note this is an alias of ObjectScene in FBaseCharacterFXEditorToolkit but with the specific type */
-	FDataflowConstructionScene* ConstructionScene;
+	TSharedPtr<FDataflowConstructionScene> ConstructionScene;
 
 	/** PreviewScene showing the objects being simulated */
 	TSharedPtr<FDataflowSimulationScene> SimulationScene;
diff --git a/Engine/Plugins/Experimental/Dataflow/Source/DataflowEditor/Public/Dataflow/DataflowEditorViewportClientBase.h b/Engine/Plugins/Experimental/Dataflow/Source/DataflowEditor/Public/Dataflow/DataflowEditorViewportClientBase.h
--- a/Engine/Plugins/Experimental/Dataflow/Source/DataflowEditor/Public/Dataflow/DataflowEditorViewportClientBase.h
+++ b/Engine/Plugins/Experimental/Dataflow/Source/DataflowEditor/Public/Dataflow/DataflowEditorViewportClientBase.h
@@ -8,7 +8,6 @@
 
 struct FDataflowBaseElement;
 class FEditorModeTools;
-class FPreviewScene;
 class SEditorViewport;
 class HHitProxy;
 class FDataflowPreviewSceneBase;
@@ -18,7 +17,7 @@ class FDataflowEditorViewportClientBase : public FEditorViewportClient, public I
 public:
 	using Super = FEditorViewportClient;
 
-	UE_API FDataflowEditorViewportClientBase(FEditorModeTools* InModeTools, FPreviewScene* InPreviewScene,  const bool bCouldTickScene, const TWeakPtr<SEditorViewport> InEditorViewportWidget = nullptr);
+	UE_API FDataflowEditorViewportClientBase(FEditorModeTools* InModeTools, TWeakPtr<FDataflowPreviewSceneBase> InPreviewScene,  const bool bCouldTickScene, const TWeakPtr<SEditorViewport> InEditorViewportWidget = nullptr);
 	
 	UE_API virtual ~FDataflowEditorViewportClientBase();
 
@@ -34,7 +33,7 @@ protected:
 	//~ End FEditorViewportClient interface
 
 	/** Dataflow preview scene from the toolkit */
-	FDataflowPreviewSceneBase* DataflowPreviewScene = nullptr;
+	TWeakPtr<FDataflowPreviewSceneBase> DataflowPreviewScene;
 
 	/** Behaviors defined by this base class */
 	TArray<TObjectPtr<UInputBehavior>> BaseBehaviors;
diff --git a/Engine/Plugins/Experimental/Dataflow/Source/DataflowEditor/Public/Dataflow/DataflowSimulationViewportClient.h b/Engine/Plugins/Experimental/Dataflow/Source/DataflowEditor/Public/Dataflow/DataflowSimulationViewportClient.h
--- a/Engine/Plugins/Experimental/Dataflow/Source/DataflowEditor/Public/Dataflow/DataflowSimulationViewportClient.h
+++ b/Engine/Plugins/Experimental/Dataflow/Source/DataflowEditor/Public/Dataflow/DataflowSimulationViewportClient.h
@@ -11,12 +11,7 @@
 #define UE_API DATAFLOWEDITOR_API
 
 class FDataflowEditorToolkit;
-class ADataflowActor;
-class UTransformProxy;
-class UCombinedTransformGizmo;
-class FTransformGizmoDataBinder;
-class FDataflowPreviewSceneBase;
-class UInputBehaviorSet;
+class FDataflowSimulationScene;
 
 class FDataflowSimulationViewportClient : public FDataflowEditorViewportClientBase
 {
@@ -24,7 +19,7 @@ class FDataflowSimulationViewportClient : public FDataflowEditorViewportClientBa
 public:
 	using Super = FDataflowEditorViewportClientBase;
 
-	UE_API FDataflowSimulationViewportClient(FEditorModeTools* InModeTools, FPreviewScene* InPreviewScene,  const bool bCouldTickScene,
+	UE_API FDataflowSimulationViewportClient(FEditorModeTools* InModeTools, TWeakPtr<FDataflowSimulationScene> InSimulationScene,  const bool bCouldTickScene,
 								  const TWeakPtr<SEditorViewport> InEditorViewportWidget = nullptr);
 
 	UE_API virtual ~FDataflowSimulationViewportClient();
