Dataflow:
- Added rendering support for TArray<FBox> and TArray<FSphere>
- Changed Experimental node color to Darker Turquoise
- Bug fixes

diff --git a/Engine/Plugins/Experimental/Dataflow/Source/DataflowEditor/Private/Dataflow/DataflowEngineRendering.cpp b/Engine/Plugins/Experimental/Dataflow/Source/DataflowEditor/Private/Dataflow/DataflowEngineRendering.cpp
--- a/Engine/Plugins/Experimental/Dataflow/Source/DataflowEditor/Private/Dataflow/DataflowEngineRendering.cpp
+++ b/Engine/Plugins/Experimental/Dataflow/Source/DataflowEditor/Private/Dataflow/DataflowEngineRendering.cpp
@@ -799,6 +799,33 @@ namespace UE::Dataflow
 		}
 	};
 
+	class FBoxesSurfaceRenderCallbacks : public FRenderingFactory::ICallbackInterface
+	{
+		virtual UE::Dataflow::FRenderKey GetRenderKey() const override
+		{
+			return { "BoxesRender", FName("TArray<FBox>") };
+		}
+
+		virtual bool CanRender(const UE::Dataflow::IDataflowConstructionViewMode& ViewMode) const override
+		{
+			return (ViewMode.GetName() == FDataflowConstruction3DViewMode::Name);
+		}
+
+		virtual void Render(GeometryCollection::Facades::FRenderingFacade& RenderCollection, const FGraphRenderingState& State) override
+		{
+			if (State.GetRenderOutputs().Num())
+			{
+				FName PrimaryOutput = State.GetRenderOutputs()[0]; // "Box"
+
+				TArray<FBox> Default;
+				const TArray<FBox>& Boxes = State.GetValue<TArray<FBox>>(PrimaryOutput, Default);
+
+				int32 GeometryIndex = RenderCollection.StartGeometryGroup(State.GetGuid().ToString());
+				RenderCollection.AddBoxes(Boxes);
+				RenderCollection.EndGeometryGroup(GeometryIndex);
+			}
+		}
+	};
 
 	class FSphereSurfaceRenderCallbacks : public FRenderingFactory::ICallbackInterface
 	{
@@ -828,6 +855,34 @@ namespace UE::Dataflow
 		}
 	};
 
+	class FSpheresSurfaceRenderCallbacks : public FRenderingFactory::ICallbackInterface
+	{
+		virtual UE::Dataflow::FRenderKey GetRenderKey() const override
+		{
+			return { "SpheresRender", FName("TArray<FSphere>") };
+		}
+
+		virtual bool CanRender(const UE::Dataflow::IDataflowConstructionViewMode& ViewMode) const override
+		{
+			return (ViewMode.GetName() == FDataflowConstruction3DViewMode::Name);
+		}
+
+		virtual void Render(GeometryCollection::Facades::FRenderingFacade& RenderCollection, const FGraphRenderingState& State) override
+		{
+			if (State.GetRenderOutputs().Num())
+			{
+				FName PrimaryOutput = State.GetRenderOutputs()[0]; // "Spheres"
+
+				TArray<FSphere> Default;
+				const TArray<FSphere>& Spheres = State.GetValue<TArray<FSphere>>(PrimaryOutput, Default);
+
+				int32 GeometryIndex = RenderCollection.StartGeometryGroup(State.GetGuid().ToString());
+				RenderCollection.AddSpheres(Spheres, FLinearColor(FDataflowEditorModule::SurfaceColor));
+				RenderCollection.EndGeometryGroup(GeometryIndex);
+			}
+		}
+	};
+
 	class FPointRenderCallbacks : public FRenderingFactory::ICallbackInterface
 	{
 		virtual UE::Dataflow::FRenderKey GetRenderKey() const override
@@ -1075,7 +1130,9 @@ namespace UE::Dataflow
 		FRenderingFactory::GetInstance()->RegisterCallbacks(MakeUnique<FDynamicMesh3SurfaceRenderCallbacks>());
 		FRenderingFactory::GetInstance()->RegisterCallbacks(MakeUnique<FDataflowMeshSurfaceRenderCallbacks>());
 		FRenderingFactory::GetInstance()->RegisterCallbacks(MakeUnique<FBoxSurfaceRenderCallbacks>());
+		FRenderingFactory::GetInstance()->RegisterCallbacks(MakeUnique<FBoxesSurfaceRenderCallbacks>());
 		FRenderingFactory::GetInstance()->RegisterCallbacks(MakeUnique<FSphereSurfaceRenderCallbacks>());
+		FRenderingFactory::GetInstance()->RegisterCallbacks(MakeUnique<FSpheresSurfaceRenderCallbacks>());
 		FRenderingFactory::GetInstance()->RegisterCallbacks(MakeUnique<FPointRenderCallbacks>());
 		FRenderingFactory::GetInstance()->RegisterCallbacks(MakeUnique<FPointsRenderCallbacks>());
 		FRenderingFactory::GetInstance()->RegisterCallbacks(MakeUnique<FGeometryCollectionTetrahedronRenderCallbacks>());
diff --git a/Engine/Plugins/Experimental/Dataflow/Source/DataflowEditor/Private/Dataflow/DataflowSNode.cpp b/Engine/Plugins/Experimental/Dataflow/Source/DataflowEditor/Private/Dataflow/DataflowSNode.cpp
--- a/Engine/Plugins/Experimental/Dataflow/Source/DataflowEditor/Private/Dataflow/DataflowSNode.cpp
+++ b/Engine/Plugins/Experimental/Dataflow/Source/DataflowEditor/Private/Dataflow/DataflowSNode.cpp
@@ -326,7 +326,7 @@ TArray<FOverlayWidgetInfo> SDataflowEdNode::GetOverlayWidgets(bool bSelected, co
 	return Widgets;
 }
 
-const FSlateColor CExperimentalColor(FColor(26, 188, 156, 255)); // Turquoise
+const FSlateColor CExperimentalColor(FColor(13, 94, 78, 255)); // Darker Turquoise
 const FSlateColor CDeprecatedColor(FColor(46, 204, 113, 255)); // Emerald
 const FSlateColor CFailedColor(FColor(243, 156, 18, 255)); // Orange
 
diff --git a/Engine/Source/Runtime/Experimental/Chaos/Private/GeometryCollection/Facades/CollectionBoundsFacade.cpp b/Engine/Source/Runtime/Experimental/Chaos/Private/GeometryCollection/Facades/CollectionBoundsFacade.cpp
--- a/Engine/Source/Runtime/Experimental/Chaos/Private/GeometryCollection/Facades/CollectionBoundsFacade.cpp
+++ b/Engine/Source/Runtime/Experimental/Chaos/Private/GeometryCollection/Facades/CollectionBoundsFacade.cpp
@@ -374,18 +374,20 @@ namespace GeometryCollection::Facades
 				const TManagedArray<int32>& VertexCountArr = VertexCountAttribute.Get();
 
 				TArray<FVector> VerticesInCollectionSpace;
-				VerticesInCollectionSpace.AddUninitialized(Vertex.Num());
+				VerticesInCollectionSpace.Reserve(Vertex.Num());
 
 				for (int32 TransformIndex = 0; TransformIndex < CollectionSpaceTransforms.Num(); ++TransformIndex)
 				{
 					const FTransform CollectionSpaceTransform = CollectionSpaceTransforms[TransformIndex];
 					const int32 GeoIndex = TransformToGeometryIndexAttribute[TransformIndex];
-					const int32 VertexStart = VertexStartArr[GeoIndex];
-					const int32 VertexCount = VertexCountArr[GeoIndex];
-
-					for (int32 VertexIdx = VertexStart; VertexIdx < VertexStart + VertexCount; ++VertexIdx)
+					if (GeoIndex != INDEX_NONE)
 					{
-						VerticesInCollectionSpace[VertexIdx] = CollectionSpaceTransform.TransformPosition((FVector)Vertex[VertexIdx]);
+						const int32 VertexStart = VertexStartArr[GeoIndex];
+						const int32 VertexCount = VertexCountArr[GeoIndex];
+						for (int32 VertexIdx = VertexStart; VertexIdx < VertexStart + VertexCount; ++VertexIdx)
+						{
+							VerticesInCollectionSpace.Add(CollectionSpaceTransform.TransformPosition((FVector)Vertex[VertexIdx]));
+						}
 					}
 				}
 
diff --git a/Engine/Source/Runtime/Experimental/Chaos/Private/GeometryCollection/Facades/CollectionRenderingFacade.cpp b/Engine/Source/Runtime/Experimental/Chaos/Private/GeometryCollection/Facades/CollectionRenderingFacade.cpp
--- a/Engine/Source/Runtime/Experimental/Chaos/Private/GeometryCollection/Facades/CollectionRenderingFacade.cpp
+++ b/Engine/Source/Runtime/Experimental/Chaos/Private/GeometryCollection/Facades/CollectionRenderingFacade.cpp
@@ -291,6 +291,14 @@ namespace GeometryCollection::Facades
 		AddSurface(MoveTemp(Vertices), MoveTemp(Tris), MoveTemp(Normals), MoveTemp(Colors));
 	}
 
+	void FRenderingFacade::AddBoxes(const TArray<FBox>& InBoxes)
+	{
+		for (int32 Idx = 0; Idx < InBoxes.Num(); ++Idx)
+		{
+			AddBox(InBoxes[Idx]);
+		}
+	}
+
 	void FRenderingFacade::AddSphere(const FSphere& InSphere, const FLinearColor& InColor)
 	{
 		AddSphere((FVector3f)InSphere.Center, (float)InSphere.W, InColor);
@@ -377,6 +385,14 @@ namespace GeometryCollection::Facades
 		AddSurface(MoveTemp(Vertices), MoveTemp(Tris), MoveTemp(VertexNormals), MoveTemp(VertexColors));
 	}
 
+	void FRenderingFacade::AddSpheres(const TArray<FSphere>& InSpheres, const FLinearColor& InColor)
+	{
+		for (int32 Idx = 0; Idx < InSpheres.Num(); ++Idx)
+		{
+			AddSphere(InSpheres[Idx], InColor);
+		}
+	}
+
 	void FRenderingFacade::AddTetrahedron(const TArray<FVector3f>& InVertices, const FIntVector4& InIndices)
 	{
 		constexpr int32 NumVertices = 4;
diff --git a/Engine/Source/Runtime/Experimental/Chaos/Public/GeometryCollection/Facades/CollectionRenderingFacade.h b/Engine/Source/Runtime/Experimental/Chaos/Public/GeometryCollection/Facades/CollectionRenderingFacade.h
--- a/Engine/Source/Runtime/Experimental/Chaos/Public/GeometryCollection/Facades/CollectionRenderingFacade.h
+++ b/Engine/Source/Runtime/Experimental/Chaos/Public/GeometryCollection/Facades/CollectionRenderingFacade.h
@@ -62,10 +62,16 @@ namespace GeometryCollection::Facades
 		CHAOS_API void AddBox(const FBox& InBox);
 		CHAOS_API void AddBox(const FVector3f& InMinVertex, const FVector3f& InMaxVertex);
 
+		/**Add boxes to the rendering view.*/
+		CHAOS_API void AddBoxes(const TArray<FBox>& InBoxes);
+
 		/**Add a sphere to the rendering view.*/
 		CHAOS_API void AddSphere(const FSphere& InSphere, const FLinearColor& InColor);
 		CHAOS_API void AddSphere(const FVector3f& InCenter, const float InRadius, const FLinearColor& InColor);
 
+		/**Add spheres to the rendering view.*/
+		CHAOS_API void AddSpheres(const TArray<FSphere>& InSpheres, const FLinearColor& InColor);
+
 		/** Add a surface to the rendering view.*/
 		CHAOS_API void AddSurface(TArray<FVector3f>&& InVertices, TArray<FIntVector>&& InIndices, TArray<FVector3f>&& InNormals, TArray<FLinearColor>&& InColors);
 
