Dataflow : when requesting shutdown of a tool with Accept , make sure accept is enable, ortherwise default to cancel
- this is an optimization so that when Accepting a tool that has no modification , the graph does not get dirty again

diff --git a/Engine/Plugins/Experimental/Dataflow/Source/DataflowEditor/Private/Dataflow/DataflowEditorMode.cpp b/Engine/Plugins/Experimental/Dataflow/Source/DataflowEditor/Private/Dataflow/DataflowEditorMode.cpp
--- a/Engine/Plugins/Experimental/Dataflow/Source/DataflowEditor/Private/Dataflow/DataflowEditorMode.cpp
+++ b/Engine/Plugins/Experimental/Dataflow/Source/DataflowEditor/Private/Dataflow/DataflowEditorMode.cpp
@@ -1113,6 +1113,11 @@ void UDataflowEditorMode::ShutdownActiveToolIfNeeded(EToolShutdownType ShutdownT
 
 	if (UInteractiveTool* const ActiveTool = ToolManager->GetActiveTool(EToolSide::Left))
 	{
+		// Accept will potentailly update the graph, if nothing has changed so far we want to simply cancel 
+		if (ShutdownType == EToolShutdownType::Accept && !ToolsContext->CanAcceptActiveTool())
+		{
+			ShutdownType = EToolShutdownType::Completed;
+		}
 		ToolManager->PostActiveToolShutdownRequest(ActiveTool, ShutdownType);
 	}
 }
