show camera and show menus back to the construction viewport - also added the profile selection menu
- renamed the views / visualization menus to avoid conflicting with the camera and show default ones
- fixed the ortho auto-focus to better frame the objects

diff --git a/Engine/Plugins/Experimental/Dataflow/Source/DataflowEditor/Private/Dataflow/DataflowConstructionViewport.cpp b/Engine/Plugins/Experimental/Dataflow/Source/DataflowEditor/Private/Dataflow/DataflowConstructionViewport.cpp
--- a/Engine/Plugins/Experimental/Dataflow/Source/DataflowEditor/Private/Dataflow/DataflowConstructionViewport.cpp
+++ b/Engine/Plugins/Experimental/Dataflow/Source/DataflowEditor/Private/Dataflow/DataflowConstructionViewport.cpp
@@ -21,7 +21,6 @@
 
 #define LOCTEXT_NAMESPACE "SDataflowConstructionViewport"
 
-
 SDataflowConstructionViewport::SDataflowConstructionViewport()
 {
 }
@@ -47,6 +46,79 @@ TWeakPtr<FDataflowConstructionScene> SDataflowConstructionViewport::GetConstruct
 	return nullptr;
 }
 
+void SDataflowConstructionViewport::BuildVisualizationMenu(FToolMenuSection& MenuSection)
+{
+	if (UUnrealEdViewportToolbarContext* Context = MenuSection.FindContext<UUnrealEdViewportToolbarContext>())
+	{
+		MenuSection.AddSubMenu(
+			"Visualization",
+			LOCTEXT("VisualizationSubmenuLabel", "Visualization"),
+			LOCTEXT("VisualizationSubmenuTooltip", "Controls visualization options"),
+			FNewMenuDelegate::CreateLambda([WeakViewport = Context->Viewport](FMenuBuilder& Menu)
+				{
+					if (const TSharedPtr<SDataflowConstructionViewport> Viewport = StaticCastSharedPtr<SDataflowConstructionViewport>(WeakViewport.Pin()))
+					{
+						if (const TSharedPtr<FEditorViewportClient> ViewportClient = Viewport->GetViewportClient())
+						{
+							for (const TPair<FName, TUniquePtr<UE::Dataflow::IDataflowConstructionVisualization>>& Visualization : UE::Dataflow::FDataflowConstructionVisualizationRegistry::GetInstance().GetVisualizations())
+							{
+								Visualization.Value->ExtendViewportShowMenu(StaticCastSharedPtr<FDataflowConstructionViewportClient>(ViewportClient), Menu);
+							}
+						}
+					}
+				}),
+			false,
+			FSlateIcon(FAppStyle::Get().GetStyleSetName(), "DetailsView.ViewOptions")
+		);
+	}
+}
+
+void SDataflowConstructionViewport::BuildViewMenu(FToolMenuSection& MenuSection)
+{
+	if (UUnrealEdViewportToolbarContext* Context = MenuSection.FindContext<UUnrealEdViewportToolbarContext>())
+	{
+		// Camera / View Options
+		const TAttribute<FText> Label = TAttribute<FText>::CreateLambda([WeakViewport = Context->Viewport]
+			{
+				if (const TSharedPtr<SDataflowConstructionViewport> Viewport = StaticCastSharedPtr<SDataflowConstructionViewport>(WeakViewport.Pin()))
+				{
+					if (UDataflowEditorMode* EditorMode = Viewport->GetEdMode())
+					{
+						return EditorMode->GetConstructionViewMode()->GetButtonText();
+					}
+				}
+
+				return LOCTEXT("DataflowViewLabel", "View");
+			});
+
+		MenuSection.AddSubMenu(
+			"DataflowView",
+			Label,
+			LOCTEXT("DataflowViewTooltip", "Display view options for the construction viewport."),
+			FNewToolMenuDelegate::CreateLambda([WeakViewport = Context->Viewport](UToolMenu* Menu)
+				{
+					FToolMenuSection& SimulationSection = Menu->AddSection("Simulation", LOCTEXT("SimulationSection", "Simulation"));
+
+					if (const TSharedPtr<SDataflowConstructionViewport> Viewport = StaticCastSharedPtr<SDataflowConstructionViewport>(WeakViewport.Pin()))
+					{
+						if (UDataflowEditorMode* EditorMode = Viewport->GetEdMode())
+						{
+							for (const TPair<FName, TSharedPtr<FUICommandInfo>>& NameAndCommand : FDataflowEditorCommandsImpl::Get().SetConstructionViewModeCommands)
+							{
+								if (EditorMode->CanChangeConstructionViewModeTo(NameAndCommand.Key))
+								{
+									SimulationSection.AddMenuEntry(NameAndCommand.Value);
+								}
+							}
+						}
+					}
+				}),
+			false,
+			FSlateIcon(FAppStyle::Get().GetStyleSetName(), "EditorViewport.Perspective")
+		);
+	}
+}
+
 TSharedPtr<SWidget> SDataflowConstructionViewport::BuildViewportToolbar()
 {
 	const FName ToolbarName = "Dataflow.ConstructionViewportToolbar";
@@ -57,89 +129,23 @@ TSharedPtr<SWidget> SDataflowConstructionViewport::BuildViewportToolbar()
 		Menu->StyleName = "ViewportToolbar";
 		
 		Menu->AddSection("Left");
-		
 		{
-			FToolMenuSection& RightSection = Menu->AddSection("Right");
+			// nothing on the left side for now 
+		}
+
+		FToolMenuSection& RightSection = Menu->AddSection("Right");
+		{
 			RightSection.Alignment = EToolMenuSectionAlign::Last;
-			
+
+			RightSection.AddEntry(UE::UnrealEd::CreateCameraSubmenu(UE::UnrealEd::FViewportCameraMenuOptions().ShowCameraMovement()));
 			RightSection.AddEntry(UE::UnrealEd::CreateViewModesSubmenu());
-			
+			RightSection.AddEntry(UE::UnrealEd::CreateDefaultShowSubmenu());
+			RightSection.AddEntry(UE::UnrealEd::CreateAssetViewerProfileSubmenu());
 			{
-				RightSection.AddDynamicEntry("DynamicShowAndCamera", FNewToolMenuSectionDelegate::CreateLambda([](FToolMenuSection& Section)
+				RightSection.AddDynamicEntry("ViewAndVisualization", FNewToolMenuSectionDelegate::CreateLambda([](FToolMenuSection& Section)
 				{
-					UUnrealEdViewportToolbarContext* Context = Section.FindContext<UUnrealEdViewportToolbarContext>();
-					if (!Context)
-					{
-						return;
-					}
-					
-					TSharedPtr<SDataflowConstructionViewport> ConstructionViewport = StaticCastSharedPtr<SDataflowConstructionViewport>(Context->Viewport.Pin());
-					if (!ConstructionViewport)
-					{
-						return;
-					}
-				
-					{
-						// Show Menu
-						Section.AddEntry(UE::UnrealEd::CreateShowSubmenu(
-							FNewMenuDelegate::CreateLambda([WeakViewport = ConstructionViewport.ToWeakPtr()](FMenuBuilder& Menu)
-							{
-								if (const TSharedPtr<SDataflowConstructionViewport> PinnedEditorViewport = WeakViewport.Pin())
-								{
-									if (const TSharedPtr<FEditorViewportClient> ViewportClient = PinnedEditorViewport->GetViewportClient())
-									{
-										for (const TPair<FName, TUniquePtr<UE::Dataflow::IDataflowConstructionVisualization>>& Visualization : UE::Dataflow::FDataflowConstructionVisualizationRegistry::GetInstance().GetVisualizations())
-										{
-											Visualization.Value->ExtendViewportShowMenu(StaticCastSharedPtr<FDataflowConstructionViewportClient>(ViewportClient), Menu);
-										}
-									}
-								}
-							})
-						));
-					}
-					
-					{
-						
-						// Camera / View Options
-						TAttribute<FText> Label = TAttribute<FText>::CreateLambda([WeakViewport = ConstructionViewport.ToWeakPtr()]
-						{
-							if (const TSharedPtr<SDataflowConstructionViewport> Viewport = WeakViewport.Pin())
-							{
-								if (UDataflowEditorMode* EditorMode = Viewport->GetEdMode())
-								{
-									return EditorMode->GetConstructionViewMode()->GetButtonText();
-								}
-							}
-							
-							return LOCTEXT("DataflowConstructionViewMenuTitle_Default", "View");
-						});
-						
-						Section.AddSubMenu(
-							"Camera",
-							Label,
-							LOCTEXT("CameraSubMenuTooltip", "Display options for the construction viewport."),
-							FNewToolMenuDelegate::CreateLambda([WeakViewport = ConstructionViewport.ToWeakPtr()](UToolMenu* Menu)
-							{
-								FToolMenuSection& SimulationSection = Menu->AddSection("Simulation", LOCTEXT("SimulationSection", "Simulation"));
-							
-								if (const TSharedPtr<SDataflowConstructionViewport> Viewport = WeakViewport.Pin())
-								{
-									if (UDataflowEditorMode* EditorMode = Viewport->GetEdMode())
-									{
-										for (const TPair<FName, TSharedPtr<FUICommandInfo>>& NameAndCommand : FDataflowEditorCommandsImpl::Get().SetConstructionViewModeCommands)
-										{
-											if (EditorMode->CanChangeConstructionViewModeTo(NameAndCommand.Key))
-											{
-												SimulationSection.AddMenuEntry(NameAndCommand.Value);
-											}
-										}
-									}
-								}
-							}),
-							false,
-							FSlateIcon(FAppStyle::Get().GetStyleSetName(), "EditorViewport.Perspective")
-						);
-					}
+					BuildVisualizationMenu(Section);
+					BuildViewMenu(Section);			
 				}));
 			}
 		}
@@ -169,9 +175,16 @@ void SDataflowConstructionViewport::OnFocusViewportToSelection()
 
 	if (DataflowEdMode)
 	{
-		const FBox BoundingBox = DataflowEdMode->SelectionBoundingBox();
+		FBox BoundingBox = DataflowEdMode->SelectionBoundingBox();
 		if (BoundingBox.IsValid && !(BoundingBox.Min == FVector::Zero() && BoundingBox.Max == FVector::Zero()))
 		{
+			if (Client->IsOrtho())
+			{
+				const FVector Center = BoundingBox.GetCenter();
+				const FVector Extent = BoundingBox.GetExtent() * 0.5;
+				BoundingBox = FBox(Center - Extent, Center + Extent);
+			}
+
 			Client->FocusViewportOnBox(BoundingBox);
 		}
 	}
diff --git a/Engine/Plugins/Experimental/Dataflow/Source/DataflowEditor/Public/Dataflow/DataflowConstructionViewport.h b/Engine/Plugins/Experimental/Dataflow/Source/DataflowEditor/Public/Dataflow/DataflowConstructionViewport.h
--- a/Engine/Plugins/Experimental/Dataflow/Source/DataflowEditor/Public/Dataflow/DataflowConstructionViewport.h
+++ b/Engine/Plugins/Experimental/Dataflow/Source/DataflowEditor/Public/Dataflow/DataflowConstructionViewport.h
@@ -8,6 +8,8 @@ class UDataflowEditorMode;
 class FDataflowConstructionScene;
 class SRichTextBlock;
 
+struct FToolMenuSection;
+
 // ----------------------------------------------------------------------------------
 
 class SDataflowConstructionViewport : public SBaseCharacterFXEditorViewport, public ICommonEditorViewportToolbarInfoProvider
@@ -37,6 +39,9 @@ public:
 	FText GetOverlayText() const;
 
 private:
+	static void BuildVisualizationMenu(FToolMenuSection& MenuSection);
+	static void BuildViewMenu(FToolMenuSection& MenuSection);
+
 	TWeakPtr<FDataflowConstructionScene> GetConstructionScene() const;
 
 	/** Pointer to the box into which the overlay text items are added */
