Add Text3d debug draw option to Dataflow debug draw

diff --git a/Engine/Plugins/Experimental/Dataflow/Source/DataflowEditor/Private/Dataflow/DataflowEditorModule.cpp b/Engine/Plugins/Experimental/Dataflow/Source/DataflowEditor/Private/Dataflow/DataflowEditorModule.cpp
--- a/Engine/Plugins/Experimental/Dataflow/Source/DataflowEditor/Private/Dataflow/DataflowEditorModule.cpp
+++ b/Engine/Plugins/Experimental/Dataflow/Source/DataflowEditor/Private/Dataflow/DataflowEditorModule.cpp
@@ -12,6 +12,7 @@
 #include "Dataflow/DataflowFunctionProperty.h"
 #include "Dataflow/DataflowFunctionPropertyCustomization.h"
 #include "Dataflow/DataflowSNodeFactories.h"
+#include "Dataflow/DebugDrawText3dVisualization.h"
 #include "Dataflow/MeshStatsConstructionVisualization.h"
 #include "Dataflow/MeshConstructionVisualization.h"
 #include "Dataflow/ScalarVertexPropertyGroupCustomization.h"
@@ -149,6 +150,7 @@ void FDataflowEditorModule::StartupModule()
 	UE::Dataflow::FDataflowConstructionVisualizationRegistry& ConstructionVisualizationRegistry = UE::Dataflow::FDataflowConstructionVisualizationRegistry::GetInstance();
 	ConstructionVisualizationRegistry.RegisterVisualization(MakeUnique<UE::Dataflow::FMeshStatsConstructionVisualization>());
 	ConstructionVisualizationRegistry.RegisterVisualization(MakeUnique<UE::Dataflow::FMeshConstructionVisualization>());
+	ConstructionVisualizationRegistry.RegisterVisualization(MakeUnique<UE::Dataflow::FDebugDrawText3dVisualization>());
 
 	FDataflowEditorCommands::Register();
 }
@@ -178,6 +180,7 @@ void FDataflowEditorModule::ShutdownModule()
 	UE::Dataflow::FDataflowConstructionVisualizationRegistry& ConstructionVisualizationRegistry = UE::Dataflow::FDataflowConstructionVisualizationRegistry::GetInstance();
 	ConstructionVisualizationRegistry.DeregisterVisualization(UE::Dataflow::FMeshStatsConstructionVisualization::Name);
 	ConstructionVisualizationRegistry.DeregisterVisualization(UE::Dataflow::FMeshConstructionVisualization::Name);
+	ConstructionVisualizationRegistry.DeregisterVisualization(UE::Dataflow::FDebugDrawText3dVisualization::Name);
 }
 
 IMPLEMENT_MODULE(FDataflowEditorModule, DataflowEditor)
diff --git a/Engine/Plugins/Experimental/Dataflow/Source/DataflowEditor/Private/Dataflow/DebugDrawText3dVisualization.cpp b/Engine/Plugins/Experimental/Dataflow/Source/DataflowEditor/Private/Dataflow/DebugDrawText3dVisualization.cpp
new file mode 100644
--- /dev/null
+++ b/Engine/Plugins/Experimental/Dataflow/Source/DataflowEditor/Private/Dataflow/DebugDrawText3dVisualization.cpp
@@ -0,0 +1,52 @@
+// Copyright Epic Games, Inc. All Rights Reserved.
+
+#include "Dataflow/DebugDrawText3dVisualization.h"
+#include "Dataflow/DataflowConstructionScene.h"
+#include "Dataflow/DataflowDebugDrawComponent.h"
+#include "Engine/Canvas.h"
+#include "Engine/Engine.h"
+
+namespace UE::Dataflow
+{
+	namespace Private
+	{
+		static void DrawText(FCanvas* Canvas, const FSceneView* SceneView, const FVector& Pos, const FString& Text, const FLinearColor& Color, const float Scale = 1.f)
+		{
+#if WITH_EDITOR
+			if (Canvas && SceneView)
+			{
+				FVector2D PixelLocation;
+				if (SceneView->WorldToPixel(Pos, PixelLocation))
+				{
+					// WorldToPixel doesn't account for DPIScale
+					const float DPIScale = Canvas->GetDPIScale();
+					FCanvasTextItem TextItem(PixelLocation / DPIScale, FText::FromString(Text), GEngine->GetSmallFont(), Color);
+					TextItem.Scale = FVector2D::UnitVector * Scale;
+					TextItem.EnableShadow(FLinearColor::Black);
+					TextItem.Draw(Canvas);
+				}
+			}
+#endif
+		}
+	}
+
+	FName FDebugDrawText3dVisualization::Name = FName("DebugDrawText3dVisualization");
+
+	void FDebugDrawText3dVisualization::DrawCanvas(const FDataflowConstructionScene* ConstructionScene, FCanvas* Canvas, const FSceneView* SceneView)
+	{
+		if (ConstructionScene)
+		{
+			if (const UDataflowDebugDrawComponent* DebugDrawComponent = ConstructionScene->GetDebugDrawComponent())
+			{
+				if (const FDataflowDebugRenderSceneProxy* SceneProxy = static_cast<FDataflowDebugRenderSceneProxy*>(DebugDrawComponent->GetSceneProxy()))
+				{
+					for (const FDebugRenderSceneProxy::FText3d& Text : SceneProxy->Texts)
+					{
+						Private::DrawText(Canvas, SceneView, Text.Location, Text.Text, Text.Color);
+					}
+				}
+			}
+		}
+	}
+}
+
diff --git a/Engine/Plugins/Experimental/Dataflow/Source/DataflowEditor/Private/Dataflow/DebugDrawText3dVisualization.h b/Engine/Plugins/Experimental/Dataflow/Source/DataflowEditor/Private/Dataflow/DebugDrawText3dVisualization.h
new file mode 100644
--- /dev/null
+++ b/Engine/Plugins/Experimental/Dataflow/Source/DataflowEditor/Private/Dataflow/DebugDrawText3dVisualization.h
@@ -0,0 +1,22 @@
+// Copyright Epic Games, Inc. All Rights Reserved.
+#pragma once
+
+#include "Dataflow/DataflowConstructionVisualization.h"
+
+namespace UE::Dataflow
+{
+	/** Display Text3d from FDataflowDebugDraw on the editor's Canvas */
+	class FDebugDrawText3dVisualization final : public IDataflowConstructionVisualization
+	{
+	public:
+		static FName Name;
+
+	private:
+		virtual FName GetName() const override
+		{
+			return Name;
+		}
+		virtual void DrawCanvas(const FDataflowConstructionScene* ConstructionScene, FCanvas* Canvas, const FSceneView* SceneView) override;
+	};
+}
+
diff --git a/Engine/Source/Runtime/Experimental/Dataflow/Core/Public/Dataflow/DataflowDebugDrawInterface.h b/Engine/Source/Runtime/Experimental/Dataflow/Core/Public/Dataflow/DataflowDebugDrawInterface.h
--- a/Engine/Source/Runtime/Experimental/Dataflow/Core/Public/Dataflow/DataflowDebugDrawInterface.h
+++ b/Engine/Source/Runtime/Experimental/Dataflow/Core/Public/Dataflow/DataflowDebugDrawInterface.h
@@ -42,6 +42,7 @@ public:
 	virtual void DrawObject(const TRefCountPtr<IDataflowDebugDrawObject>& Object) = 0;
 	virtual void DrawPoint(const FVector& Position) = 0;
 	virtual void DrawLine(const FVector& Start, const FVector& End) const = 0;
+	virtual void DrawText3d(const FString& String, const FVector& Location) const = 0;
 
 	struct IDebugDrawMesh
 	{
diff --git a/Engine/Source/Runtime/Experimental/Dataflow/Engine/Private/Dataflow/DataflowDebugDraw.cpp b/Engine/Source/Runtime/Experimental/Dataflow/Engine/Private/Dataflow/DataflowDebugDraw.cpp
--- a/Engine/Source/Runtime/Experimental/Dataflow/Engine/Private/Dataflow/DataflowDebugDraw.cpp
+++ b/Engine/Source/Runtime/Experimental/Dataflow/Engine/Private/Dataflow/DataflowDebugDraw.cpp
@@ -217,6 +217,11 @@ void FDataflowDebugDraw::DrawCapsule(const FVector& Center, const double& Radius
 	DebugRenderSceneProxy->Capsules.Add(FDebugRenderSceneProxy::FCapsule(Center, Radius, XAxis, YAxis, ZAxis, HalfHeight, ColorWithTranslucency, DrawType));
 }
 
+void FDataflowDebugDraw::DrawText3d(const FString& String, const FVector& Location) const
+{
+	DebugRenderSceneProxy->Texts.Add(FDebugRenderSceneProxy::FText3d(String, Location, Color));
+}
+
 void FDataflowDebugDraw::DrawOverlayText(const FString& InString)
 {
 	OverlayStrings.Add(InString);
diff --git a/Engine/Source/Runtime/Experimental/Dataflow/Engine/Public/Dataflow/DataflowDebugDraw.h b/Engine/Source/Runtime/Experimental/Dataflow/Engine/Public/Dataflow/DataflowDebugDraw.h
--- a/Engine/Source/Runtime/Experimental/Dataflow/Engine/Public/Dataflow/DataflowDebugDraw.h
+++ b/Engine/Source/Runtime/Experimental/Dataflow/Engine/Public/Dataflow/DataflowDebugDraw.h
@@ -44,6 +44,7 @@ private:
 	virtual void DrawBox(const FVector& Extents, const FQuat& Rotation, const FVector& Center, double UniformScale) const override;
 	virtual void DrawSphere(const FVector& Center, double Radius) const override;
 	virtual void DrawCapsule(const FVector& Center, const double& Radius, const double& HalfHeight, const FVector& XAxis, const FVector& YAxis, const FVector &ZAxis) const override;
+	virtual void DrawText3d(const FString& String, const FVector& Location) const override;
 
 	virtual void DrawOverlayText(const FString& InString) override;
 	virtual FString GetOverlayText() const override;
