Dataflow : Add support for inline editing of bool / int / float / string typed pins
- Stay in sync with the detail panel if opened
- Disabled for anytype pins for now

diff --git a/Engine/Plugins/Experimental/Dataflow/Source/DataflowEditor/Private/Dataflow/DataflowEditorCommands.cpp b/Engine/Plugins/Experimental/Dataflow/Source/DataflowEditor/Private/Dataflow/DataflowEditorCommands.cpp
--- a/Engine/Plugins/Experimental/Dataflow/Source/DataflowEditor/Private/Dataflow/DataflowEditorCommands.cpp
+++ b/Engine/Plugins/Experimental/Dataflow/Source/DataflowEditor/Private/Dataflow/DataflowEditorCommands.cpp
@@ -416,6 +416,11 @@ void FDataflowEditorCommands::OnPropertyValueChanged(UDataflow* OutDataflow, TSh
 							SelectedNode->SetEnabledState(DataflowNode->IsActive(bCheckIsActiveFlagOnly) ? ENodeEnabledState::Enabled : ENodeEnabledState::Disabled);
 						}
 					}
+					else
+					{
+						// make sure the ditable values are up-to-date with the changes
+						SelectedNode->UpdatePinsDefaultValuesFromNode();
+					}
 					if (DataflowNode->ShouldInvalidateOnPropertyChanged(InPropertyChangedEvent))
 					{
 						// Invalidate the node and reset the editor timestamp
diff --git a/Engine/Plugins/Experimental/Dataflow/Source/DataflowEditor/Private/Dataflow/DataflowSNode.cpp b/Engine/Plugins/Experimental/Dataflow/Source/DataflowEditor/Private/Dataflow/DataflowSNode.cpp
--- a/Engine/Plugins/Experimental/Dataflow/Source/DataflowEditor/Private/Dataflow/DataflowSNode.cpp
+++ b/Engine/Plugins/Experimental/Dataflow/Source/DataflowEditor/Private/Dataflow/DataflowSNode.cpp
@@ -140,21 +140,6 @@ TSharedPtr<SGraphPin> SDataflowEdNode::CreatePinWidget(UEdGraphPin* Pin) const
 							DataflowNode->LastModifiedTimestamp.IsInvalid() ||
 							!DataflowContext->IsCacheEntryAfterTimestamp(Output->CacheKey(), DataflowNode->LastModifiedTimestamp);
 						PRAGMA_ENABLE_DEPRECATION_WARNINGS
-
-						if (FDataflowArrayTypePolicy::SupportsTypeStatic(Output->GetType()))
-						{
-							Pin->PinType.ContainerType = EPinContainerType::Array;
-						}
-					}
-				}
-			}
-			else if (Pin->Direction == EEdGraphPinDirection::EGPD_Input)
-			{
-				if (FDataflowInput* Input = DataflowNode->FindInput(Pin->GetFName()))
-				{
-					if (FDataflowArrayTypePolicy::SupportsTypeStatic(Input->GetType()))
-					{
-						Pin->PinType.ContainerType = EPinContainerType::Array;
 					}
 				}
 			}
diff --git a/Engine/Plugins/Experimental/Dataflow/Source/DataflowEditor/Private/Dataflow/DataflowSPin.cpp b/Engine/Plugins/Experimental/Dataflow/Source/DataflowEditor/Private/Dataflow/DataflowSPin.cpp
--- a/Engine/Plugins/Experimental/Dataflow/Source/DataflowEditor/Private/Dataflow/DataflowSPin.cpp
+++ b/Engine/Plugins/Experimental/Dataflow/Source/DataflowEditor/Private/Dataflow/DataflowSPin.cpp
@@ -1,6 +1,14 @@
 // Copyright Epic Games, Inc. All Rights Reserved.
 
 #include "Dataflow/DataflowSPin.h"
+#include "Dataflow/DataflowEdNode.h"
+
+#include "Widgets/Input/SCheckBox.h"
+#include "Widgets/Input/SMultiLineEditableTextBox.h"
+#include "Widgets/Input/SNumericEntryBox.h"
+#include "Widgets/Layout/SBox.h"
+
+#include "ScopedTransaction.h"
 
 #define LOCTEXT_NAMESPACE "SDataflowPin"
 
@@ -31,8 +39,7 @@ FSlateColor SDataflowPin::GetPinColor() const
 	const bool IsColorOverriden = bIsPinColorOverriden.Get();
 	const FLinearColor ColorOverride = PinColorOverride.Get();
 
-	UEdGraphPin* GraphPin = GetPinObj();
-	if (GraphPin && !GraphPin->IsPendingKill())
+	if (UEdGraphPin* GraphPin = GetPinObj())
 	{
 		if (bIsDiffHighlighted)
 		{
@@ -65,5 +72,214 @@ FSlateColor SDataflowPin::GetPinColor() const
 
 	return FLinearColor::White;
 }
-#undef LOCTEXT_NAMESPACE
 
+TSharedRef<SWidget> SDataflowPin::GetDefaultValueWidget()
+{
+	if (UEdGraphPin* GraphPin = GetPinObj())
+	{
+		if (UDataflowEdNode::SupportsEditablePinType(*GraphPin))
+		{
+			const FName Pintype = GraphPin->PinType.PinCategory;
+			if (!GraphPin->PinType.IsArray())
+			{
+				if (Pintype == TDataflowPolicyTypeName<bool>::GetName())
+				{
+					return MakeBoolWidget(*GraphPin);
+				}
+				else if (Pintype == TDataflowPolicyTypeName<int32>::GetName())
+				{
+					return MakeIntWidget(*GraphPin);
+				}
+				else if (Pintype == TDataflowPolicyTypeName<float>::GetName())
+				{
+					return MakeFloatWidget(*GraphPin);
+				}
+				else if (Pintype == TDataflowPolicyTypeName<FString>::GetName())
+				{
+					return MakeStringWidget(*GraphPin);
+				}
+			}
+		}
+	}
+	return SGraphPin::GetDefaultValueWidget();
+}
+
+TSharedPtr<FDataflowNode> SDataflowPin::GetDataflowNode() const
+{
+	if (UEdGraphPin* GraphPin = GetPinObj())
+	{
+		if (UDataflowEdNode* EdNode = Cast<UDataflowEdNode>(GraphPin->GetOwningNode()))
+		{
+			return EdNode->GetDataflowNode();
+		}
+	}
+	return {};
+}
+
+//////////////////////////////////////////////////////////////////////////////////////////
+
+ECheckBoxState SDataflowPin::GetBoolValue() const
+{
+	if (UEdGraphPin* GraphPin = GetPinObj()) 
+	{
+		FString CurrentValue = GraphPin->GetDefaultAsString();
+		return CurrentValue.ToBool() ? ECheckBoxState::Checked : ECheckBoxState::Unchecked;
+	}
+	return ECheckBoxState::Unchecked;
+}
+
+void SDataflowPin::SetBoolValue(ECheckBoxState InIsChecked)
+{
+	if (UEdGraphPin* GraphPin = GetPinObj()) 
+	{
+		const FString BoolString = (InIsChecked == ECheckBoxState::Checked) ? TEXT("true") : TEXT("false");
+		if (GraphPin->GetDefaultAsString() != BoolString)
+		{
+			const FScopedTransaction Transaction(LOCTEXT("ChangeBoolPinValue", "Change Bool Pin Value"));
+			GraphPin->Modify();
+			GraphPin->GetSchema()->TrySetDefaultValue(*GraphPin, BoolString);
+		}
+	}
+}
+
+TSharedRef<SWidget> SDataflowPin::MakeBoolWidget(const UEdGraphPin& Pin)
+{
+	return SNew(SCheckBox)
+		.IsEnabled(this, &SGraphPin::GetDefaultValueIsEditable)
+		.Visibility(this, &SGraphPin::GetDefaultValueVisibility)
+		.IsChecked(this, &SDataflowPin::GetBoolValue)
+		.OnCheckStateChanged(this, &SDataflowPin::SetBoolValue)
+		;
+}
+
+//////////////////////////////////////////////////////////////////////////////////////////
+
+TOptional<int32> SDataflowPin::GetIntValue() const
+{
+	int32 Value = 0;
+	if (UEdGraphPin* GraphPin = GetPinObj())
+	{
+		LexFromString(Value, *GraphPin->GetDefaultAsString());
+	}
+	return Value;
+}
+
+void SDataflowPin::SetIntValue(int32 InValue, ETextCommit::Type CommitType)
+{
+	if (UEdGraphPin* GraphPin = GetPinObj())
+	{
+		int32 CurrentValue = 0;
+		LexFromString(CurrentValue, *GraphPin->GetDefaultAsString());
+		if (CurrentValue != InValue)
+		{
+			const FScopedTransaction Transaction(LOCTEXT("ChangeIntPinValue", "Change Integer Pin Value"));
+			GraphPin->Modify();
+			GraphPin->GetSchema()->TrySetDefaultValue(*GraphPin, *LexToString(InValue));
+		}
+	}
+}
+
+TSharedRef<SWidget>	SDataflowPin::MakeIntWidget(const UEdGraphPin& Pin)
+{
+	return SNew(SBox)
+		.MinDesiredWidth(35.0f)
+		.MaxDesiredWidth(100.0f)
+		[
+			SNew(SNumericEntryBox<int32>)
+				.IsEnabled(this, &SGraphPin::GetDefaultValueIsEditable)
+				.Visibility(this, &SGraphPin::GetDefaultValueVisibility)
+				.EditableTextBoxStyle(FAppStyle::Get(), "Graph.EditableTextBox")
+				.BorderForegroundColor(FSlateColor::UseForeground())
+				.Value(this, &SDataflowPin::GetIntValue)
+				.OnValueCommitted(this, &SDataflowPin::SetIntValue)
+		];
+}
+
+//////////////////////////////////////////////////////////////////////////////////////////
+
+TOptional<float> SDataflowPin::GetFloatValue() const
+{
+	float Value = 0;
+	if (UEdGraphPin* GraphPin = GetPinObj())
+	{
+		LexFromString(Value, *GraphPin->GetDefaultAsString());
+	}
+	return Value;
+}
+
+void SDataflowPin::SetFloatValue(float InValue, ETextCommit::Type CommitType)
+{
+	if (UEdGraphPin* GraphPin = GetPinObj())
+	{
+		float CurrentValue = 0.0f;
+		LexFromString(CurrentValue, *GraphPin->GetDefaultAsString());
+		if (CurrentValue != InValue)
+		{
+			const FScopedTransaction Transaction(LOCTEXT("ChangeFloatPinValue", "Change Float Pin Value"));
+			GraphPin->Modify();
+			GraphPin->GetSchema()->TrySetDefaultValue(*GraphPin, *LexToString(InValue));
+		}
+	}
+}
+
+TSharedRef<SWidget>	SDataflowPin::MakeFloatWidget(const UEdGraphPin& Pin)
+{
+	return SNew(SBox)
+		.MinDesiredWidth(35.0f)
+		.MaxDesiredWidth(100.0f)
+		[
+			SNew(SNumericEntryBox<float>)
+				.IsEnabled(this, &SGraphPin::GetDefaultValueIsEditable)
+				.Visibility(this, &SGraphPin::GetDefaultValueVisibility)
+				.EditableTextBoxStyle(FAppStyle::Get(), "Graph.EditableTextBox")
+				.BorderForegroundColor(FSlateColor::UseForeground())
+				.Value(this, &SDataflowPin::GetFloatValue)
+				.OnValueCommitted(this, &SDataflowPin::SetFloatValue)
+		];
+}
+
+//////////////////////////////////////////////////////////////////////////////////////////
+
+FText SDataflowPin::GetStringValue() const
+{
+	if (UEdGraphPin* GraphPin = GetPinObj())
+	{
+		return FText::FromString(GraphPin->GetDefaultAsString());
+	}
+	return {};
+}
+
+void SDataflowPin::SetStringValue(const FText& NewTypeInValue, ETextCommit::Type /*CommitInfo*/)
+{
+	if (UEdGraphPin* GraphPin = GetPinObj())
+	{
+		if (!GraphPin->GetDefaultAsString().Equals(NewTypeInValue.ToString()))
+		{
+			const FScopedTransaction Transaction(LOCTEXT("ChangeStringPinValue", "Change String Pin Value"));
+			GraphPin->Modify();
+			GraphPin->GetSchema()->TrySetDefaultValue(*GraphPinObj, NewTypeInValue.ToString());
+		}
+	}
+}
+
+TSharedRef<SWidget>	SDataflowPin::MakeStringWidget(const UEdGraphPin& Pin)
+{
+	return SNew(SBox)
+		.MinDesiredWidth(100.0f)
+		.MaxDesiredHeight(200.0f)
+		[
+			SNew(SMultiLineEditableTextBox)
+				.IsEnabled(this, &SGraphPin::GetDefaultValueIsEditable)
+				.Visibility(this, &SGraphPin::GetDefaultValueVisibility)
+				.Style(FAppStyle::Get(), "Graph.EditableTextBox")
+				.Text(this, &SDataflowPin::GetStringValue)
+				.SelectAllTextWhenFocused(true)
+				.OnTextCommitted(this, &SDataflowPin::SetStringValue)
+				.ForegroundColor(FSlateColor::UseForeground())
+				.WrapTextAt(400.0f)
+				.ModiferKeyForNewLine(EModifierKey::Shift)
+		];
+}
+
+
+#undef LOCTEXT_NAMESPACE
diff --git a/Engine/Plugins/Experimental/Dataflow/Source/DataflowEditor/Public/Dataflow/DataflowSPin.h b/Engine/Plugins/Experimental/Dataflow/Source/DataflowEditor/Public/Dataflow/DataflowSPin.h
--- a/Engine/Plugins/Experimental/Dataflow/Source/DataflowEditor/Public/Dataflow/DataflowSPin.h
+++ b/Engine/Plugins/Experimental/Dataflow/Source/DataflowEditor/Public/Dataflow/DataflowSPin.h
@@ -4,6 +4,9 @@
 
 #include "SGraphPin.h"
 
+class UEdGraphPin;
+struct FDataflowNode;
+
 /**
 * class used to for all dataflow node pins 
 */
@@ -23,8 +26,28 @@ public:
 	void Construct(const FArguments& InArgs, UEdGraphPin* InPin);
 
 	virtual FSlateColor GetPinColor() const override;
+	virtual TSharedRef<SWidget> GetDefaultValueWidget() override;
 
 private:
+	TSharedPtr<FDataflowNode> GetDataflowNode() const;
+
+	ECheckBoxState GetBoolValue() const;
+	void SetBoolValue(ECheckBoxState InIsChecked);
+	TSharedRef<SWidget> MakeBoolWidget(const UEdGraphPin& Pin);
+
+	TSharedRef<SWidget>	MakeIntWidget(const UEdGraphPin& Pin);
+	TOptional<int32> GetIntValue() const;
+	void SetIntValue(int32 InValue, ETextCommit::Type CommitType);
+
+	TSharedRef<SWidget>	MakeFloatWidget(const UEdGraphPin& Pin);
+	TOptional<float> GetFloatValue() const;
+	void SetFloatValue(float InValue, ETextCommit::Type CommitType);
+
+	TSharedRef<SWidget>	MakeStringWidget(const UEdGraphPin& Pin);
+	FText GetStringValue() const;
+	void SetStringValue(const FText& NewTypeInValue, ETextCommit::Type /*CommitInfo*/);
+
+
 	TAttribute<bool> bIsPinInvalid;
 	TAttribute<FLinearColor> PinColorOverride;
 	TAttribute<bool> bIsPinColorOverriden;
diff --git a/Engine/Source/Runtime/Experimental/Dataflow/Engine/Private/Dataflow/DataflowEdNode.cpp b/Engine/Source/Runtime/Experimental/Dataflow/Engine/Private/Dataflow/DataflowEdNode.cpp
--- a/Engine/Source/Runtime/Experimental/Dataflow/Engine/Private/Dataflow/DataflowEdNode.cpp
+++ b/Engine/Source/Runtime/Experimental/Dataflow/Engine/Private/Dataflow/DataflowEdNode.cpp
@@ -134,6 +134,18 @@ TSharedPtr<const FDataflowNode> UDataflowEdNode::GetDataflowNode() const
 	return TSharedPtr<FDataflowNode>(nullptr);
 }
 
+#if WITH_EDITOR && !(UE_BUILD_SHIPPING || UE_BUILD_TEST)
+UEdGraphPin* UDataflowEdNode::CreateEdPin(const UE::Dataflow::FPin& Pin)
+{
+	UEdGraphPin* const EdPin = CreatePin(UE::Dataflow::Private::DataflowDirectionToEdPinDirection(Pin.Direction), Pin.Type, Pin.Name);
+	EdPin->bHidden = Pin.bHidden;
+	if (FDataflowArrayTypePolicy::SupportsTypeStatic(Pin.Type))
+	{
+		EdPin->PinType.ContainerType = EPinContainerType::Array;
+	}
+	return EdPin;
+}
+#endif
 
 void UDataflowEdNode::AllocateDefaultPins()
 {
@@ -149,12 +161,13 @@ void UDataflowEdNode::AllocateDefaultPins()
 			{
 				for (const UE::Dataflow::FPin& Pin : DataflowNode->GetPins())
 				{
-					UEdGraphPin* const EdPin = CreatePin(UE::Dataflow::Private::DataflowDirectionToEdPinDirection(Pin.Direction), Pin.Type, Pin.Name);
-					EdPin->bHidden = Pin.bHidden;
+					CreateEdPin(Pin);
 				}
 			}
 		}
 	}
+
+	UpdatePinsDefaultValuesFromNode();
 #endif // WITH_EDITOR && !UE_BUILD_SHIPPING
 }
 
@@ -221,12 +234,10 @@ void UDataflowEdNode::UpdatePinsFromDataflowNode()
 				for (const UE::Dataflow::FPin& Pin : DataflowNode->GetPins())
 				{
 					const EEdGraphPinDirection EdDirection = UE::Dataflow::Private::DataflowDirectionToEdPinDirection(Pin.Direction);
-					UEdGraphPin* EdPin = FindPin(Pin.Name, EdDirection);
-					if (!EdPin)
+					if (nullptr == FindPin(Pin.Name, EdDirection))
 					{
-						EdPin = CreatePin(EdDirection, Pin.Type, Pin.Name);
+						CreateEdPin(Pin);
 					}
-					EdPin->bHidden = Pin.bHidden;
 				}
 
 				// reorder the pins 
@@ -256,6 +267,9 @@ void UDataflowEdNode::UpdatePinsFromDataflowNode()
 			}
 		}
 	}
+
+	UpdatePinsDefaultValuesFromNode();
+
 #endif // WITH_EDITOR && !UE_BUILD_SHIPPING
 }
 
@@ -354,6 +368,74 @@ void UDataflowEdNode::UpdatePinsConnectionsFromDataflowNode()
 #endif // WITH_EDITOR && !UE_BUILD_SHIPPING
 }
 
+#if WITH_EDITOR && !(UE_BUILD_SHIPPING || UE_BUILD_TEST)
+bool UDataflowEdNode::SupportsEditablePinType(const UEdGraphPin& Pin)
+{
+	if (const FDataflowConnection* Connection = GetConnectionFromPin(&Pin))
+	{
+		// todo(dataflow): Add support for anytype in the future 
+		if (!Connection->IsAnyType())
+		{
+			const FName Pintype = Pin.PinType.PinCategory;
+			if (Pin.PinType.IsArray())
+			{
+				return false;
+			}
+			return (Pintype == TDataflowPolicyTypeName<bool>::GetName())
+				|| (Pintype == TDataflowPolicyTypeName<int>::GetName())
+				|| (Pintype == TDataflowPolicyTypeName<float>::GetName())
+				|| (Pintype == TDataflowPolicyTypeName<FString>::GetName())
+				;
+		}
+	}
+	return false;
+}
+#endif // WITH_EDITOR && !UE_BUILD_SHIPPING
+
+void UDataflowEdNode::UpdatePinsDefaultValuesFromNode()
+{
+#if WITH_EDITOR && !(UE_BUILD_SHIPPING || UE_BUILD_TEST)
+	if (!DataflowGraph || !DataflowNodeGuid.IsValid())
+	{
+		return;
+	}
+	TSharedPtr<FDataflowNode> DataflowNode = DataflowGraph->FindBaseNode(DataflowNodeGuid);
+	if (!DataflowNode)
+	{
+		return;
+	}
+
+	if (bEditablePinReentranceGuard)
+	{
+		return; // 
+	}
+
+	bEditablePinReentranceGuard = true;
+
+	for (UEdGraphPin* Pin : GetAllPins())
+	{
+		if (ensure(Pin) && Pin->Direction == EEdGraphPinDirection::EGPD_Input)
+		{
+			if (SupportsEditablePinType(*Pin))
+			{
+				if (const FDataflowInput* DataflowInput = DataflowNode->FindInput(Pin->GetFName()))
+				{
+					if (const FProperty* Property = DataflowInput->GetProperty())
+					{
+						Pin->DefaultValue.Empty();
+						const void* SourceValue = DataflowInput->RealAddress();
+						//const void* SourceValue = Property->ContainerPtrToValuePtr<void>(DataflowNode.Get());
+						Property->ExportText_Direct(Pin->DefaultValue, SourceValue, SourceValue, nullptr, PPF_None);
+					}
+				}
+			}
+		}
+	}
+
+	bEditablePinReentranceGuard = false;
+#endif
+}
+
 void UDataflowEdNode::AddOptionPin()
 {
 #if WITH_EDITOR && !(UE_BUILD_SHIPPING || UE_BUILD_TEST)
@@ -370,7 +452,7 @@ void UDataflowEdNode::AddOptionPin()
 				{
 				case UE::Dataflow::FPin::EDirection::INPUT:
 				case UE::Dataflow::FPin::EDirection::OUTPUT:
-					CreatePin(UE::Dataflow::Private::DataflowDirectionToEdPinDirection(Pin.Direction), Pin.Type, Pin.Name);
+					CreateEdPin(Pin);
 					ReconstructNode();
 					break;
 				default:
@@ -897,6 +979,76 @@ void UDataflowEdNode::PostEditUndo()
 	}
 }
 
+void UDataflowEdNode::PinDefaultValueChanged(UEdGraphPin* Pin)
+{
+	if (bEditablePinReentranceGuard)
+	{
+		return; // we are setting the values from UpdatePinsDefaultValuesFromNode
+	}
+
+	if (Pin && SupportsEditablePinType(*Pin))
+	{
+		if (TSharedPtr<FDataflowNode> DataflowNode = GetDataflowNode())
+		{
+			if (const FDataflowInput* Input = DataflowNode->FindInput(Pin->PinName))
+			{
+				if (FProperty* Property = const_cast<FProperty*>(Input->GetProperty()))
+				{
+					bEditablePinReentranceGuard = true;
+
+					const FName PropertyType = FDataflowConnection::GetTypeNameFromProperty(Property);
+					const FString PinStringValue = Pin->GetDefaultAsString();
+					if (PropertyType == Pin->PinType.PinCategory)
+					{
+						//void* TargetAddress = Property->ContainerPtrToValuePtr<void>(DataflowNode.Get());
+						void* TargetAddress = Input->RealAddress();
+						if (PropertyType == TDataflowPolicyTypeName<bool>::GetName())
+						{
+							const bool bValue = PinStringValue.ToBool();
+							Property->CopyCompleteValue(TargetAddress, &bValue);
+						}
+						else if (PropertyType == TDataflowPolicyTypeName<int>::GetName())
+						{
+							const int32 IntValue = (PinStringValue.IsNumeric()) ? FCString::Atoi(*PinStringValue) : 0;
+							Property->CopyCompleteValue(TargetAddress, &IntValue);
+						}
+						else if (PropertyType == TDataflowPolicyTypeName<float>::GetName())
+						{
+							const float FloatValue = (PinStringValue.IsNumeric()) ? FCString::Atof(*PinStringValue) : 0;
+							Property->CopyCompleteValue(TargetAddress, &FloatValue);
+						}
+						else if (PropertyType == TDataflowPolicyTypeName<FString>::GetName())
+						{
+							Property->CopyCompleteValue(TargetAddress, &PinStringValue);
+						}
+
+						bEditablePinReentranceGuard = false;
+
+						// we cannot invalidate the node right away becasue invalidation cause delegate to invalidate the pin 
+						// we are already receiving a notification from causing ensure in the EdGraphPin code 
+						// so we schedule an invalidation for later 
+						ScheduleNodeInvalidation();
+					}
+				}
+			}
+		}
+	}
+}
+
+void UDataflowEdNode::ScheduleNodeInvalidation()
+{
+	if (TSharedPtr<FDataflowNode> DataflowNode = GetDataflowNode())
+	{
+		AsyncTask(ENamedThreads::GameThread, [WeakNode = DataflowNode.ToWeakPtr()]()
+			{
+				if (TSharedPtr<FDataflowNode> DataflowNode = WeakNode.Pin())
+				{
+					DataflowNode->Invalidate();
+				}
+			});
+	}
+}
+
 void UDataflowEdNode::HideAllInputPins()
 {
 	bool bAnyHidden = false;
diff --git a/Engine/Source/Runtime/Experimental/Dataflow/Engine/Public/Dataflow/DataflowEdNode.h b/Engine/Source/Runtime/Experimental/Dataflow/Engine/Public/Dataflow/DataflowEdNode.h
--- a/Engine/Source/Runtime/Experimental/Dataflow/Engine/Public/Dataflow/DataflowEdNode.h
+++ b/Engine/Source/Runtime/Experimental/Dataflow/Engine/Public/Dataflow/DataflowEdNode.h
@@ -51,6 +51,7 @@ public:
 	DATAFLOWENGINE_API virtual void OnPinRemoved(UEdGraphPin* InRemovedPin) override;
 	DATAFLOWENGINE_API virtual bool ShouldDrawNodeAsControlPointOnly(int32& OutInputPinIndex, int32& OutOutputPinIndex) const override;
 	DATAFLOWENGINE_API virtual void PostEditUndo() override;
+	DATAFLOWENGINE_API virtual void PinDefaultValueChanged(UEdGraphPin* Pin) override;
 #endif // WITH_EDITOR
 	// End of UEdGraphNode interface
 
@@ -66,6 +67,7 @@ public:
 
 	DATAFLOWENGINE_API void UpdatePinsFromDataflowNode();
 	DATAFLOWENGINE_API void UpdatePinsConnectionsFromDataflowNode();
+	DATAFLOWENGINE_API void UpdatePinsDefaultValuesFromNode();
 
 	FGuid GetDataflowNodeGuid() const { return DataflowNodeGuid; }
 	DATAFLOWENGINE_API void SetDataflowNodeGuid(FGuid InGuid);
@@ -117,8 +119,21 @@ public:
 	DATAFLOWENGINE_API static TSharedPtr<FDataflowNode> GetDataflowNodeFromEdNode(UEdGraphNode* EdNode);
 	DATAFLOWENGINE_API static TSharedPtr<const FDataflowNode> GetDataflowNodeFromEdNode(const UEdGraphNode* EdNode);
 
+#if WITH_EDITOR && !(UE_BUILD_SHIPPING || UE_BUILD_TEST)
+	DATAFLOWENGINE_API static bool SupportsEditablePinType(const UEdGraphPin& Pin);
+#endif
+
 private:
 
+#if WITH_EDITOR && !(UE_BUILD_SHIPPING || UE_BUILD_TEST)
+	UEdGraphPin* CreateEdPin(const UE::Dataflow::FPin& Pin);
+#endif
+
+	void ScheduleNodeInvalidation();
+
+	/** used to avoid update loops setting and editing the editable pins */
+	bool bEditablePinReentranceGuard = false;
+
 	UPROPERTY(Transient)
 	bool bRenderInAssetEditor = false;
 
