Dataflow : fix potential crash in DataflowTooltarget code
- fix the check for array size before accessing first element

diff --git a/Engine/Plugins/Experimental/Dataflow/Source/DataflowEditor/Private/Dataflow/DataflowToolTarget.cpp b/Engine/Plugins/Experimental/Dataflow/Source/DataflowEditor/Private/Dataflow/DataflowToolTarget.cpp
--- a/Engine/Plugins/Experimental/Dataflow/Source/DataflowEditor/Private/Dataflow/DataflowToolTarget.cpp
+++ b/Engine/Plugins/Experimental/Dataflow/Source/DataflowEditor/Private/Dataflow/DataflowToolTarget.cpp
@@ -66,10 +66,12 @@ bool UDataflowReadOnlyToolTarget::CommitMaterialSetUpdate(const FComponentMateri
 		}
 	}
 
-	if(IsValid() && FilteredMaterials.Num() == 0)
+	if(IsValid() && FilteredMaterials.Num() > 0)
 	{
-		if(UMaterial* Material = Cast<UMaterial>(FilteredMaterials[0]))
-		Dataflow->Material = Material;
+		if (UMaterial* Material = Cast<UMaterial>(FilteredMaterials[0]))
+		{
+			Dataflow->Material = Material;
+		}
 		return true;
 	}
 	return false;
