Dataflow weight map paint tool : hide "Update Node" UI section as it is redundant with the node details
- Cannot remove because the logic relies on it to exists for verifications

diff --git a/Engine/Plugins/Experimental/Dataflow/Source/DataflowEditor/Public/DataflowEditorTools/DataflowEditorWeightMapPaintTool.h b/Engine/Plugins/Experimental/Dataflow/Source/DataflowEditor/Public/DataflowEditorTools/DataflowEditorWeightMapPaintTool.h
--- a/Engine/Plugins/Experimental/Dataflow/Source/DataflowEditor/Public/DataflowEditorTools/DataflowEditorWeightMapPaintTool.h
+++ b/Engine/Plugins/Experimental/Dataflow/Source/DataflowEditor/Public/DataflowEditorTools/DataflowEditorWeightMapPaintTool.h
@@ -255,7 +255,7 @@ class UDataflowEditorUpdateWeightMapProperties : public UInteractiveToolProperty
 
 public:
 
-	UPROPERTY(EditAnywhere, Category = UpdateNode, meta = (DisplayName = "Name"))
+	UPROPERTY()
 	FString Name;
 
 private:
