Fix CIS (name not used in dataflow construction scene constructor)

diff --git a/Engine/Plugins/Experimental/Dataflow/Source/DataflowEditor/Private/Dataflow/DataflowConstructionScene.cpp b/Engine/Plugins/Experimental/Dataflow/Source/DataflowEditor/Private/Dataflow/DataflowConstructionScene.cpp
--- a/Engine/Plugins/Experimental/Dataflow/Source/DataflowEditor/Private/Dataflow/DataflowConstructionScene.cpp
+++ b/Engine/Plugins/Experimental/Dataflow/Source/DataflowEditor/Private/Dataflow/DataflowConstructionScene.cpp
@@ -27,9 +27,7 @@
 
 FDataflowConstructionScene::FDataflowConstructionScene(FPreviewScene::ConstructionValues ConstructionValues, UDataflowEditor* InEditor)
 	: FDataflowPreviewSceneBase(ConstructionValues, InEditor, FName("Construction Components"))
-{
-	RootSceneActor->GetName();
-}
+{}
 
 FDataflowConstructionScene::~FDataflowConstructionScene()
 {
