Dataflow weight map painting tool : display value at brush on the brush

diff --git a/Engine/Plugins/Experimental/Dataflow/Source/DataflowEditor/Private/DataflowEditorTools/DataflowEditorWeightMapPaintTool.cpp b/Engine/Plugins/Experimental/Dataflow/Source/DataflowEditor/Private/DataflowEditorTools/DataflowEditorWeightMapPaintTool.cpp
--- a/Engine/Plugins/Experimental/Dataflow/Source/DataflowEditor/Private/DataflowEditorTools/DataflowEditorWeightMapPaintTool.cpp
+++ b/Engine/Plugins/Experimental/Dataflow/Source/DataflowEditor/Private/DataflowEditorTools/DataflowEditorWeightMapPaintTool.cpp
@@ -1463,6 +1463,30 @@ void UDataflowEditorWeightMapPaintTool::DrawHUD(FCanvas* Canvas, IToolsContextRe
 	{
 		PolygonSelectionMechanic->DrawHUD(Canvas, RenderAPI);
 	}
+
+	if (Canvas && RenderAPI && BrushIndicator)
+	{
+		// display the value under the brush
+		const float BrushRadius = BrushIndicator->BrushRadius;
+		const FVector Offset = (-RenderAPI->GetSceneView()->GetViewUp() * BrushRadius * 1.1f);
+		const FVector BrushPos = BrushIndicator->BrushPosition + Offset;
+
+		FVector4 ScreenBrushPos = RenderAPI->GetSceneView()->WorldToScreen(BrushPos);
+		FVector2D ScreenBrushPos2D = FVector2D(ScreenBrushPos.X, ScreenBrushPos.Y);
+
+		RenderAPI->GetSceneView()->ScreenToPixel(ScreenBrushPos, ScreenBrushPos2D);
+
+		const bool bShouldShowValueAtBrush = IsInBrushSubMode();
+		if (bShouldShowValueAtBrush)
+		{
+			const FText BaseText = LOCTEXT("DataflowWeightMapPaintTool_ValueAtBrush_TextFormat", "Value: {0}");
+			const FText ValueAtBrushText = FText::Format(BaseText, FText::AsNumber(FilterProperties->ValueAtBrush));
+			FCanvasTextItem TextItem(ScreenBrushPos2D, ValueAtBrushText, GEngine->GetLargeFont(), FLinearColor::White);
+			TextItem.EnableShadow(FLinearColor::Black);
+			TextItem.bCentreX = true;
+			Canvas->DrawItem(TextItem);
+		}
+	}
 }
 
 void UDataflowEditorWeightMapPaintTool::Render(IToolsContextRenderAPI* RenderAPI)
