Dataflow : new rendering system - infrastructure and initial implementation
- not connected to the rest of the system yet ( will come in a later change )
- renderable types , registry and instance classes
- implementation for surface and Uvs for static mesh , skeletal mesh , dynamic mesh and geometry collection types

diff --git a/Engine/Plugins/Experimental/Dataflow/Source/DataflowEditor/DataflowEditor.Build.cs b/Engine/Plugins/Experimental/Dataflow/Source/DataflowEditor/DataflowEditor.Build.cs
--- a/Engine/Plugins/Experimental/Dataflow/Source/DataflowEditor/DataflowEditor.Build.cs
+++ b/Engine/Plugins/Experimental/Dataflow/Source/DataflowEditor/DataflowEditor.Build.cs
@@ -63,6 +63,8 @@ namespace UnrealBuildTool.Rules
 					"MeshModelingToolsExp",
 					"ModelingComponentsEditorOnly",
 					"ModelingComponents",
+//					"ModelingToolsEditorMode",
+					"PlanarCut",
 					"Projects",
 					"PropertyEditor",
 					"RenderCore",
diff --git a/Engine/Plugins/Experimental/Dataflow/Source/DataflowEditor/Private/DataflowRendering/DataflowDynamicMeshRenderableType.cpp b/Engine/Plugins/Experimental/Dataflow/Source/DataflowEditor/Private/DataflowRendering/DataflowDynamicMeshRenderableType.cpp
new file mode 100644
--- /dev/null
+++ b/Engine/Plugins/Experimental/Dataflow/Source/DataflowEditor/Private/DataflowRendering/DataflowDynamicMeshRenderableType.cpp
@@ -0,0 +1,63 @@
+// Copyright Epic Games, Inc. All Rights Reserved.
+
+#include "DataflowRendering/DataflowDynamicMeshRenderableType.h"
+
+#include "Components/DynamicMeshComponent.h"
+
+#include "Dataflow/DataflowRenderingViewMode.h"
+
+#include "DataflowRendering/DataflowRenderableComponents.h"
+#include "DataflowRendering/DataflowRenderableType.h"
+#include "DataflowRendering/DataflowRenderableTypeInstance.h"
+#include "DataflowRendering/DataflowRenderableTypeRegistry.h"
+#include "DataflowRendering/DataflowRenderableTypeUtils.h"
+
+#include "UObject/ObjectPtr.h"
+
+namespace UE::Dataflow::Private
+{
+	///////////////////////////////////////////////////////////////////////////////////////////////////////////////
+
+	class FDynamicMeshSurfaceRenderableType : public IRenderableType
+	{
+		UE_DATAFLOW_IRENDERABLE_PRIMARY_TYPE(TObjectPtr<UDynamicMesh>, DynamicMesh);
+		UE_DATAFLOW_IRENDERABLE_RENDER_GROUP(Surface);
+		UE_DATAFLOW_IRENDERABLE_VIEW_MODE(FDataflowConstruction3DViewMode);
+
+		virtual void GetPrimitiveComponents(const FRenderableTypeInstance& Instance, FRenderableComponents& OutComponents) const override
+		{
+			if (const TObjectPtr<UDynamicMesh> DynamicMesh = GetDynamicMesh(Instance, nullptr))
+			{
+				if (UDynamicMeshComponent* Component = OutComponents.AddNewComponent<UDynamicMeshComponent>())
+				{
+					Component->SetDynamicMesh(DynamicMesh);
+				}
+			}
+		}
+	};
+
+	////////////////////////////////////////////////////////////////////////////////////////////////////////////////////////////
+
+	class FDynamicMeshUvsRenderableType : public IRenderableType
+	{
+		UE_DATAFLOW_IRENDERABLE_PRIMARY_TYPE(TObjectPtr<UDynamicMesh>, DynamicMesh);
+		UE_DATAFLOW_IRENDERABLE_RENDER_GROUP(Uvs);
+		UE_DATAFLOW_IRENDERABLE_VIEW_MODE(FDataflowConstructionUVViewMode);
+
+		virtual void GetPrimitiveComponents(const FRenderableTypeInstance& Instance, FRenderableComponents& OutComponents) const override
+		{
+			if (const TObjectPtr<UDynamicMesh> DynamicMesh = GetDynamicMesh(Instance, nullptr))
+			{
+				Rendering::AddUvDynamicMeshComponent(DynamicMesh->GetMeshRef(), 0, OutComponents);
+			}
+		}
+	};
+
+	///////////////////////////////////////////////////////////////////////////////////////////////////////////////
+
+	void RegisterDynamicMeshRenderableTypes()
+	{
+		UE_DATAFLOW_REGISTER_RENDERABLE_TYPE(FDynamicMeshSurfaceRenderableType);
+		UE_DATAFLOW_REGISTER_RENDERABLE_TYPE(FDynamicMeshUvsRenderableType);
+	}
+}
\ No newline at end of file
diff --git a/Engine/Plugins/Experimental/Dataflow/Source/DataflowEditor/Private/DataflowRendering/DataflowDynamicMeshRenderableType.h b/Engine/Plugins/Experimental/Dataflow/Source/DataflowEditor/Private/DataflowRendering/DataflowDynamicMeshRenderableType.h
new file mode 100644
--- /dev/null
+++ b/Engine/Plugins/Experimental/Dataflow/Source/DataflowEditor/Private/DataflowRendering/DataflowDynamicMeshRenderableType.h
@@ -0,0 +1,8 @@
+// Copyright Epic Games, Inc. All Rights Reserved.
+
+#pragma once
+
+namespace UE::Dataflow::Private
+{
+	void RegisterDynamicMeshRenderableTypes();
+}
diff --git a/Engine/Plugins/Experimental/Dataflow/Source/DataflowEditor/Private/DataflowRendering/DataflowGeometryCollectionProximityRenderableType.cpp b/Engine/Plugins/Experimental/Dataflow/Source/DataflowEditor/Private/DataflowRendering/DataflowGeometryCollectionProximityRenderableType.cpp
new file mode 100644
--- /dev/null
+++ b/Engine/Plugins/Experimental/Dataflow/Source/DataflowEditor/Private/DataflowRendering/DataflowGeometryCollectionProximityRenderableType.cpp
@@ -0,0 +1,134 @@
+// Copyright Epic Games, Inc. All Rights Reserved.
+
+#include "DataflowRendering/DataflowGeometryCollectionProximityRenderableType.h"
+
+#include "Drawing/LineSetComponent.h"
+#include "Drawing/PointSetComponent.h"
+
+#include "Dataflow/DataflowRenderingViewMode.h"
+
+#include "DataflowRendering/DataflowRenderableComponents.h"
+#include "DataflowRendering/DataflowRenderableType.h"
+#include "DataflowRendering/DataflowRenderableTypeInstance.h"
+#include "DataflowRendering/DataflowRenderableTypeRegistry.h"
+#include "DataflowRendering/DataflowRenderableTypeUtils.h"
+
+#include "GeometryCollection/GeometryCollection.h"
+#include "GeometryCollection/GeometryCollectionAlgo.h"
+#include "GeometryCollection/ManagedArrayCollection.h"
+
+#include "Materials/MaterialInterface.h"
+#include "Materials/Material.h"
+
+#include "UObject/ObjectPtr.h"
+
+namespace UE::Dataflow::Private
+{
+	////////////////////////////////////////////////////////////////////////////////////////////////////////////////////////////
+
+	class FGeometryCollectionProximityRenderableType : public IRenderableType
+	{
+		UE_DATAFLOW_IRENDERABLE_PRIMARY_TYPE(FManagedArrayCollection, Collection);
+		UE_DATAFLOW_IRENDERABLE_RENDER_GROUP(Proximity);
+		UE_DATAFLOW_IRENDERABLE_VIEW_MODE(FDataflowConstruction3DViewMode);
+
+	public:
+		FGeometryCollectionProximityRenderableType()
+		{
+			constexpr bool bDepthTested = false;
+			Material = (bDepthTested) 
+				? LoadObject<UMaterial>(nullptr, TEXT("/MeshModelingToolsetExp/Materials/PointSetComponentMaterial"))
+				: LoadObject<UMaterial>(nullptr, TEXT("/MeshModelingToolsetExp/Materials/PointSetOverlaidComponentMaterial"));
+		}
+
+	private:
+		virtual bool CanRender(const FRenderableTypeInstance& Instance) const
+		{
+			const FManagedArrayCollection& Collection = GetCollection(Instance);
+			return Collection.HasAttribute(ProximityAttributeName, FGeometryCollection::GeometryGroup)
+				&& Collection.HasAttribute(BoundingBoxAttributeName, FGeometryCollection::GeometryGroup);
+		}
+
+		virtual void GetPrimitiveComponents(const FRenderableTypeInstance& Instance, FRenderableComponents& OutComponents) const override
+		{
+			const FManagedArrayCollection& Collection = GetCollection(Instance);
+
+			if (Collection.HasAttribute(ProximityAttributeName, FGeometryCollection::GeometryGroup) &&
+				Collection.HasAttribute(BoundingBoxAttributeName, FGeometryCollection::GeometryGroup))
+			{
+				UPointSetComponent* PointComponent = OutComponents.AddNewComponent<UPointSetComponent>();
+				ULineSetComponent* LineComponent = OutComponents.AddNewComponent<ULineSetComponent>();
+
+				if (PointComponent && LineComponent)
+				{
+					const TManagedArray<TSet<int32>>& Proximity = Collection.GetAttribute<TSet<int32>>(ProximityAttributeName, FGeometryCollection::GeometryGroup);
+					const TManagedArray<FBox>& BoundingBox = Collection.GetAttribute<FBox>(BoundingBoxAttributeName, FGeometryCollection::GeometryGroup);
+					const TManagedArray<FTransform3f>* Transform = Collection.FindAttributeTyped<FTransform3f>(FGeometryCollection::TransformAttribute, FGeometryCollection::TransformGroup);
+					const TManagedArray<int32>* Parent = Collection.FindAttributeTyped<int32>(FGeometryCollection::ParentAttribute, FGeometryCollection::TransformGroup);
+					const TManagedArray<int32>* GeometryToTransformIndex = Collection.FindAttributeTyped<int32>(FGeometryCollection::TransformIndexAttribute, FGeometryCollection::GeometryGroup);
+
+					// compute transforms in collection space 
+					TArray<FTransform> GlobalTransformArray;
+					if (Transform && Parent)
+					{
+						GeometryCollectionAlgo::GlobalMatrices(*Transform, *Parent, GlobalTransformArray);
+					}
+
+					auto GetGeometryTransform =
+						[&GeometryToTransformIndex, &GlobalTransformArray](int32 GeoIndex) -> FTransform
+						{
+							if (GeometryToTransformIndex && GeometryToTransformIndex->IsValidIndex(GeoIndex))
+							{
+								const int32 TransformIndex = (*GeometryToTransformIndex)[GeoIndex];
+								if (GlobalTransformArray.IsValidIndex(TransformIndex))
+								{
+									return GlobalTransformArray[TransformIndex];
+								}
+							}
+							return FTransform::Identity;
+						};
+
+					const int32 NumGeometry = Collection.NumElements(FGeometryCollection::GeometryGroup);
+					TArray<FVector> Centers;
+					for (int32 GeoIndex = 0; GeoIndex < NumGeometry; ++GeoIndex)
+					{
+						if (Proximity[GeoIndex].Num() > 0)
+						{
+							const FTransform GeometryTransform = GetGeometryTransform(GeoIndex);
+							const FVector Center = GeometryTransform.TransformPosition(BoundingBox[GeoIndex].GetCenter());
+
+							PointComponent->AddPoint(FRenderablePoint(Center, FColor::Blue, 15.0f));
+
+							for (const int32 OtherGeoIndex : Proximity[GeoIndex])
+							{
+								if (OtherGeoIndex > GeoIndex)
+								{
+									const FTransform OtherGeometryTransform = GetGeometryTransform(OtherGeoIndex);
+									const FVector OtherCenter = OtherGeometryTransform.TransformPosition(BoundingBox[OtherGeoIndex].GetCenter());
+
+									LineComponent->AddLine(FRenderableLine(Center, OtherCenter, FColor::Yellow, 2.0f));
+								}
+							}
+						}
+					}
+
+					PointComponent->SetPointMaterial(Material);
+					LineComponent->SetLineMaterial(Material);
+				}
+			}
+		}
+
+		static inline const FName ProximityAttributeName = TEXT("Proximity");
+		static inline const FName BoundingBoxAttributeName = TEXT("BoundingBox");
+
+		UMaterialInterface* Material = nullptr;
+	};
+
+
+	///////////////////////////////////////////////////////////////////////////////////////////////////////////////
+
+	void RegisterGeometryCollectionProximityRenderableTypes()
+	{
+		UE_DATAFLOW_REGISTER_RENDERABLE_TYPE(FGeometryCollectionProximityRenderableType);
+	}
+}
\ No newline at end of file
diff --git a/Engine/Plugins/Experimental/Dataflow/Source/DataflowEditor/Private/DataflowRendering/DataflowGeometryCollectionProximityRenderableType.h b/Engine/Plugins/Experimental/Dataflow/Source/DataflowEditor/Private/DataflowRendering/DataflowGeometryCollectionProximityRenderableType.h
new file mode 100644
--- /dev/null
+++ b/Engine/Plugins/Experimental/Dataflow/Source/DataflowEditor/Private/DataflowRendering/DataflowGeometryCollectionProximityRenderableType.h
@@ -0,0 +1,8 @@
+// Copyright Epic Games, Inc. All Rights Reserved.
+
+#pragma once
+
+namespace UE::Dataflow::Private
+{
+	void RegisterGeometryCollectionProximityRenderableTypes();
+}
diff --git a/Engine/Plugins/Experimental/Dataflow/Source/DataflowEditor/Private/DataflowRendering/DataflowGeometryCollectionRenderableType.cpp b/Engine/Plugins/Experimental/Dataflow/Source/DataflowEditor/Private/DataflowRendering/DataflowGeometryCollectionRenderableType.cpp
new file mode 100644
--- /dev/null
+++ b/Engine/Plugins/Experimental/Dataflow/Source/DataflowEditor/Private/DataflowRendering/DataflowGeometryCollectionRenderableType.cpp
@@ -0,0 +1,115 @@
+// Copyright Epic Games, Inc. All Rights Reserved.
+
+#include "DataflowRendering/DataflowGeometryCollectionRenderableType.h"
+
+#include "Components/DynamicMeshComponent.h"
+
+#include "Dataflow/DataflowRenderingViewMode.h"
+
+#include "DataflowRendering/DataflowRenderableComponents.h"
+#include "DataflowRendering/DataflowRenderableType.h"
+#include "DataflowRendering/DataflowRenderableTypeInstance.h"
+#include "DataflowRendering/DataflowRenderableTypeRegistry.h"
+#include "DataflowRendering/DataflowRenderableTypeUtils.h"
+
+#include "GeometryCollection/Facades/CollectionMeshFacade.h"
+#include "GeometryCollection/GeometryCollectionUtility.h"
+#include "GeometryCollection/GeometryCollectionAlgo.h"
+#include "GeometryCollection/ManagedArrayCollection.h"
+
+#include "UObject/ObjectPtr.h"
+
+#include "PlanarCut.h" // for ConvertGeometryCollectionToDynamicMesh
+
+namespace UE::Dataflow::Private
+{
+	////////////////////////////////////////////////////////////////////////////////////////////////////////////////////////////
+
+	class FGeometryCollectionSurfaceRenderableType : public IRenderableType
+	{
+		UE_DATAFLOW_IRENDERABLE_PRIMARY_TYPE(FManagedArrayCollection, Collection);
+		UE_DATAFLOW_IRENDERABLE_RENDER_GROUP(Surface);
+		UE_DATAFLOW_IRENDERABLE_VIEW_MODE(FDataflowConstruction3DViewMode);
+
+		virtual bool CanRender(const FRenderableTypeInstance& Instance) const
+		{
+			const FManagedArrayCollection& Collection = GetCollection(Instance);
+			return Collection.HasGroup(FGeometryCollection::GeometryGroup);
+		}
+
+		virtual void GetPrimitiveComponents(const FRenderableTypeInstance& Instance, FRenderableComponents& OutComponents) const override
+		{
+			const FManagedArrayCollection& Collection = GetCollection(Instance);
+
+			// TODO(dataflow) : use a facade or a more generic way to convert a collection to a dynamic mesh
+			const TManagedArray<FTransform3f>* TransformAttribute = Collection.FindAttribute<FTransform3f>(FTransformCollection::TransformAttribute, FTransformCollection::TransformGroup);
+			const TManagedArray<int32>* TransformIndexAttribute = Collection.FindAttribute<int32>(FGeometryCollection::TransformIndexAttribute, FGeometryCollection::GeometryGroup);
+
+			if (TransformAttribute && TransformIndexAttribute)
+			{
+				// TODO(dataflow) : need a version of ConvertGeometryCollectionToDynamicMesh that only takes a managed array collection
+				if (TUniquePtr<const FGeometryCollection> GeomCollection = TUniquePtr<FGeometryCollection>(Collection.NewCopy<FGeometryCollection>()))
+				{
+					// prepare the dynamic meshes
+					TArray<FDynamicMesh3> DynamicMeshes;
+					DynamicMeshes.SetNum(TransformAttribute->Num());
+
+					TArrayView<const FTransform3f> BoneTransforms = MakeArrayView(TransformAttribute->GetConstArray());
+
+					// convert to dynamic meshes in parallel
+					ParallelFor(DynamicMeshes.Num(),
+						[&DynamicMeshes, &BoneTransforms, &GeomCollection](int32 TransformIndex)
+						{
+							FDynamicMesh3& DynamicMesh = DynamicMeshes[TransformIndex];
+							TArrayView<const int32> TransformIndicesToConvert(&TransformIndex, 1);
+
+							FTransform UnusedTransform;
+							::ConvertGeometryCollectionToDynamicMesh(DynamicMesh, UnusedTransform, false, *GeomCollection, true, BoneTransforms, true, TransformIndicesToConvert);
+						});
+
+					// only create components for the one that have vertices
+					for (int32 TransformIndex = 0; TransformIndex < TransformAttribute->Num(); ++TransformIndex)
+					{
+						if (DynamicMeshes[TransformIndex].VertexCount() > 0)
+						{
+							if (UDynamicMeshComponent* Component = OutComponents.AddNewComponent<UDynamicMeshComponent>())
+							{
+								Component->SetMesh(MoveTemp(DynamicMeshes[TransformIndex]));
+							}
+						}
+					}
+				}
+			}
+		}
+	};
+
+	////////////////////////////////////////////////////////////////////////////////////////////////////////////////////////////
+
+	class FGeometryCollectionUvsRenderableType : public IRenderableType
+	{
+		UE_DATAFLOW_IRENDERABLE_PRIMARY_TYPE(FManagedArrayCollection, Collection);
+		UE_DATAFLOW_IRENDERABLE_RENDER_GROUP(Uvs);
+		UE_DATAFLOW_IRENDERABLE_VIEW_MODE(FDataflowConstructionUVViewMode);
+
+		virtual bool CanRender(const FRenderableTypeInstance& Instance) const
+		{
+			const FManagedArrayCollection& Collection = GetCollection(Instance);
+			return Collection.HasGroup(FGeometryCollection::GeometryGroup);
+		}
+
+		virtual void GetPrimitiveComponents(const FRenderableTypeInstance& Instance, FRenderableComponents& OutComponents) const override
+		{
+			const FManagedArrayCollection& Collection = GetCollection(Instance);
+
+			Rendering::AddUvDynamicMeshComponent(Collection, 0, OutComponents);
+		}
+	};
+
+	///////////////////////////////////////////////////////////////////////////////////////////////////////////////
+
+	void RegisterGeometryCollectionRenderableTypes()
+	{
+		UE_DATAFLOW_REGISTER_RENDERABLE_TYPE(FGeometryCollectionSurfaceRenderableType);
+		UE_DATAFLOW_REGISTER_RENDERABLE_TYPE(FGeometryCollectionUvsRenderableType);
+	}
+}
\ No newline at end of file
diff --git a/Engine/Plugins/Experimental/Dataflow/Source/DataflowEditor/Private/DataflowRendering/DataflowGeometryCollectionRenderableType.h b/Engine/Plugins/Experimental/Dataflow/Source/DataflowEditor/Private/DataflowRendering/DataflowGeometryCollectionRenderableType.h
new file mode 100644
--- /dev/null
+++ b/Engine/Plugins/Experimental/Dataflow/Source/DataflowEditor/Private/DataflowRendering/DataflowGeometryCollectionRenderableType.h
@@ -0,0 +1,8 @@
+// Copyright Epic Games, Inc. All Rights Reserved.
+
+#pragma once
+
+namespace UE::Dataflow::Private
+{
+	void RegisterGeometryCollectionRenderableTypes();
+}
diff --git a/Engine/Plugins/Experimental/Dataflow/Source/DataflowEditor/Private/DataflowRendering/DataflowRenderableTypeInstance.cpp b/Engine/Plugins/Experimental/Dataflow/Source/DataflowEditor/Private/DataflowRendering/DataflowRenderableTypeInstance.cpp
new file mode 100644
--- /dev/null
+++ b/Engine/Plugins/Experimental/Dataflow/Source/DataflowEditor/Private/DataflowRendering/DataflowRenderableTypeInstance.cpp
@@ -0,0 +1,61 @@
+// Copyright Epic Games, Inc. All Rights Reserved.
+
+#include "DataflowRendering/DataflowRenderableTypeInstance.h"
+
+#include "DataflowRendering/DataflowRenderableComponents.h"
+#include "DataflowRendering/DataflowRenderableType.h"
+#include "DataflowRendering/DataflowRenderableTypeRegistry.h"
+
+namespace UE::Dataflow
+{
+	/* static */ void FRenderableTypeInstance::GetRenderableInstancesForNode(FContext& Context, const IDataflowConstructionViewMode& ViewMode, const FDataflowNode& Node, TArray<FRenderableTypeInstance>& OutRenderableTypeInstances)
+	{
+		OutRenderableTypeInstances.Reset();
+
+		TArray<FDataflowOutput*> Outputs = Node.GetOutputs();
+		for (const FDataflowOutput* Output : Outputs)
+		{
+			if (Output)
+			{
+				const FRenderableTypeRegistry::FRenderableTypes& RenderableTypes = FRenderableTypeRegistry::GetInstance().GetRenderableTypes(Output->GetType());
+				for (const IRenderableType* RenderableType : RenderableTypes)
+				{
+					if (RenderableType  && RenderableType->IsViewModeSupported(ViewMode))
+					{
+						FRenderableTypeInstance RenderableTypeInstance(Context, Node.AsShared(), Output->GetName(), RenderableType);
+						if (RenderableTypeInstance.CanRender())
+						{
+							OutRenderableTypeInstances.Emplace(RenderableTypeInstance);
+						}
+					}
+				}
+			}
+		}
+	}
+
+	bool FRenderableTypeInstance::CanRender() const
+	{
+		return RenderableType ? RenderableType->CanRender(*this) : false;
+	}
+
+	FString FRenderableTypeInstance::GetDisplayName() const
+	{
+		if (RenderableType)
+		{
+			return FString::Format(TEXT("{0}.{1}"), { RenderableType->GetOutputType().ToString(), RenderableType->GetRenderGroup().ToString() });
+		}
+		return { TEXT("-- Unknown --") };
+	}
+
+	void FRenderableTypeInstance::GetPrimitiveComponents(FRenderableComponents& OutComponents) const
+	{
+		if (RenderableType)
+		{
+			RenderableType->GetPrimitiveComponents(*this, OutComponents);
+			if (Node)
+			{
+				Node->OnRenderOutput(EvaluationContext, OutputName, RenderableType->GetRenderGroup(), OutComponents.GetComponents());
+			}
+		}
+	}
+}
diff --git a/Engine/Plugins/Experimental/Dataflow/Source/DataflowEditor/Private/DataflowRendering/DataflowRenderableTypeRegistry.cpp b/Engine/Plugins/Experimental/Dataflow/Source/DataflowEditor/Private/DataflowRendering/DataflowRenderableTypeRegistry.cpp
new file mode 100644
--- /dev/null
+++ b/Engine/Plugins/Experimental/Dataflow/Source/DataflowEditor/Private/DataflowRendering/DataflowRenderableTypeRegistry.cpp
@@ -0,0 +1,35 @@
+// Copyright Epic Games, Inc. All Rights Reserved.
+
+#include "DataflowRendering/DataflowRenderableTypeRegistry.h"
+
+#include "DataflowRendering/DataflowRenderableType.h"
+
+namespace UE::Dataflow
+{
+	FRenderableTypeRegistry& FRenderableTypeRegistry::GetInstance()
+	{
+		static FRenderableTypeRegistry Instance;
+		return Instance;
+	}
+
+	void FRenderableTypeRegistry::Register(const IRenderableType* RenderableType)
+	{
+		if (RenderableType)
+		{
+			RenderableTypesByPrimaryType
+				.FindOrAdd(RenderableType->GetOutputType())
+				.Add(RenderableType);
+		}
+	}
+
+	const FRenderableTypeRegistry::FRenderableTypes& FRenderableTypeRegistry::GetRenderableTypes(FName PrimaryType) const
+	{
+		static const FRenderableTypes EmptyTypes;
+
+		if (const FRenderableTypes* RenderableTypes = RenderableTypesByPrimaryType.Find(PrimaryType))
+		{
+			return *RenderableTypes;
+		}
+		return EmptyTypes;
+	}
+}
diff --git a/Engine/Plugins/Experimental/Dataflow/Source/DataflowEditor/Private/DataflowRendering/DataflowRenderableTypeUtils.cpp b/Engine/Plugins/Experimental/Dataflow/Source/DataflowEditor/Private/DataflowRendering/DataflowRenderableTypeUtils.cpp
new file mode 100644
--- /dev/null
+++ b/Engine/Plugins/Experimental/Dataflow/Source/DataflowEditor/Private/DataflowRendering/DataflowRenderableTypeUtils.cpp
@@ -0,0 +1,145 @@
+// Copyright Epic Games, Inc. All Rights Reserved.
+
+#include "DataflowRendering/DataflowRenderableTypeUtils.h"
+
+#include "Components/DynamicMeshComponent.h"
+
+#include "DataflowRendering/DataflowRenderableComponents.h"
+
+#include "Engine/SkeletalMesh.h"
+#include "Engine/StaticMesh.h"
+
+#include "GeometryCollection/Facades/CollectionMeshFacade.h"
+#include "GeometryCollection/GeometryCollectionUtility.h"
+#include "GeometryCollection/GeometryCollectionAlgo.h"
+#include "GeometryCollection/ManagedArrayCollection.h"
+
+#include "DynamicMeshToMeshDescription.h"
+#include "MeshDescriptionUVsToDynamicMesh.h"
+#include "StaticMeshAttributes.h"
+
+namespace UE::Dataflow::Rendering
+{
+	namespace Private
+	{
+		constexpr float UvScaleFactor = 100.f;
+
+		FVector3d ConvertUvToPosition(const FVector2f& Uv)
+		{
+			return FVector3d((1 - Uv.Y) * UvScaleFactor, (Uv.X * UvScaleFactor), 0);
+		}
+	}
+
+	///////////////////////////////////////////////////////////////////////////////////////////////////////////////
+
+	TSharedPtr<UE::Geometry::FDynamicMesh3> GenerateUvMesh(const FMeshDescription& MeshDescription, int32 UvChannel)
+	{
+		UE::Geometry::FMeshDescriptionUVsToDynamicMesh Converter;
+		Converter.UVLayerIndex = UvChannel;
+		Converter.ScaleFactor = Private::UvScaleFactor;
+		return Converter.GetUVMesh(&MeshDescription);
+	}
+
+	///////////////////////////////////////////////////////////////////////////////////////////////////////////////
+
+	TSharedPtr<UE::Geometry::FDynamicMesh3> GenerateUvMesh(const UStaticMesh& StaticMesh, int32 UvChannel)
+	{
+		if (const FMeshDescription* MeshDescription = StaticMesh.GetMeshDescription(0))
+		{
+			return GenerateUvMesh(*MeshDescription, UvChannel);
+		}
+		return {};
+	}
+
+	///////////////////////////////////////////////////////////////////////////////////////////////////////////////
+
+	TSharedPtr<UE::Geometry::FDynamicMesh3> GenerateUvMesh(const USkeletalMesh& SkeletalMesh, int32 UvChannel)
+	{
+		if (const FMeshDescription* MeshDescription = SkeletalMesh.GetMeshDescription(0))
+		{
+			return GenerateUvMesh(*MeshDescription, UvChannel);
+		}
+		return {};
+	}
+
+	///////////////////////////////////////////////////////////////////////////////////////////////////////////////
+
+	TSharedPtr<UE::Geometry::FDynamicMesh3> GenerateUvMesh(const UE::Geometry::FDynamicMesh3& DynamicMesh, int32 UvChannel)
+	{
+		// first create  mesh description for this dynamic mesh
+		FMeshDescription MeshDescription;
+		{
+			FStaticMeshAttributes Attributes(MeshDescription);
+			Attributes.Register();
+			FConversionToMeshDescriptionOptions ConverterOptions;
+			FDynamicMeshToMeshDescription Converter(ConverterOptions);
+			Converter.Convert(&DynamicMesh, MeshDescription);
+		}
+
+		// generate the UV mesh from it 
+		return GenerateUvMesh(MeshDescription, UvChannel);
+	}
+
+	///////////////////////////////////////////////////////////////////////////////////////////////////////////////
+
+	TSharedPtr<UE::Geometry::FDynamicMesh3> GenerateUvMesh(const FManagedArrayCollection& Collection, int32 UvChannel)
+	{
+		const GeometryCollection::Facades::FCollectionUVFacade UVFacade(Collection);
+		const GeometryCollection::Facades::FCollectionMeshFacade MeshFacade(Collection);
+		if (!UVFacade.IsValid() || !MeshFacade.IndicesAttribute.IsValid())
+		{
+			return {};
+		}
+
+		if (!UVFacade.FindUVLayer(UvChannel))
+		{
+			return {};
+		}
+
+		TSharedPtr<FDynamicMesh3> MeshOut = MakeShared<FDynamicMesh3>();
+
+		// create vertices from Uvs
+		MeshOut->BeginUnsafeVerticesInsert();
+		const TManagedArray<FVector2f>& UVs = UVFacade.GetUVLayer(UvChannel);
+		for (int32 UvIndex = 0; UvIndex < UVs.Num(); ++UvIndex)
+		{
+			MeshOut->InsertVertex(UvIndex, Private::ConvertUvToPosition(UVs[UvIndex]), true);
+		}
+		MeshOut->EndUnsafeVerticesInsert();
+
+		// connect them using the regular triangles
+		MeshOut->BeginUnsafeTrianglesInsert();
+		for (int32 TriIndex = 0; TriIndex < MeshFacade.IndicesAttribute.Num(); ++TriIndex)
+		{
+			const bool bIsVisible = (MeshFacade.VisibleAttribute.IsValid()) ? MeshFacade.VisibleAttribute.Get()[TriIndex] : true;
+			if (bIsVisible)
+			{
+				const FIntVector& Triangle = MeshFacade.IndicesAttribute[TriIndex];
+				MeshOut->InsertTriangle(TriIndex, { Triangle.X, Triangle.Y, Triangle.Z }, 0, true);
+			}
+		}
+		MeshOut->EndUnsafeTrianglesInsert();
+
+		return MeshOut;
+	}
+
+	///////////////////////////////////////////////////////////////////////////////////////////////////////////////
+
+	void SetUvMaterial(UDynamicMeshComponent& Component, int32 MaterialIndex)
+	{
+		UMaterialInterface* UvMaterial = LoadObject<UMaterialInterface>(nullptr, TEXT("/UVEditor/Materials/UVEditor_UnwrapMaterial"));
+		Component.SetMaterial(MaterialIndex, UvMaterial);
+		Component.SetTwoSided(true); // need to be two sided for the Uv display
+	}
+
+	///////////////////////////////////////////////////////////////////////////////////////////////////////////////
+	UDynamicMeshComponent* AddDynamicMeshComponent(UE::Geometry::FDynamicMesh3&& DynamicMesh, FRenderableComponents& OutComponents)
+	{
+		if (UDynamicMeshComponent* Component = OutComponents.AddNewComponent<UDynamicMeshComponent>())
+		{
+			Component->SetMesh(MoveTemp(DynamicMesh));
+			return Component;
+		}
+		return nullptr;
+	}
+}
diff --git a/Engine/Plugins/Experimental/Dataflow/Source/DataflowEditor/Private/DataflowRendering/DataflowSkeletalMeshRenderableType.cpp b/Engine/Plugins/Experimental/Dataflow/Source/DataflowEditor/Private/DataflowRendering/DataflowSkeletalMeshRenderableType.cpp
new file mode 100644
--- /dev/null
+++ b/Engine/Plugins/Experimental/Dataflow/Source/DataflowEditor/Private/DataflowRendering/DataflowSkeletalMeshRenderableType.cpp
@@ -0,0 +1,65 @@
+// Copyright Epic Games, Inc. All Rights Reserved.
+
+#include "DataflowRendering/DataflowSkeletalMeshRenderableType.h"
+
+#include "Components/SkeletalMeshComponent.h"
+
+#include "Dataflow/DataflowRenderingViewMode.h"
+
+#include "DataflowRendering/DataflowRenderableComponents.h"
+#include "DataflowRendering/DataflowRenderableType.h"
+#include "DataflowRendering/DataflowRenderableTypeInstance.h"
+#include "DataflowRendering/DataflowRenderableTypeRegistry.h"
+#include "DataflowRendering/DataflowRenderableTypeUtils.h"
+
+#include "Engine/SkeletalMesh.h"
+
+#include "UObject/ObjectPtr.h"
+
+namespace UE::Dataflow::Private
+{
+	///////////////////////////////////////////////////////////////////////////////////////////////////////////////
+
+	class FSkeletalMeshSurfaceRenderableType : public IRenderableType
+	{
+		UE_DATAFLOW_IRENDERABLE_PRIMARY_TYPE(TObjectPtr<USkeletalMesh>, SkeletalMesh);
+		UE_DATAFLOW_IRENDERABLE_RENDER_GROUP(Surface);
+		UE_DATAFLOW_IRENDERABLE_VIEW_MODE(FDataflowConstruction3DViewMode);
+
+		virtual void GetPrimitiveComponents(const FRenderableTypeInstance& Instance, FRenderableComponents& OutComponents) const override
+		{
+			if (const TObjectPtr<USkeletalMesh> SkeletalMesh = GetSkeletalMesh(Instance, nullptr))
+			{
+				if (USkeletalMeshComponent* Component = OutComponents.AddNewComponent<USkeletalMeshComponent>())
+				{
+					Component->SetSkeletalMesh(SkeletalMesh);
+				}
+			}
+		}
+	};
+
+	////////////////////////////////////////////////////////////////////////////////////////////////////////////////////////////
+
+	class FSkeletalMeshUvsRenderableType : public IRenderableType
+	{
+		UE_DATAFLOW_IRENDERABLE_PRIMARY_TYPE(TObjectPtr<USkeletalMesh>, SkeletalMesh);
+		UE_DATAFLOW_IRENDERABLE_RENDER_GROUP(Uvs);
+		UE_DATAFLOW_IRENDERABLE_VIEW_MODE(FDataflowConstructionUVViewMode);
+
+		virtual void GetPrimitiveComponents(const FRenderableTypeInstance& Instance, FRenderableComponents& OutComponents) const override
+		{
+			if (const TObjectPtr<USkeletalMesh> SkeletalMesh = GetSkeletalMesh(Instance, nullptr))
+			{
+				Rendering::AddUvDynamicMeshComponent(*SkeletalMesh, 0, OutComponents);
+			}
+		}
+	};
+
+	///////////////////////////////////////////////////////////////////////////////////////////////////////////////
+
+	void RegisterSkeletalMeshRenderableTypes()
+	{
+		UE_DATAFLOW_REGISTER_RENDERABLE_TYPE(FSkeletalMeshSurfaceRenderableType);
+		UE_DATAFLOW_REGISTER_RENDERABLE_TYPE(FSkeletalMeshUvsRenderableType);
+	}
+}
\ No newline at end of file
diff --git a/Engine/Plugins/Experimental/Dataflow/Source/DataflowEditor/Private/DataflowRendering/DataflowSkeletalMeshRenderableType.h b/Engine/Plugins/Experimental/Dataflow/Source/DataflowEditor/Private/DataflowRendering/DataflowSkeletalMeshRenderableType.h
new file mode 100644
--- /dev/null
+++ b/Engine/Plugins/Experimental/Dataflow/Source/DataflowEditor/Private/DataflowRendering/DataflowSkeletalMeshRenderableType.h
@@ -0,0 +1,8 @@
+// Copyright Epic Games, Inc. All Rights Reserved.
+
+#pragma once
+
+namespace UE::Dataflow::Private
+{
+	void RegisterSkeletalMeshRenderableTypes();
+}
diff --git a/Engine/Plugins/Experimental/Dataflow/Source/DataflowEditor/Private/DataflowRendering/DataflowStaticMeshRenderableType.cpp b/Engine/Plugins/Experimental/Dataflow/Source/DataflowEditor/Private/DataflowRendering/DataflowStaticMeshRenderableType.cpp
new file mode 100644
--- /dev/null
+++ b/Engine/Plugins/Experimental/Dataflow/Source/DataflowEditor/Private/DataflowRendering/DataflowStaticMeshRenderableType.cpp
@@ -0,0 +1,65 @@
+// Copyright Epic Games, Inc. All Rights Reserved.
+
+#include "DataflowRendering/DataflowStaticMeshRenderableType.h"
+
+#include "Components/StaticMeshComponent.h"
+
+#include "Dataflow/DataflowRenderingViewMode.h"
+
+#include "DataflowRendering/DataflowRenderableComponents.h"
+#include "DataflowRendering/DataflowRenderableType.h"
+#include "DataflowRendering/DataflowRenderableTypeInstance.h"
+#include "DataflowRendering/DataflowRenderableTypeRegistry.h"
+#include "DataflowRendering/DataflowRenderableTypeUtils.h"
+
+#include "Engine/StaticMesh.h"
+
+#include "UObject/ObjectPtr.h"
+
+namespace UE::Dataflow::Private
+{
+	///////////////////////////////////////////////////////////////////////////////////////////////////////////////
+
+	class FStaticMeshSurfaceRenderableType : public IRenderableType
+	{
+		UE_DATAFLOW_IRENDERABLE_PRIMARY_TYPE(TObjectPtr<UStaticMesh>, StaticMesh);
+		UE_DATAFLOW_IRENDERABLE_RENDER_GROUP(Surface);
+		UE_DATAFLOW_IRENDERABLE_VIEW_MODE(FDataflowConstruction3DViewMode);
+
+		virtual void GetPrimitiveComponents(const FRenderableTypeInstance& Instance, FRenderableComponents& OutComponents) const override
+		{
+			if (const TObjectPtr<UStaticMesh> StaticMesh = GetStaticMesh(Instance, nullptr))
+			{
+				if (UStaticMeshComponent* Component = OutComponents.AddNewComponent<UStaticMeshComponent>())
+				{
+					Component->SetStaticMesh(StaticMesh);
+				}
+			}
+		}
+	};
+
+	///////////////////////////////////////////////////////////////////////////////////////////////////////////////
+
+	class FStaticMeshUvsRenderableType : public IRenderableType
+	{
+		UE_DATAFLOW_IRENDERABLE_PRIMARY_TYPE(TObjectPtr<UStaticMesh>, StaticMesh)
+		UE_DATAFLOW_IRENDERABLE_RENDER_GROUP(Uvs)
+		UE_DATAFLOW_IRENDERABLE_VIEW_MODE(FDataflowConstructionUVViewMode);
+
+		virtual void GetPrimitiveComponents(const FRenderableTypeInstance& Instance, FRenderableComponents& OutComponents) const override
+		{
+			if (const TObjectPtr<UStaticMesh> StaticMesh = GetStaticMesh(Instance, nullptr))
+			{
+				Rendering::AddUvDynamicMeshComponent(*StaticMesh, 0, OutComponents);
+			}
+		}
+	};
+
+	///////////////////////////////////////////////////////////////////////////////////////////////////////////////
+
+	void RegisterStaticMeshRenderableTypes()
+	{
+		UE_DATAFLOW_REGISTER_RENDERABLE_TYPE(FStaticMeshSurfaceRenderableType);
+		UE_DATAFLOW_REGISTER_RENDERABLE_TYPE(FStaticMeshUvsRenderableType);
+	}
+}
\ No newline at end of file
diff --git a/Engine/Plugins/Experimental/Dataflow/Source/DataflowEditor/Private/DataflowRendering/DataflowStaticMeshRenderableType.h b/Engine/Plugins/Experimental/Dataflow/Source/DataflowEditor/Private/DataflowRendering/DataflowStaticMeshRenderableType.h
new file mode 100644
--- /dev/null
+++ b/Engine/Plugins/Experimental/Dataflow/Source/DataflowEditor/Private/DataflowRendering/DataflowStaticMeshRenderableType.h
@@ -0,0 +1,8 @@
+// Copyright Epic Games, Inc. All Rights Reserved.
+
+#pragma once
+
+namespace UE::Dataflow::Private
+{
+	void RegisterStaticMeshRenderableTypes();
+}
diff --git a/Engine/Plugins/Experimental/Dataflow/Source/DataflowEditor/Public/DataflowRendering/DataflowRenderableComponents.h b/Engine/Plugins/Experimental/Dataflow/Source/DataflowEditor/Public/DataflowRendering/DataflowRenderableComponents.h
new file mode 100644
--- /dev/null
+++ b/Engine/Plugins/Experimental/Dataflow/Source/DataflowEditor/Public/DataflowRendering/DataflowRenderableComponents.h
@@ -0,0 +1,47 @@
+// Copyright Epic Games, Inc. All Rights Reserved.
+
+#pragma once
+
+#include "Containers/Array.h"
+#include "Templates/UnrealTypeTraits.h"
+
+class AActor;
+class UPrimitiveComponent;
+
+namespace UE::Dataflow
+{
+	/**
+	* List of renderable components for a sepcific root actor
+	* this class is used by IRenderableType to collect the components to render 
+	*/
+	struct FRenderableComponents
+	{
+	public:
+		FRenderableComponents(AActor* InRootActor)
+			: RootActor(InRootActor)
+		{}
+
+		/**
+		* Create a new component of a specific type ( needs to inherit from UPrimitiveComponent )
+		* and add it to the list 
+		* returns  the newly added component so that it can be configured by the caller 
+		*/
+		template<typename T UE_REQUIRES(TIsDerivedFrom<T, UPrimitiveComponent>::Value)>
+		T* AddNewComponent()
+		{
+			// for now always allocate a new one - in the future we could use a pool 
+			T* NewComponent = NewObject<T>(RootActor);
+			if (NewComponent)
+			{
+				Components.Add(NewComponent);
+			}
+			return NewComponent;
+		}
+
+		const TArray<UPrimitiveComponent*>& GetComponents() const { return Components; }
+
+	private:
+		AActor* RootActor = nullptr;
+		TArray<UPrimitiveComponent*> Components;
+	};
+}
diff --git a/Engine/Plugins/Experimental/Dataflow/Source/DataflowEditor/Public/DataflowRendering/DataflowRenderableType.h b/Engine/Plugins/Experimental/Dataflow/Source/DataflowEditor/Public/DataflowRendering/DataflowRenderableType.h
new file mode 100644
--- /dev/null
+++ b/Engine/Plugins/Experimental/Dataflow/Source/DataflowEditor/Public/DataflowRendering/DataflowRenderableType.h
@@ -0,0 +1,55 @@
+// Copyright Epic Games, Inc. All Rights Reserved.
+
+#pragma once
+
+#include "UObject/NameTypes.h"
+
+namespace UE::Dataflow
+{
+	struct FRenderableComponents;
+	struct FRenderableTypeInstance;
+	class IDataflowConstructionViewMode;
+
+	/**
+	* this interface defines a way to render a specific data type in Dataflow
+	* rendering is achieve by provided a number of components that will be added to the rednering scene
+	*/
+	struct IRenderableType
+	{
+	public:
+		/** Returns the renderable primary type (ie: TObjectPtr<UDynamicMesh>) */
+		virtual FName GetOutputType() const = 0;
+
+		/** Returns the supported render group, useful when rendering a specific type in a particular way (ie : "Surface" may render the geometry related groups ) */
+		virtual FName GetRenderGroup() const { return NAME_None; }
+
+		/** return true if the ViewMode is supported for rendering this */
+		virtual bool IsViewModeSupported(const IDataflowConstructionViewMode& ViewMode) const = 0;
+
+		/** Whether this type can be rendered in this context */
+		virtual bool CanRender(const FRenderableTypeInstance& Instance) const { return true; }
+
+		/** Get the primitive components for the specific render context  */
+		virtual void GetPrimitiveComponents(const FRenderableTypeInstance& Instance, FRenderableComponents& OutComponents) const = 0;
+	};
+}
+
+// Macros to ease the definition of renderable types
+
+#define UE_DATAFLOW_IRENDERABLE_PRIMARY_TYPE(Type, Name) \
+	static inline const FName OutputTypeName = TEXT(#Type); \
+	virtual FName GetOutputType() const override \
+		{ return OutputTypeName; } \
+	const Type& Get##Name(const FRenderableTypeInstance& Instance, const Type& Default = Type()) const \
+		{ return Instance.GetOutputValue<Type>(Default); }
+
+#define UE_DATAFLOW_IRENDERABLE_RENDER_GROUP(Name) \
+	static inline const FName RenderGroupName = TEXT(#Name); \
+	virtual FName GetRenderGroup() const override \
+		{ return RenderGroupName; } 
+
+#define UE_DATAFLOW_IRENDERABLE_VIEW_MODE(ViewModeClass) \
+	virtual bool IsViewModeSupported(const UE::Dataflow::IDataflowConstructionViewMode& InViewMode) const override \
+		{ return (InViewMode.GetName() == ViewModeClass::Name); }
+
+
diff --git a/Engine/Plugins/Experimental/Dataflow/Source/DataflowEditor/Public/DataflowRendering/DataflowRenderableTypeInstance.h b/Engine/Plugins/Experimental/Dataflow/Source/DataflowEditor/Public/DataflowRendering/DataflowRenderableTypeInstance.h
new file mode 100644
--- /dev/null
+++ b/Engine/Plugins/Experimental/Dataflow/Source/DataflowEditor/Public/DataflowRendering/DataflowRenderableTypeInstance.h
@@ -0,0 +1,82 @@
+// Copyright Epic Games, Inc. All Rights Reserved.
+
+#pragma once
+
+#include "Containers/UnrealString.h"
+#include "UObject/NameTypes.h"
+
+#include "Dataflow/DataflowNodeParameters.h"
+#include "Dataflow/DataflowNode.h"
+
+class UPrimitiveComponent;
+class AActor;
+
+namespace UE::Dataflow
+{
+	struct IRenderableType;
+	struct FRenderableComponents;
+	class IDataflowConstructionViewMode;
+
+	/**
+	* This describe a renderable type instance which match to a specific node + output(s)
+	* And to be rendered in the context a specific Actor and Dataflow context
+	*/
+	struct FRenderableTypeInstance final
+	{
+	public:
+		/** find all renderable instances for a specific node and a specific view mode */
+		static void GetRenderableInstancesForNode(FContext& Context, const IDataflowConstructionViewMode& ViewMode, const FDataflowNode& Node, TArray<FRenderableTypeInstance>& OutRenderableInstances);
+
+		/** Whether this instance can render */
+		bool CanRender() const;
+
+		/** Returns the display name of this instance */
+		FString GetDisplayName() const;
+
+		/** Get the corresponding  */
+		void GetPrimitiveComponents(FRenderableComponents& OutComponents) const;
+
+		/** 
+		* Get value for the primary output to render 
+		* Returns the default value if not found of this fails internally 
+		*/
+		template <typename T>
+		const T& GetOutputValue(const T& Default = T()) const
+		{
+			return GetOutputValue(OutputName, Default);
+		}
+
+	private:
+		template <typename T>
+		const T& GetOutputValue(FName Name, const T& Default = T()) const
+		{
+			if (Node)
+			{
+				if (const FDataflowOutput* Output = Node->FindOutput(Name))
+				{
+					if (EvaluationContext.IsThreaded())
+					{
+						return Output->ReadValue<T>(EvaluationContext, Default);
+					}
+					return Output->GetValue<T>(EvaluationContext, Default);
+				}
+			}
+			return Default;
+		}
+
+		FRenderableTypeInstance(FContext& InContext, const TSharedPtr<const FDataflowNode>& InNode, const FName InOutputName, const IRenderableType* InRenderableType)
+			: EvaluationContext(InContext)
+			, Node(InNode)
+			, OutputName(InOutputName)
+			, RenderableType(InRenderableType)
+		{}
+
+	private:
+		FContext& EvaluationContext;
+		TSharedPtr<const FDataflowNode> Node;
+		const FName OutputName;
+		const IRenderableType* RenderableType;
+	};
+}
+
+
diff --git a/Engine/Plugins/Experimental/Dataflow/Source/DataflowEditor/Public/DataflowRendering/DataflowRenderableTypeRegistry.h b/Engine/Plugins/Experimental/Dataflow/Source/DataflowEditor/Public/DataflowRendering/DataflowRenderableTypeRegistry.h
new file mode 100644
--- /dev/null
+++ b/Engine/Plugins/Experimental/Dataflow/Source/DataflowEditor/Public/DataflowRendering/DataflowRenderableTypeRegistry.h
@@ -0,0 +1,40 @@
+// Copyright Epic Games, Inc. All Rights Reserved.
+
+#pragma once
+
+#include "Containers/Array.h"
+#include "Containers/Map.h"
+#include "UObject/NameTypes.h"
+
+namespace UE::Dataflow
+{
+	struct IRenderableType;
+
+	/** this registry contains all Dataflow renderable types instances */
+	struct FRenderableTypeRegistry
+	{
+	public:
+		using FRenderableTypes = TArray<const IRenderableType*>;
+
+		/** register a renderable type instance - use UE_DATAFLOW_REGISTER_RENDERABLE_TYPE instead of calling this directly for ease of use */
+		void Register(const IRenderableType* RenderableType);
+
+		/** Get all the available renderable types for a specific primary type */
+		const FRenderableTypes& GetRenderableTypes(FName PrimaryType) const;
+
+		/** get the one and only instance of this registry */
+		static FRenderableTypeRegistry& GetInstance();
+
+	private:
+		FRenderableTypeRegistry() = default;
+
+	private:
+		TMap<FName, FRenderableTypes> RenderableTypesByPrimaryType;
+	};
+}
+
+// macro to use in a module / plugin to register a specific renderable type class
+
+#define UE_DATAFLOW_REGISTER_RENDERABLE_TYPE(Type) \
+	static const Type InstanceOf##Type; \
+	FRenderableTypeRegistry::GetInstance().Register(&InstanceOf##Type);
diff --git a/Engine/Plugins/Experimental/Dataflow/Source/DataflowEditor/Public/DataflowRendering/DataflowRenderableTypeUtils.h b/Engine/Plugins/Experimental/Dataflow/Source/DataflowEditor/Public/DataflowRendering/DataflowRenderableTypeUtils.h
new file mode 100644
--- /dev/null
+++ b/Engine/Plugins/Experimental/Dataflow/Source/DataflowEditor/Public/DataflowRendering/DataflowRenderableTypeUtils.h
@@ -0,0 +1,59 @@
+// Copyright Epic Games, Inc. All Rights Reserved.
+
+#pragma once
+
+#include "Templates/SharedPointer.h"
+
+struct FManagedArrayCollection;
+struct FMeshDescription;
+
+class UDynamicMeshComponent;
+class USkeletalMesh;
+class UStaticMesh;
+
+namespace UE::Geometry
+{
+	class FDynamicMesh3;
+}
+
+namespace UE::Dataflow
+{
+	struct FRenderableComponents;
+
+	namespace Rendering
+	{
+		/** Generate a dynamic mesh from a mesh description UV data */
+		TSharedPtr<UE::Geometry::FDynamicMesh3> GenerateUvMesh(const FMeshDescription& MeshDescription, int32 UvChannel);
+
+		/** Generate a dynamic mesh from a static mesh UV data */
+		TSharedPtr<UE::Geometry::FDynamicMesh3> GenerateUvMesh(const UStaticMesh& StaticMesh, int32 UvChannel);
+
+		/** Generate a dynamic mesh from a skeletal mesh UV data */
+		TSharedPtr<UE::Geometry::FDynamicMesh3> GenerateUvMesh(const USkeletalMesh& SkeletalMesh, int32 UvChannel);
+
+		/** Generate a dynamic mesh from a dynamic mesh UV data */
+		TSharedPtr<UE::Geometry::FDynamicMesh3> GenerateUvMesh(const UE::Geometry::FDynamicMesh3& DynamicMesh, int32 UvChannel);
+
+		/** Generate a dynamic mesh from a managed array collection UV data */
+		TSharedPtr<UE::Geometry::FDynamicMesh3> GenerateUvMesh(const FManagedArrayCollection& Collection, int32 UvChannel);
+
+		/** Make a dynamic mesh component and add it to the OutComponents */
+		UDynamicMeshComponent* AddDynamicMeshComponent(UE::Geometry::FDynamicMesh3&& DynamicMesh, FRenderableComponents& OutComponents);
+
+		/** set UV material on a primitive component */
+		void SetUvMaterial(UDynamicMeshComponent& Component, int32 MaterialIndex);
+
+		/** Generate a Uv mesh from a templated source and make a dynamic mesh component out of it and add it to the OutComponents */
+		template<typename T>
+		void AddUvDynamicMeshComponent(const T& Source, int32 UvIndex, FRenderableComponents& OutComponents)
+		{
+			if (TSharedPtr<UE::Geometry::FDynamicMesh3> UvMesh = GenerateUvMesh(Source, UvIndex))
+			{
+				if (UDynamicMeshComponent* Component = AddDynamicMeshComponent(MoveTemp(*UvMesh), OutComponents))
+				{
+					SetUvMaterial(*Component, 0);
+				}
+			}
+		}
+	}
+}
diff --git a/Engine/Source/Runtime/Experimental/Dataflow/Core/Public/Dataflow/DataflowNode.h b/Engine/Source/Runtime/Experimental/Dataflow/Core/Public/Dataflow/DataflowNode.h
--- a/Engine/Source/Runtime/Experimental/Dataflow/Core/Public/Dataflow/DataflowNode.h
+++ b/Engine/Source/Runtime/Experimental/Dataflow/Core/Public/Dataflow/DataflowNode.h
@@ -760,6 +760,10 @@ PRAGMA_ENABLE_DEPRECATION_WARNINGS
 	// returns true to disable node properties in its detail view when the corresponding input is connected ( default is true )
 	virtual bool MakeConnectedPropertiesReadOnly() const { return true; }
 
+	// this method is called when the construction viewport render a specific output in the viewport 
+	// this allow the node to be able to set extra properties on the component or set it up differently 
+	DATAFLOWCORE_API virtual void OnRenderOutput(UE::Dataflow::FContext& Context, const FName OutputName, const FName RenderGroup, const TArray<UPrimitiveComponent*>& RenderComponents) const {};
+
 	bool HasValidConnections() const { return bHasValidConnections; }
 
 	virtual bool IsA(FName InType) const 
