Force debug draw redraw when entering a Dataflow Tool.

diff --git a/Engine/Plugins/Experimental/Dataflow/Source/DataflowEditor/Private/Dataflow/DataflowEditorMode.cpp b/Engine/Plugins/Experimental/Dataflow/Source/DataflowEditor/Private/Dataflow/DataflowEditorMode.cpp
--- a/Engine/Plugins/Experimental/Dataflow/Source/DataflowEditor/Private/Dataflow/DataflowEditorMode.cpp
+++ b/Engine/Plugins/Experimental/Dataflow/Source/DataflowEditor/Private/Dataflow/DataflowEditorMode.cpp
@@ -380,6 +380,12 @@ void UDataflowEditorMode::OnToolStarted(UInteractiveToolManager* Manager, UInter
 
 	// Temporarily disable wireframe render toggle switch on all nodes
 	SetWireframeRenderToggleEnabled(false);
+
+	// Force redraw
+	if (const TSharedPtr<FDataflowConstructionScene> ConstructionScenePtr = ConstructionScene.Pin())
+	{
+		ConstructionScenePtr->SetDirtyFlag();
+	}
 }
 
 void UDataflowEditorMode::OnToolEnded(UInteractiveToolManager* Manager, UInteractiveTool* Tool)
diff --git a/Engine/Plugins/Experimental/Dataflow/Source/DataflowEditor/Public/Dataflow/DataflowEditorPreviewSceneBase.h b/Engine/Plugins/Experimental/Dataflow/Source/DataflowEditor/Public/Dataflow/DataflowEditorPreviewSceneBase.h
--- a/Engine/Plugins/Experimental/Dataflow/Source/DataflowEditor/Public/Dataflow/DataflowEditorPreviewSceneBase.h
+++ b/Engine/Plugins/Experimental/Dataflow/Source/DataflowEditor/Public/Dataflow/DataflowEditorPreviewSceneBase.h
@@ -131,6 +131,9 @@ public:
 	/** Return true if the preview scene is dirty */
 	bool IsSceneDirty() const {return bPreviewSceneDirty;}
 
+	/** Set the dirty flag */
+	void SetDirtyFlag() { bPreviewSceneDirty = true; }
+
 	/** Reset the dirty flag */
 	void ResetDirtyFlag() {bPreviewSceneDirty = false;}
 
