Update DataflowPrimitiveNode to take in a Dataflow Context.

diff --git a/Engine/Plugins/Experimental/Dataflow/Source/DataflowEditor/Private/Dataflow/DataflowConstructionScene.cpp b/Engine/Plugins/Experimental/Dataflow/Source/DataflowEditor/Private/Dataflow/DataflowConstructionScene.cpp
--- a/Engine/Plugins/Experimental/Dataflow/Source/DataflowEditor/Private/Dataflow/DataflowConstructionScene.cpp
+++ b/Engine/Plugins/Experimental/Dataflow/Source/DataflowEditor/Private/Dataflow/DataflowConstructionScene.cpp
@@ -184,13 +184,13 @@ void FDataflowConstructionScene::UpdateDynamicMeshComponents()
 					GeometryCollection::Facades::FRenderingFacade RenderingFacade(*RenderCollection);
 					RenderingFacade.DefineSchema();
 					
-					const bool bHasPrimitves = UE::Dataflow::RenderNodeOutput(RenderingFacade, *Target, *EditorContent, bEvaluateOutputs);
+					const bool bHasRenderCollectionPrimitives = UE::Dataflow::RenderNodeOutput(RenderingFacade, *Target, *EditorContent, bEvaluateOutputs);
 
 					if (Target == EditorContent->GetSelectedNode())
 					{
 						EditorContent->SetRenderCollection(RenderCollection);
 					}
-					if(!bHasPrimitves)
+					if(!bHasRenderCollectionPrimitives)
 					{
 						const int32 NumGeometry = RenderingFacade.NumGeometry();
 	
@@ -271,7 +271,7 @@ void FDataflowConstructionScene::UpdateDynamicMeshComponents()
 					GeometryCollection::Facades::FRenderingFacade RenderingFacade(*RenderCollection);
 					RenderingFacade.DefineSchema();
 
-					const bool bHasPrimitves = UE::Dataflow::RenderNodeOutput(RenderingFacade, *Target, *EditorContent, bEvaluateOutputs);
+					const bool bHasRenderCollectionPrimitives = UE::Dataflow::RenderNodeOutput(RenderingFacade, *Target, *EditorContent, bEvaluateOutputs);
 
 					if (Target == EditorContent->GetSelectedNode())
 
@@ -279,7 +279,7 @@ void FDataflowConstructionScene::UpdateDynamicMeshComponents()
 						EditorContent->SetRenderCollection(RenderCollection);
 					}
 
-					if(!bHasPrimitves)
+					if(!bHasRenderCollectionPrimitives)
 					{
 						const int32 NumGeometry = RenderingFacade.NumGeometry();
 						for (int32 MeshIndex = 0; MeshIndex < NumGeometry; ++MeshIndex)
@@ -351,7 +351,7 @@ void FDataflowConstructionScene::UpdatePrimitiveComponents()
 				{
 					if(SelectedDataflowNode->IsA(FDataflowPrimitiveNode::StaticType()))
 					{
-						StaticCastSharedPtr<FDataflowPrimitiveNode>(SelectedDataflowNode)->AddPrimitiveComponents( 
+						StaticCastSharedPtr<FDataflowPrimitiveNode>(SelectedDataflowNode)->AddPrimitiveComponents(*DataflowContext,  
 							EditorContent->GetRenderCollection(), DataflowContext->Owner, RootSceneActor, PrimitiveComponents);
 					}
 				}
diff --git a/Engine/Plugins/Experimental/Dataflow/Source/DataflowEditor/Private/Dataflow/DataflowEditorUtil.cpp b/Engine/Plugins/Experimental/Dataflow/Source/DataflowEditor/Private/Dataflow/DataflowEditorUtil.cpp
--- a/Engine/Plugins/Experimental/Dataflow/Source/DataflowEditor/Private/Dataflow/DataflowEditorUtil.cpp
+++ b/Engine/Plugins/Experimental/Dataflow/Source/DataflowEditor/Private/Dataflow/DataflowEditorUtil.cpp
@@ -148,7 +148,7 @@ namespace UE::Dataflow
 		const IDataflowConstructionViewMode* ConstructionViewMode = EditorContent.GetConstructionViewMode();
 		const TSharedPtr<FEngineContext>& DataflowContext = EditorContent.GetDataflowContext();
 
-		bool bHasPrimitives = false;
+		bool bHasRenderCollectionPrimitives = false;
 		if (DataflowAsset && DataflowContext && ConstructionViewMode)
 		{
 			if (FRenderingFactory* const Factory = FRenderingFactory::GetInstance())
@@ -160,13 +160,13 @@ namespace UE::Dataflow
 						if (TSharedPtr<const FDataflowNode> NodeTarget = Graph->FindBaseNode(FName(Node.GetName())))
 						{
 							Factory->RenderNodeOutput(RenderingFacade, FGraphRenderingState{ Node.GetDataflowNodeGuid(), NodeTarget.Get(), Parameter, *DataflowContext, *ConstructionViewMode, bEvaluateOutputs });
-							bHasPrimitives = bHasPrimitives || NodeTarget->HasPrimitives();
+							bHasRenderCollectionPrimitives = bHasRenderCollectionPrimitives || NodeTarget->HasRenderCollectionPrimitives();
 						}
 					}
 				}
 			}
 		}
-		return bHasPrimitives;
+		return bHasRenderCollectionPrimitives;
 	}
 
 
diff --git a/Engine/Plugins/Experimental/Dataflow/Source/DataflowNodes/Private/Dataflow/DataflowCollectionEditSkinWeightsNode.cpp b/Engine/Plugins/Experimental/Dataflow/Source/DataflowNodes/Private/Dataflow/DataflowCollectionEditSkinWeightsNode.cpp
--- a/Engine/Plugins/Experimental/Dataflow/Source/DataflowNodes/Private/Dataflow/DataflowCollectionEditSkinWeightsNode.cpp
+++ b/Engine/Plugins/Experimental/Dataflow/Source/DataflowNodes/Private/Dataflow/DataflowCollectionEditSkinWeightsNode.cpp
@@ -271,7 +271,7 @@ TArray<UE::Dataflow::FRenderingParameter> FDataflowCollectionEditSkinWeightsNode
 	return FDataflowAddScalarVertexPropertyCallbackRegistry::Get().GetRenderingParameters(VertexGroup.Name);
 }
 
-void FDataflowCollectionEditSkinWeightsNode::AddPrimitiveComponents(const TSharedPtr<const FManagedArrayCollection> RenderCollection, TObjectPtr<UObject> NodeOwner, 
+void FDataflowCollectionEditSkinWeightsNode::AddPrimitiveComponents(UE::Dataflow::FContext& Context, const TSharedPtr<const FManagedArrayCollection> RenderCollection, TObjectPtr<UObject> NodeOwner,
 	TObjectPtr<AActor> RootActor, TArray<TObjectPtr<UPrimitiveComponent>>& PrimitiveComponents)  
 {
 	if(RootActor)
diff --git a/Engine/Plugins/Experimental/Dataflow/Source/DataflowNodes/Public/Dataflow/DataflowCollectionEditSkinWeightsNode.h b/Engine/Plugins/Experimental/Dataflow/Source/DataflowNodes/Public/Dataflow/DataflowCollectionEditSkinWeightsNode.h
--- a/Engine/Plugins/Experimental/Dataflow/Source/DataflowNodes/Public/Dataflow/DataflowCollectionEditSkinWeightsNode.h
+++ b/Engine/Plugins/Experimental/Dataflow/Source/DataflowNodes/Public/Dataflow/DataflowCollectionEditSkinWeightsNode.h
@@ -113,7 +113,7 @@ private:
 	//~ Begin FDataflowPrimitiveNode interface
 	UE_API virtual TArray<UE::Dataflow::FRenderingParameter> GetRenderParametersImpl() const override;
 	UE_API virtual void Evaluate(UE::Dataflow::FContext& Context, const FDataflowOutput* Out) const override;
-	UE_API virtual void AddPrimitiveComponents(const TSharedPtr<const FManagedArrayCollection> RenderCollection,
+	UE_API virtual void AddPrimitiveComponents(UE::Dataflow::FContext& Context, const TSharedPtr<const FManagedArrayCollection> RenderCollection,
 		TObjectPtr<UObject> NodeOwner, TObjectPtr<AActor> RootActor, TArray<TObjectPtr<UPrimitiveComponent>>& PrimitiveComponents) override;
 	UE_API virtual void OnInvalidate() override;
 #if WITH_EDITOR
diff --git a/Engine/Plugins/Experimental/Dataflow/Source/DataflowNodes/Public/Dataflow/DataflowPrimitiveNode.h b/Engine/Plugins/Experimental/Dataflow/Source/DataflowNodes/Public/Dataflow/DataflowPrimitiveNode.h
--- a/Engine/Plugins/Experimental/Dataflow/Source/DataflowNodes/Public/Dataflow/DataflowPrimitiveNode.h
+++ b/Engine/Plugins/Experimental/Dataflow/Source/DataflowNodes/Public/Dataflow/DataflowPrimitiveNode.h
@@ -1,4 +1,4 @@
-﻿// Copyright Epic Games, Inc. All Rights Reserved.
+// Copyright Epic Games, Inc. All Rights Reserved.
 
 #pragma once
 
@@ -33,13 +33,13 @@ struct FDataflowPrimitiveNode : public FDataflowNode
 			|| Super::IsA(InType); 
 	}
 
-	virtual bool HasPrimitives() const override
+	virtual bool HasRenderCollectionPrimitives() const override
 	{
 		return true;
 	}
 
 	/** Add primitive components to the construction scene */
-	virtual void AddPrimitiveComponents(const TSharedPtr<const FManagedArrayCollection> RenderCollection, TObjectPtr<UObject> NodeOwner, 
+	virtual void AddPrimitiveComponents(UE::Dataflow::FContext& Context, const TSharedPtr<const FManagedArrayCollection> RenderCollection, TObjectPtr<UObject> NodeOwner,
 		TObjectPtr<AActor> RootActor, TArray<TObjectPtr<UPrimitiveComponent>>& PrimitiveComponents) {}
 };
 
diff --git a/Engine/Source/Runtime/Experimental/Dataflow/Core/Public/Dataflow/DataflowNode.h b/Engine/Source/Runtime/Experimental/Dataflow/Core/Public/Dataflow/DataflowNode.h
--- a/Engine/Source/Runtime/Experimental/Dataflow/Core/Public/Dataflow/DataflowNode.h
+++ b/Engine/Source/Runtime/Experimental/Dataflow/Core/Public/Dataflow/DataflowNode.h
@@ -117,8 +117,8 @@ PRAGMA_ENABLE_DEPRECATION_WARNINGS
 	UE_DEPRECATED(5.4, "FDataflowNode::IsExperimental is deprecated.")
 	virtual bool IsExperimental() { return false; }
 
-	/** Override if the node could generate primitives .*/
-	virtual bool HasPrimitives() const
+	/** Override if node will generate primitives from the RenderCollection.*/
+	virtual bool HasRenderCollectionPrimitives() const
 	{
 		return false;
 	}
