Dataflow : refactor the UpdateDynamicMeshComponents in preparation for enabling the new rednering system
* share as much code as possible between the wireframe and solid code paths
* move debug mesh update in its own method
* move floor update iun its own method

diff --git a/Engine/Plugins/Experimental/Dataflow/Source/DataflowEditor/Private/Dataflow/DataflowConstructionScene.cpp b/Engine/Plugins/Experimental/Dataflow/Source/DataflowEditor/Private/Dataflow/DataflowConstructionScene.cpp
--- a/Engine/Plugins/Experimental/Dataflow/Source/DataflowEditor/Private/Dataflow/DataflowConstructionScene.cpp
+++ b/Engine/Plugins/Experimental/Dataflow/Source/DataflowEditor/Private/Dataflow/DataflowConstructionScene.cpp
@@ -10,6 +10,7 @@
 #include "Dataflow/DataflowObject.h"
 #include "Dataflow/DataflowEditorStyle.h"
 #include "Dataflow/DataflowRenderingFactory.h"
+#include "DataflowRendering/DataflowRenderableComponents.h"
 #include "Drawing/MeshElementsVisualizer.h"
 #include "Elements/Framework/EngineElementsLibrary.h"
 #include "Selection.h"
@@ -178,70 +179,13 @@ void FDataflowConstructionScene::UpdateDynamicMeshComponents()
 		{
 			for (TObjectPtr<const UDataflowEdNode> Target : DataflowAsset->GetRenderTargets())
 			{
-				if (Target)
-				{
-					TSharedPtr<FManagedArrayCollection> RenderCollection(new FManagedArrayCollection);
-					GeometryCollection::Facades::FRenderingFacade RenderingFacade(*RenderCollection);
-					RenderingFacade.DefineSchema();
-					
-					const bool bHasRenderCollectionPrimitives = UE::Dataflow::RenderNodeOutput(RenderingFacade, *Target, *EditorContent, bEvaluateOutputs);
-
-					if (Target == EditorContent->GetSelectedNode())
-					{
-						EditorContent->SetRenderCollection(RenderCollection);
-					}
-					if(!bHasRenderCollectionPrimitives)
-					{
-						const int32 NumGeometry = RenderingFacade.NumGeometry();
-	
-						DebugMesh.Reset();
-						int32 VertexIndex = 0, FaceIndex = 0;
-
-						const TManagedArray<int32>& VertexStart = RenderingFacade.GetVertexStart();
-						const TManagedArray<int32>& VertexCount = RenderingFacade.GetVertexCount();
-						const TManagedArray<int32>& FaceStart = RenderingFacade.GetIndicesStart();
-						const TManagedArray<int32>& FaceCount = RenderingFacade.GetIndicesCount();
-
-						for (int32 MeshIndex = 0; MeshIndex < NumGeometry; ++MeshIndex)
-						{
-							FDynamicMesh3 DynamicMesh;
-							UE::Dataflow::Conversion::RenderingFacadeToDynamicMesh(RenderingFacade, MeshIndex, DynamicMesh);
+				UpdateComponentsFromNode(Target, /*bWireframe*/false);
+			}
 
-							if (DynamicMesh.VertexCount())
-							{
-								const FString MeshName = RenderingFacade.GetGeometryName()[MeshIndex];
-								const FTransform& MeshTransform = RenderingFacade.GetGeometryTransform()[MeshIndex];
-
-								const TManagedArray<FString>& MaterialPaths = RenderingFacade.GetMaterialPaths();
-								const int32 MaterialStart = RenderingFacade.GetMaterialStart()[MeshIndex];
-								const int32 MaterialCount = RenderingFacade.GetMaterialCount()[MeshIndex];
-
-								TArray<UMaterialInterface*> Materials;
-								for (int32 MaterialIndex = MaterialStart; MaterialIndex < MaterialStart + MaterialCount; ++MaterialIndex)
-								{
-									const FString& Path = MaterialPaths[MaterialIndex];
-									UMaterialInterface* const Material = LoadObject<UMaterialInterface>(nullptr, *Path);
-									Materials.Add(Material);
-								}
-
-								AddDynamicMeshComponent({Target, MeshIndex}, MeshName, MoveTemp(DynamicMesh), Materials, MeshTransform);
-
-
-								for (int32 Idx = 0; Idx < VertexCount[MeshIndex]; ++Idx)
-								{
-									DebugMesh.VertexMap.Add(VertexIndex, VertexStart[MeshIndex] + Idx);
-									VertexIndex++;
-								}
-
-								for (int32 Idx = 0; Idx < FaceCount[MeshIndex]; ++Idx)
-								{
-									DebugMesh.FaceMap.Add(FaceIndex, FaceStart[MeshIndex] + Idx);
-									FaceIndex++;
-								}
-							}
-						}
-					}
-				}
+			MeshComponentsForWireframeRendering.Reset();
+			for (TObjectPtr<const UDataflowEdNode> Target : DataflowAsset->GetWireframeRenderTargets())
+			{
+				UpdateComponentsFromNode(Target, /*bWireframe*/true);
 			}
 
 			// If we have a single mesh component in the scene, select it
@@ -254,90 +198,14 @@ void FDataflowConstructionScene::UpdateDynamicMeshComponents()
 					DynamicMeshComponent->PushSelectionToProxy();
 				}
 			}
-
-			// Add hidden DynamicMeshComponents for any targets that we want to render in wireframe
-			// 
-			// Note: UMeshElementsVisualizers need source meshes to pull from. We add invisible dynamic mesh components to the existing DynamicMeshComponents collection
-			// for this purpose, but could have instead created a separate collection of meshes for wireframe rendering. We are choosing to keep all the scene DynamicMeshComponents 
-			// in one place and using separate structures to dictate how they are used (MeshComponentsForWireframeRendering in this case), in case visualization requirements 
-			// change in the future.
-			//
-
-			MeshComponentsForWireframeRendering.Reset();
-			for (TObjectPtr<const UDataflowEdNode> Target : DataflowAsset->GetWireframeRenderTargets())
-			{
-				if (Target)
-				{
-					TSharedPtr<FManagedArrayCollection> RenderCollection(new FManagedArrayCollection);
-					GeometryCollection::Facades::FRenderingFacade RenderingFacade(*RenderCollection);
-					RenderingFacade.DefineSchema();
-
-					const bool bHasRenderCollectionPrimitives = UE::Dataflow::RenderNodeOutput(RenderingFacade, *Target, *EditorContent, bEvaluateOutputs);
-
-					if (Target == EditorContent->GetSelectedNode())
-
-					{
-						EditorContent->SetRenderCollection(RenderCollection);
-					}
-
-					if(!bHasRenderCollectionPrimitives)
-					{
-						const int32 NumGeometry = RenderingFacade.NumGeometry();
-						for (int32 MeshIndex = 0; MeshIndex < NumGeometry; ++MeshIndex)
-						{
-							FDataflowRenderKey WireframeDynamicMeshKey{ Target, MeshIndex };
-
-							if (DynamicMeshComponents.Contains(WireframeDynamicMeshKey))
-							{
-								UDynamicMeshComponent* const ExistingMeshComponent = DynamicMeshComponents[WireframeDynamicMeshKey];
-								MeshComponentsForWireframeRendering.Add(ExistingMeshComponent);
-							}
-							else
-							{
-								FDynamicMesh3 DynamicMesh;
-								UE::Dataflow::Conversion::RenderingFacadeToDynamicMesh(RenderingFacade, MeshIndex, DynamicMesh);
-
-								if (DynamicMesh.VertexCount())
-								{
-									const FString MeshName = RenderingFacade.GetGeometryName()[MeshIndex];
-									const FTransform& MeshTransform = RenderingFacade.GetGeometryTransform()[MeshIndex];
-									const FString UniqueObjectName = MakeUniqueObjectName(RootSceneActor, UDataflowEditorCollectionComponent::StaticClass(), FName(MeshName)).ToString();
-									UDynamicMeshComponent* const NewDynamicMeshComponent = AddDynamicMeshComponent(WireframeDynamicMeshKey, UniqueObjectName, MoveTemp(DynamicMesh), {}, MeshTransform);
-									NewDynamicMeshComponent->SetVisibility(false);
-									MeshComponentsForWireframeRendering.Add(NewDynamicMeshComponent);
-								}
-							}
-						}
-					}
-				}
-			}
-
-			// Hide the floor in orthographic view modes
-			if (const UE::Dataflow::IDataflowConstructionViewMode* ConstructionViewMode = EditorContent->GetConstructionViewMode())
-			{
-				if (!ConstructionViewMode->IsPerspective())
-				{
-					constexpr bool bDontModifyProfile = true;
-					SetFloorVisibility(false, bDontModifyProfile);
-				}
-				else
-				{
-					// Restore visibility from profile settings
-					const int32 ProfileIndex = GetCurrentProfileIndex();
-					if (DefaultSettings->Profiles.IsValidIndex(ProfileIndex))
-					{
-						const bool bProfileSetting = DefaultSettings->Profiles[CurrentProfileIndex].bShowFloor;
-						constexpr bool bDontModifyProfile = true;
-						SetFloorVisibility(bProfileSetting, bDontModifyProfile);
-					}
-				}
-			}
 		}
 	
 		// Build a single mesh out of all the components
 		DebugMesh.Build(GetDynamicMeshComponents());
 	}
 
+	UpdateFloorVisibility();
+
 	bPreviewSceneDirty = true;
 }
 
@@ -454,7 +322,7 @@ TObjectPtr<UDynamicMeshComponent>& FDataflowConstructionScene::AddDynamicMeshCom
 		}
 	}
 
-	AddComponent(DynamicMeshComponent, InMeshTransform);	
+	AddComponent(DynamicMeshComponent, InMeshTransform);
 	DynamicMeshComponents.Emplace(InKey, DynamicMeshComponent);
 	AddSceneObject(DynamicMeshComponent, true);
 	return DynamicMeshComponents[InKey];
@@ -529,6 +397,33 @@ bool FDataflowConstructionScene::HasRenderableGeometry()
 	return false;
 }
 
+void FDataflowConstructionScene::UpdateFloorVisibility()
+{
+	// Hide the floor in orthographic view modes
+	if (const TObjectPtr<UDataflowBaseContent>& EditorContent = GetEditorContent())
+	{
+		if (const UE::Dataflow::IDataflowConstructionViewMode* ConstructionViewMode = EditorContent->GetConstructionViewMode())
+		{
+			if (!ConstructionViewMode->IsPerspective())
+			{
+				constexpr bool bDontModifyProfile = true;
+				SetFloorVisibility(false, bDontModifyProfile);
+			}
+			else
+			{
+				// Restore visibility from profile settings
+				const int32 ProfileIndex = GetCurrentProfileIndex();
+				if (DefaultSettings->Profiles.IsValidIndex(ProfileIndex))
+				{
+					const bool bProfileSetting = DefaultSettings->Profiles[CurrentProfileIndex].bShowFloor;
+					constexpr bool bDontModifyProfile = true;
+					SetFloorVisibility(bProfileSetting, bDontModifyProfile);
+				}
+			}
+		}
+	}
+}
+
 void FDataflowConstructionScene::ResetConstructionScene()
 {
 	// The ModeManagerss::USelection will hold references to Components, but 
@@ -628,6 +523,116 @@ void FDataflowConstructionScene::SelectNodeComponents(const UDataflowEdNode* EdN
 	}
 }
 
+void FDataflowConstructionScene::UpdateDebugMesh(const GeometryCollection::Facades::FRenderingFacade& RenderingFacade)
+{
+	DebugMesh.Reset();
+
+	int32 VertexIndex = 0;
+	int32 FaceIndex = 0;
+
+	const int32 NumGeometry = RenderingFacade.NumGeometry();
+	for (int32 MeshIndex = 0; MeshIndex < NumGeometry; ++MeshIndex)
+	{
+		// Accumulate in a unified debug mesh for visualization
+		// todo(dataflow) : we should get rid of this eventually as this is costly to build and maintain
+		const TManagedArray<int32>& VertexStart = RenderingFacade.GetVertexStart();
+		const TManagedArray<int32>& VertexCount = RenderingFacade.GetVertexCount();
+		for (int32 Idx = 0; Idx < VertexCount[MeshIndex]; ++Idx)
+		{
+			DebugMesh.VertexMap.Add(VertexIndex, VertexStart[MeshIndex] + Idx);
+			VertexIndex++;
+		}
+
+		const TManagedArray<int32>& FaceStart = RenderingFacade.GetIndicesStart();
+		const TManagedArray<int32>& FaceCount = RenderingFacade.GetIndicesCount();
+		for (int32 Idx = 0; Idx < FaceCount[MeshIndex]; ++Idx)
+		{
+			DebugMesh.FaceMap.Add(FaceIndex, FaceStart[MeshIndex] + Idx);
+			FaceIndex++;
+		}
+	}
+}
+
+
+void FDataflowConstructionScene::UpdateComponentsFromNode(const UDataflowEdNode* EdNode, bool bWireframe)
+{
+	if (EdNode == nullptr)
+	{
+		return;
+	}
+
+	using namespace UE::Geometry;//FDynamicMesh3
+
+	const FDataflowEditorToolkit* EditorToolkit = DataflowEditor ? static_cast<FDataflowEditorToolkit*>(DataflowEditor->GetInstanceInterface()) : nullptr;
+	const bool bEvaluateOutputs = EditorToolkit ? (EditorToolkit->GetEvaluationMode() != EDataflowEditorEvaluationMode::Manual) : true;
+
+	if (const TObjectPtr<UDataflowBaseContent>& EditorContent = GetEditorContent())
+	{
+		TSharedPtr<FManagedArrayCollection> RenderCollection(new FManagedArrayCollection);
+		GeometryCollection::Facades::FRenderingFacade RenderingFacade(*RenderCollection);
+		RenderingFacade.DefineSchema();
+
+		const bool bHasRenderCollectionPrimitives = UE::Dataflow::RenderNodeOutput(RenderingFacade, *EdNode, *EditorContent, bEvaluateOutputs);
+
+		if (EdNode == EditorContent->GetSelectedNode())
+		{
+			EditorContent->SetRenderCollection(RenderCollection);
+		}
+		if (!bHasRenderCollectionPrimitives)
+		{
+			const int32 NumGeometry = RenderingFacade.NumGeometry();
+			for (int32 MeshIndex = 0; MeshIndex < NumGeometry; ++MeshIndex)
+			{
+				FDataflowRenderKey DynamicMeshKey{ EdNode, MeshIndex };
+				if (bWireframe && DynamicMeshComponents.Contains(DynamicMeshKey))
+				{
+					UDynamicMeshComponent* const ExistingMeshComponent = DynamicMeshComponents[DynamicMeshKey];
+					MeshComponentsForWireframeRendering.Add(ExistingMeshComponent);
+				}
+				else
+				{
+					FDynamicMesh3 DynamicMesh;
+					UE::Dataflow::Conversion::RenderingFacadeToDynamicMesh(RenderingFacade, MeshIndex, DynamicMesh);
+
+					if (DynamicMesh.VertexCount())
+					{
+						const FString MeshName = RenderingFacade.GetGeometryName()[MeshIndex];
+						const FTransform& MeshTransform = RenderingFacade.GetGeometryTransform()[MeshIndex];
+						if (bWireframe)
+						{
+							// Add hidden DynamicMeshComponents for any targets that we want to render in wireframe
+							// 
+							// Note: UMeshElementsVisualizers need source meshes to pull from. We add invisible dynamic mesh components to the existing DynamicMeshComponents collection
+							// for this purpose, but could have instead created a separate collection of meshes for wireframe rendering. We are choosing to keep all the scene DynamicMeshComponents 
+							// in one place and using separate structures to dictate how they are used (MeshComponentsForWireframeRendering in this case), in case visualization requirements 
+							// change in the future.
+							//
+							const FString UniqueObjectName = MakeUniqueObjectName(RootSceneActor, UDataflowEditorCollectionComponent::StaticClass(), FName(MeshName)).ToString();
+							UDynamicMeshComponent* const NewDynamicMeshComponent = AddDynamicMeshComponent(DynamicMeshKey, UniqueObjectName, MoveTemp(DynamicMesh), {}, MeshTransform);
+							NewDynamicMeshComponent->SetVisibility(false);
+							MeshComponentsForWireframeRendering.Add(NewDynamicMeshComponent);
+						}
+						else
+						{
+							const TManagedArray<FString>& MaterialPaths = RenderingFacade.GetMaterialPaths();
+							const int32 MaterialStart = RenderingFacade.GetMaterialStart()[MeshIndex];
+							const int32 MaterialCount = RenderingFacade.GetMaterialCount()[MeshIndex];
+
+							TArray<UMaterialInterface*> Materials;
+							for (int32 MaterialIndex = MaterialStart; MaterialIndex < MaterialStart + MaterialCount; ++MaterialIndex)
+							{
+								const FString& Path = MaterialPaths[MaterialIndex];
+								UMaterialInterface* const Material = LoadObject<UMaterialInterface>(nullptr, *Path);
+								Materials.Add(Material);
+							}
+							AddDynamicMeshComponent(DynamicMeshKey, MeshName, MoveTemp(DynamicMesh), Materials, MeshTransform);
+						}
+					}
+				}
+			}
+		}
+	}
+}
 
 #undef LOCTEXT_NAMESPACE
 
diff --git a/Engine/Plugins/Experimental/Dataflow/Source/DataflowEditor/Public/Dataflow/DataflowConstructionScene.h b/Engine/Plugins/Experimental/Dataflow/Source/DataflowEditor/Public/Dataflow/DataflowConstructionScene.h
--- a/Engine/Plugins/Experimental/Dataflow/Source/DataflowEditor/Public/Dataflow/DataflowConstructionScene.h
+++ b/Engine/Plugins/Experimental/Dataflow/Source/DataflowEditor/Public/Dataflow/DataflowConstructionScene.h
@@ -97,13 +97,21 @@ private:
 	/** Build the primitive components from the dataflow graph selected node */
     UE_API void UpdatePrimitiveComponents();
 	
+	/** Build the primitive components from the dataflow graph selected node */
+	void UpdateComponentsFromNode(const UDataflowEdNode* Node, bool bWireframe);
+
 	/** Add a dynamic mesh component to the scene */
 	UE_API TObjectPtr<UDynamicMeshComponent>& AddDynamicMeshComponent(
 		FDataflowRenderKey InKey, const FString& MeshName, UE::Geometry::FDynamicMesh3&& DynamicMesh,  const TArray<UMaterialInterface*>& MaterialSet, const FTransform& InMeshTransform = FTransform::Identity);
 
 	/** Remove the scene component from selection and destroy it */
 	UE_API void RemoveSceneComponent(USelection* SelectedComponents, UPrimitiveComponent* PrimitiveComponent);
-	
+
+	void UpdateDebugMesh(const GeometryCollection::Facades::FRenderingFacade& RenderingFacade);
+
+	/** update floor visibility based on the view type */
+	void UpdateFloorVisibility();
+
 	/** Dynamic mesh components generated by the tools for edit/manipulation*/
 	typedef TPair<FDataflowRenderKey, TObjectPtr<UDynamicMeshComponent>> FRenderElement;
 	TMap<FDataflowRenderKey, TObjectPtr<UDynamicMeshComponent>> DynamicMeshComponents;
