Dataflow editor : Add vertex attribute paint tool
- This is meant to replace the existing weightmap tool once this is at parity featurewise
- It is behind a cvar at the moment and disabled by default ( p.DataflowEditor.UseNewWeightMapTool )
- The UX/UI is aligned with the skinweight tool to provide a familiar experience
- there will be more iterations happening on the code in the coming days/weeks

diff --git a/Engine/Plugins/Experimental/Dataflow/Dataflow.uplugin b/Engine/Plugins/Experimental/Dataflow/Dataflow.uplugin
--- a/Engine/Plugins/Experimental/Dataflow/Dataflow.uplugin
+++ b/Engine/Plugins/Experimental/Dataflow/Dataflow.uplugin
@@ -55,6 +55,11 @@
 			"Enabled": true,
 			"TargetAllowList": [ "Editor" ]
 		},
+		{
+			"Name": "ModelingToolsEditorMode",
+			"Enabled": true,
+			"TargetAllowList": [ "Editor" ]
+		},
 		{
 			"Name": "ChaosCaching",
 			"Enabled": true
diff --git a/Engine/Plugins/Experimental/Dataflow/Source/DataflowEditor/DataflowEditor.Build.cs b/Engine/Plugins/Experimental/Dataflow/Source/DataflowEditor/DataflowEditor.Build.cs
--- a/Engine/Plugins/Experimental/Dataflow/Source/DataflowEditor/DataflowEditor.Build.cs
+++ b/Engine/Plugins/Experimental/Dataflow/Source/DataflowEditor/DataflowEditor.Build.cs
@@ -63,7 +63,7 @@ namespace UnrealBuildTool.Rules
 					"MeshModelingToolsExp",
 					"ModelingComponentsEditorOnly",
 					"ModelingComponents",
-//					"ModelingToolsEditorMode",
+					"ModelingToolsEditorMode",
 					"PlanarCut",
 					"Projects",
 					"PropertyEditor",
diff --git a/Engine/Plugins/Experimental/Dataflow/Source/DataflowEditor/Private/Dataflow/DataflowEditorMode.cpp b/Engine/Plugins/Experimental/Dataflow/Source/DataflowEditor/Private/Dataflow/DataflowEditorMode.cpp
--- a/Engine/Plugins/Experimental/Dataflow/Source/DataflowEditor/Private/Dataflow/DataflowEditorMode.cpp
+++ b/Engine/Plugins/Experimental/Dataflow/Source/DataflowEditor/Private/Dataflow/DataflowEditorMode.cpp
@@ -8,6 +8,7 @@
 #include "Components/DynamicMeshComponent.h"
 #include "ContextObjectStore.h"
 #include "DataflowEditorTools/DataflowEditorWeightMapPaintTool.h"
+#include "DataflowEditorTools/DataflowEditorVertexAttributePaintTool.h"
 #include "Dataflow/DataflowComponentToolTarget.h"
 #include "Dataflow/DataflowCollectionAddScalarVertexPropertyNode.h"
 #include "Dataflow/DataflowConstructionViewportClient.h"
@@ -310,6 +311,10 @@ void UDataflowEditorMode::RegisterTools()
 			{
 				WeightMapPaintToolBuilder->SetEditorMode(this);
 			}
+			else if (UDataflowEditorVertexAttributePaintToolBuilder* VertexAttributePaintToolBuilder = Cast<UDataflowEditorVertexAttributePaintToolBuilder>(Builder))
+			{
+				VertexAttributePaintToolBuilder->SetEditorMode(this);
+			}
 
 			RegisterDataflowTool(CommandInfo, RegisteredNodeName.ToString() + FString(TEXT("Tool")), Builder, ConstructionViewportToolsContext);
 
diff --git a/Engine/Plugins/Experimental/Dataflow/Source/DataflowEditor/Private/Dataflow/DataflowEditorModule.cpp b/Engine/Plugins/Experimental/Dataflow/Source/DataflowEditor/Private/Dataflow/DataflowEditorModule.cpp
--- a/Engine/Plugins/Experimental/Dataflow/Source/DataflowEditor/Private/Dataflow/DataflowEditorModule.cpp
+++ b/Engine/Plugins/Experimental/Dataflow/Source/DataflowEditor/Private/Dataflow/DataflowEditorModule.cpp
@@ -19,7 +19,9 @@
 #include "Dataflow/DataflowAnyTypeCustomization.h"
 #include "Dataflow/DataflowToolRegistry.h"
 #include "DataflowEditorTools/DataflowEditorWeightMapPaintTool.h"
+#include "DataflowEditorTools/DataflowEditorVertexAttributePaintTool.h"
 #include "Dataflow/DataflowCollectionAddScalarVertexPropertyNode.h"
+#include "DataflowEditorTools/DataflowEditorVertexAttributePaintToolDetailCustomization.h"
 
 #include "Dataflow/DataflowInstanceDetails.h"
 
@@ -126,6 +128,7 @@ void FDataflowEditorModule::StartupModule()
 		PropertyModule->RegisterCustomPropertyTypeLayout(DataflowFreezeActionsName, FOnGetPropertyTypeCustomizationInstance::CreateStatic(&UE::Dataflow::FFreezeActionsCustomization::MakeInstance));
 
 		PropertyModule->RegisterCustomPropertyTypeLayout(DataflowColorRampName, FOnGetPropertyTypeCustomizationInstance::CreateStatic(&UE::Dataflow::FDataflowColorRampCustomization::MakeInstance));
+		PropertyModule->RegisterCustomClassLayout(UDataflowEditorVertexAttributePaintToolProperties::StaticClass()->GetFName(), FOnGetDetailCustomizationInstance::CreateStatic(&UE::Dataflow::Editor::FVertexAttributePaintToolDetailCustomization::MakeInstance));
 
 		UE::Dataflow::RegisterAnyTypeCustomizations(*PropertyModule);
 	}
@@ -136,7 +139,7 @@ void FDataflowEditorModule::StartupModule()
 	TSharedRef<const FDataflowToolActionCommandBindings> Actions = MakeShared<FDataflowToolActionCommandBindings>();
 	
 	ToolRegistry.AddNodeToToolMapping(FDataflowCollectionAddScalarVertexPropertyNode::StaticType(),
-		NewObject<UDataflowEditorWeightMapPaintToolBuilder>(), Actions,
+		NewObject<UDataflowEditorVertexAttributePaintToolBuilder>(), Actions,
 		FSlateIcon(FName("DataflowEditorStyle"), FName("Dataflow.PaintWeightMap")), LOCTEXT("AddWeightMapNodeButtonText", "Paint Weight Map"));
 
 	ToolRegistry.AddNodeToToolMapping(FDataflowCollectionEditSkinWeightsNode::StaticType(), 
@@ -151,7 +154,7 @@ void FDataflowEditorModule::StartupModule()
 	ConstructionVisualizationRegistry.RegisterVisualization(MakeUnique<UE::Dataflow::FMeshStatsConstructionVisualization>());
 	ConstructionVisualizationRegistry.RegisterVisualization(MakeUnique<UE::Dataflow::FMeshConstructionVisualization>());
 	ConstructionVisualizationRegistry.RegisterVisualization(MakeUnique<UE::Dataflow::FDebugDrawText3dVisualization>());
-
+	
 	FDataflowEditorCommands::Register();
 }
 
diff --git a/Engine/Plugins/Experimental/Dataflow/Source/DataflowEditor/Private/DataflowEditorTools/DataflowEditorVertexAttributePaintTool.cpp b/Engine/Plugins/Experimental/Dataflow/Source/DataflowEditor/Private/DataflowEditorTools/DataflowEditorVertexAttributePaintTool.cpp
new file mode 100644
--- /dev/null
+++ b/Engine/Plugins/Experimental/Dataflow/Source/DataflowEditor/Private/DataflowEditorTools/DataflowEditorVertexAttributePaintTool.cpp
@@ -0,0 +1,2210 @@
+// Copyright Epic Games, Inc. All Rights Reserved.
+
+#include "DataflowEditorTools/DataflowEditorVertexAttributePaintTool.h"
+
+#include "BaseBehaviors/TwoAxisPropertyEditBehavior.h"
+#include "BaseGizmos/BrushStampIndicator.h"
+#include "CanvasTypes.h"
+#include "CanvasItem.h"
+#include "ContextObjectStore.h"
+#include "Dataflow/DataflowCollectionAddScalarVertexPropertyNode.h"
+#include "Dataflow/DataflowConstructionScene.h"
+#include "Dataflow/DataflowContent.h"
+#include "Dataflow/DataflowEditorCollectionComponent.h"
+#include "Dataflow/DataflowEditorMode.h"
+#include "Dataflow/DataflowNode.h"
+#include "DataflowEditorTools/DataflowEditorWeightMapPaintBrushOps.h"
+#include "DataflowEditorTools/DataflowEditorWeightMapPaintTool.h"
+#include "DynamicMesh/DynamicMeshAttributeSet.h"
+#include "DynamicMesh/DynamicMeshChangeTracker.h"
+#include "Drawing/MeshElementsVisualizer.h"
+#include "DynamicMesh/MeshIndexUtil.h"
+#include "DynamicMesh/NonManifoldMappingSupport.h"
+#include "DynamicMeshToMeshDescription.h"
+#include "Engine/Engine.h"
+#include "InteractiveToolManager.h"
+#include "Intersection/IntrLine2Line2.h"
+#include "Intersection/IntrSegment2Segment2.h"
+#include "Framework/Notifications/NotificationManager.h"
+#include "HAL/PlatformApplicationMisc.h"
+#include "Logging/LogMacros.h"
+#include "ModelingToolTargetUtil.h"
+#include "Polygon2.h"
+#include "PreviewMesh.h"
+#include "StaticMeshAttributes.h"
+#include "Selections/MeshConnectedComponents.h"
+#include "Selections/MeshVertexSelection.h"
+#include "ToolSetupUtil.h"
+#include "Util/BufferUtil.h"
+#include "Widgets/Notifications/SNotificationList.h"
+
+#include UE_INLINE_GENERATED_CPP_BY_NAME(DataflowEditorVertexAttributePaintTool)
+
+DEFINE_LOG_CATEGORY_STATIC(LogDataflowEditorVertexAttributePaintTool, Warning, All);
+
+#define LOCTEXT_NAMESPACE "DataflowEditorVertexAttributePaintTool"
+
+///////////////////////////////////////////////////////////////////////////////////////////////////////////////////////////////////////////////////
+
+
+namespace UE::DataflowEditorVertexAttributePaintTool::CVars
+{
+	bool DataflowEditorUseNewWeightMapTool = false;
+	FAutoConsoleVariableRef CVarDataflowEditorUseNewWeightMapTool(
+		TEXT("p.DataflowEditor.UseNewWeightMapTool"),
+		DataflowEditorUseNewWeightMapTool,
+		TEXT("when on, the new weight map tool with improved UX will be used"));
+}
+
+namespace UE::DataflowEditorVertexAttributePaintTool::Private
+{
+	// probably should be something defined for the whole tool framework...
+#if WITH_EDITOR
+	static EAsyncExecution WeightPaintToolAsyncExecTarget = EAsyncExecution::LargeThreadPool;
+#else
+	static EAsyncExecution WeightPaintToolAsyncExecTarget = EAsyncExecution::ThreadPool;
+#endif
+
+	static bool HasManagedArrayCollection(const FDataflowNode* InDataflowNode, const TSharedPtr<UE::Dataflow::FEngineContext> Context)
+	{
+		static const FName CollectionTypeName = "FManagedArrayCollection";
+		if (InDataflowNode && Context)
+		{
+			for (const FDataflowOutput* const Output : InDataflowNode->GetOutputs())
+			{
+				if (Output->GetType() == CollectionTypeName)
+				{
+					return true;
+				}
+			}
+		}
+		return false;
+	}
+
+	static void ShowEditorMessage(ELogVerbosity::Type InMessageType, const FText& InMessage)
+	{
+		FNotificationInfo Notification(InMessage);
+		Notification.bUseSuccessFailIcons = true;
+		Notification.ExpireDuration = 5.0f;
+
+		SNotificationItem::ECompletionState State = SNotificationItem::CS_Success;
+
+		switch (InMessageType)
+		{
+		case ELogVerbosity::Warning:
+			UE_LOG(LogDataflowEditorVertexAttributePaintTool, Warning, TEXT("%s"), *InMessage.ToString());
+			break;
+		case ELogVerbosity::Error:
+			State = SNotificationItem::CS_Fail;
+			UE_LOG(LogDataflowEditorVertexAttributePaintTool, Error, TEXT("%s"), *InMessage.ToString());
+			break;
+		default:
+			break; // don't log anything unless a warning or error
+		}
+
+		FSlateNotificationManager::Get().AddNotification(Notification)->SetCompletionState(State);
+	}
+
+	/**
+		* A wrapper change that applies a given change to the unwrap canonical mesh of an input, and uses that
+		* to update the other views. Causes a broadcast of OnCanonicalModified.
+		*/
+	class  FMeshChange : public FToolCommandChange
+	{
+	public:
+		FMeshChange(UDynamicMeshComponent* DynamicMeshComponentIn, TUniquePtr<UE::Geometry::FDynamicMeshChange> DynamicMeshChangeIn)
+			: DynamicMeshComponent(DynamicMeshComponentIn)
+			, DynamicMeshChange(MoveTemp(DynamicMeshChangeIn))
+		{
+			ensure(DynamicMeshComponentIn);
+			ensure(DynamicMeshChange);
+		};
+
+		virtual void Apply(UObject* Object) override
+		{
+			DynamicMeshChange->Apply(DynamicMeshComponent->GetMesh(), false);
+		}
+
+		virtual void Revert(UObject* Object) override
+		{
+			DynamicMeshChange->Apply(DynamicMeshComponent->GetMesh(), true);
+		}
+
+		virtual bool HasExpired(UObject* Object) const override
+		{
+			return !(DynamicMeshComponent.IsValid() && DynamicMeshChange);
+		}
+
+		virtual FString ToString() const override
+		{
+			return TEXT("DataflowEditorVertexAttributePaintToolMeshChange");
+		}
+
+	protected:
+		TWeakObjectPtr<UDynamicMeshComponent> DynamicMeshComponent;
+		TUniquePtr<UE::Geometry::FDynamicMeshChange> DynamicMeshChange;
+	};
+}
+
+///////////////////////////////////////////////////////////////////////////////////////////////////////////////////////////////////////////////////
+//
+// ToolBuilder
+//
+///////////////////////////////////////////////////////////////////////////////////////////////////////////////////////////////////////////////////
+
+void UDataflowEditorVertexAttributePaintToolBuilder::GetSupportedConstructionViewModes(const UDataflowContextObject& ContextObject, TArray<const UE::Dataflow::IDataflowConstructionViewMode*>& Modes) const
+{
+	//Modes.Add(UE::Dataflow::EDataflowPatternVertexType::Sim3D);
+	//Modes.Add(UE::Dataflow::EDataflowPatternVertexType::Sim2D);
+}
+
+bool UDataflowEditorVertexAttributePaintToolBuilder::CanBuildTool(const FToolBuilderState& SceneState) const
+{
+	if (UMeshSurfacePointMeshEditingToolBuilder::CanBuildTool(SceneState))
+	{
+		if (SceneState.SelectedComponents.Num() == 1)
+		{
+			if (TObjectPtr<UDataflowEditorCollectionComponent> Component = Cast<UDataflowEditorCollectionComponent>(SceneState.SelectedComponents[0]))
+			{
+				if (UDataflowBaseContent* ContextObject = SceneState.ToolManager->GetContextObjectStore()->FindContext<UDataflowBaseContent>())
+				{
+					if (ContextObject->GetSelectedNode() == Component->Node)
+					{
+						if (const TSharedPtr<UE::Dataflow::FEngineContext> EvaluationContext = ContextObject->GetDataflowContext())
+						{
+							if (const FDataflowNode* PrimarySelection = ContextObject->GetSelectedNodeOfType<FDataflowCollectionAddScalarVertexPropertyNode>())
+							{
+								return UE::DataflowEditorVertexAttributePaintTool::Private::HasManagedArrayCollection(PrimarySelection, EvaluationContext);
+							}
+						}
+					}
+				}
+			}
+		}
+	}
+	return false;
+}
+
+UMeshSurfacePointTool* UDataflowEditorVertexAttributePaintToolBuilder::CreateNewTool(const FToolBuilderState& SceneState) const
+{
+	UDataflowContextObject* ContextObject = SceneState.ToolManager->GetContextObjectStore()->FindContext<UDataflowContextObject>();
+
+	if (UE::DataflowEditorVertexAttributePaintTool::CVars::DataflowEditorUseNewWeightMapTool)
+	{
+		UDataflowEditorVertexAttributePaintTool* PaintTool = NewObject<UDataflowEditorVertexAttributePaintTool>(SceneState.ToolManager);
+		PaintTool->SetEditorMode(Mode);
+		PaintTool->SetWorld(SceneState.World);
+		if (ContextObject)
+		{
+			PaintTool->SetDataflowEditorContextObject(ContextObject);
+		}
+		return PaintTool;
+	}
+
+	// fallback : use the old tool
+	// todo(ccaillaud) : remove this fall back when the new tool is on par with the old one 
+	UDataflowEditorWeightMapPaintTool* PaintTool = NewObject<UDataflowEditorWeightMapPaintTool>(SceneState.ToolManager);
+	PaintTool->SetEditorMode(Mode);
+	PaintTool->SetWorld(SceneState.World);
+	if (ContextObject)
+	{
+		PaintTool->SetDataflowEditorContextObject(ContextObject);
+	}
+	return PaintTool;
+}
+
+///////////////////////////////////////////////////////////////////////////////////////////////////////////////////////////////////////////////////
+//
+// Data
+//
+///////////////////////////////////////////////////////////////////////////////////////////////////////////////////////////////////////////////////
+void FDataflowEditorVertexAttributePaintToolData::Setup(FDynamicMesh3& InMesh, FDataflowCollectionAddScalarVertexPropertyNode* InNodeToUpdate, TObjectPtr<UDataflowContextObject> InDataflowEditorContextObject)
+{
+	// this should only be called once or have shutdown called in between 
+	ensure(ActiveWeightMap == nullptr);
+	NodeToUpdate = InNodeToUpdate;
+
+	if (InDataflowEditorContextObject)
+	{
+		// Generate mapping array from and to non-manifiold and manifold mesh 
+		if (const TSharedPtr<const FManagedArrayCollection> Collection = InDataflowEditorContextObject->GetRenderCollection())
+		{
+			using namespace UE::Dataflow;
+			const UE::Geometry::FNonManifoldMappingSupport NonManifoldMapping(InMesh);
+
+			bHaveDynamicMeshToWeightConversion = NonManifoldMapping.IsNonManifoldVertexInSource();
+			if (bHaveDynamicMeshToWeightConversion)
+			{
+				DynamicMeshToWeight.SetNumUninitialized(InMesh.VertexCount());
+				WeightToDynamicMesh.Reset();
+				WeightToDynamicMesh.SetNum(Collection->NumElements("Vertices")); // todo(ccaillaud) Do we realy need the render collection for that ?
+				for (int32 DynamicMeshVert = 0; DynamicMeshVert < InMesh.VertexCount(); ++DynamicMeshVert)
+				{
+					DynamicMeshToWeight[DynamicMeshVert] = NonManifoldMapping.GetOriginalNonManifoldVertexID(DynamicMeshVert);
+					if (0 <= DynamicMeshToWeight[DynamicMeshVert] && DynamicMeshToWeight[DynamicMeshVert] < WeightToDynamicMesh.Num())
+					{
+						WeightToDynamicMesh[DynamicMeshToWeight[DynamicMeshVert]].Add(DynamicMeshVert);
+					}
+					else
+					{
+						bHaveDynamicMeshToWeightConversion = false;
+						UE_LOG(LogTemp, Warning, TEXT("Weight map misalignment."));
+						break;
+					}
+				}
+			}
+		}
+
+		// Create an attribute layer to temporarily paint into
+		const int NumAttributeLayers = InMesh.Attributes()->NumWeightLayers();
+		InMesh.Attributes()->SetNumWeightLayers(NumAttributeLayers + 1);
+		ActiveWeightMap = InMesh.Attributes()->GetWeightLayer(NumAttributeLayers);
+		ActiveWeightMap->SetName(FName("PaintLayer"));
+
+		if (NodeToUpdate)
+		{
+			// Setup DynamicMeshToWeight conversion and get Input weight map (if it exists)
+			const int32 NumExpectedWeights = bHaveDynamicMeshToWeightConversion ? WeightToDynamicMesh.Num() : InMesh.MaxVertexID();
+			SetupWeights.SetNumZeroed(NumExpectedWeights);
+
+			// Get the setup attribute values
+			if (TSharedPtr<UE::Dataflow::FEngineContext> DataflowContext = InDataflowEditorContextObject->GetDataflowContext())
+			{
+				const FCollectionAttributeKey AttributeKey = NodeToUpdate->GetWeightAttributeKey(*DataflowContext);
+				NodeToUpdate->FillAttributeWeights(InDataflowEditorContextObject->GetSelectedCollection(), AttributeKey, SetupWeights);
+			}
+
+			// Copy attribute  from selected node to the preview mesh
+			TArray<float> CurrentWeights;
+			CurrentWeights.SetNumZeroed(NumExpectedWeights);
+			NodeToUpdate->ExtractVertexWeights(SetupWeights, TArrayView<float>(CurrentWeights));
+
+			if (bHaveDynamicMeshToWeightConversion)
+			{
+				if (WeightToDynamicMesh.Num() == CurrentWeights.Num())	// Only copy node weights if they match the number of mesh vertices
+				{
+					for (int32 WeightID = 0; WeightID < CurrentWeights.Num(); ++WeightID)
+					{
+						for (const int32 VertexID : WeightToDynamicMesh[WeightID])
+						{
+							ActiveWeightMap->SetValue(VertexID, &CurrentWeights[WeightID]);
+						}
+					}
+				}
+			}
+			else
+			{
+				if (InMesh.MaxVertexID() == CurrentWeights.Num())	// Only copy node weights if they match the number of mesh vertices
+				{
+					for (int32 VertexID = 0; VertexID < CurrentWeights.Num(); ++VertexID)
+					{
+						ActiveWeightMap->SetValue(VertexID, &CurrentWeights[VertexID]);
+					}
+				}
+			}
+		}
+	}
+}
+
+void FDataflowEditorVertexAttributePaintToolData::Shutdown()
+{
+	if (NodeToUpdate)
+	{
+		NodeToUpdate->Invalidate();
+	}
+}
+
+void FDataflowEditorVertexAttributePaintToolData::BeginChange(UDynamicMeshComponent* Component)
+{
+	check(ActiveWeightEditChangeTracker == nullptr);
+	if (Component)
+	{
+		ActiveWeightEditChangeTracker = MakeUnique<UE::Geometry::FDynamicMeshChangeTracker>(Component->GetMesh());
+		ActiveWeightEditChangeTracker->BeginChange();
+	}
+}
+
+TUniquePtr<UE::DataflowEditorVertexAttributePaintTool::Private::FMeshChange> FDataflowEditorVertexAttributePaintToolData::EndChange(UDynamicMeshComponent* Component)
+{
+	check(ActiveWeightEditChangeTracker);
+
+	if (ActiveWeightEditChangeTracker && Component)
+	{
+		using namespace UE::DataflowEditorVertexAttributePaintTool;
+		TUniquePtr<UE::Geometry::FDynamicMeshChange> EditResult = ActiveWeightEditChangeTracker->EndChange();
+		TUniquePtr<Private::FMeshChange> MeshChange = MakeUnique<Private::FMeshChange>(Component, MoveTemp(EditResult));
+		ActiveWeightEditChangeTracker = nullptr;
+
+		return MeshChange;
+	}
+	return {};
+}
+
+float FDataflowEditorVertexAttributePaintToolData::GetValue(int32 VertexIdx) const
+{
+	float Value = 0.f;
+	if (ActiveWeightMap && VertexIdx != IndexConstants::InvalidID)
+	{
+		ActiveWeightMap->GetValue(VertexIdx, &Value);
+	}
+	return Value;
+}
+
+float FDataflowEditorVertexAttributePaintToolData::GetAverageValue(const TArray<int32>& Vertices) const
+{
+	double AverageValue = 0.f;
+	if (ActiveWeightMap)
+	{
+		int32 NumValidValues = 0;
+		for (int32 VertexIdx : Vertices)
+		{
+			if (VertexIdx != IndexConstants::InvalidID)
+			{
+				float Value = 0.0f;
+				ActiveWeightMap->GetValue(VertexIdx, &Value);
+				AverageValue += (double)Value;
+				++NumValidValues;
+			}
+		}
+		if (NumValidValues > 0)
+		{
+			AverageValue = AverageValue / (double)NumValidValues;
+		}
+	}
+	return (float)AverageValue;
+}
+
+void FDataflowEditorVertexAttributePaintToolData::UpdateNode()
+{
+	if (ActiveWeightMap && NodeToUpdate)
+	{
+		if (const FDynamicMesh3* Mesh = ActiveWeightMap->GetParent())
+		{
+			TArray<float> WeightsToApply;
+
+			const int32 NumVertices = Mesh->VertexCount();
+			WeightsToApply.SetNumUninitialized(NumVertices);
+			for (int32 VertexID = 0; VertexID < NumVertices; ++VertexID)
+			{
+				ActiveWeightMap->GetValue(VertexID, &WeightsToApply[VertexID]);
+			}
+
+			if (bHaveDynamicMeshToWeightConversion)
+			{
+				NodeToUpdate->ReportVertexWeights(SetupWeights, WeightsToApply, DynamicMeshToWeight);
+			}
+			else
+			{
+				NodeToUpdate->ReportVertexWeights(SetupWeights, WeightsToApply, {});
+			}
+
+			NodeToUpdate->Invalidate();
+		}
+	}
+}
+
+///////////////////////////////////////////////////////////////////////////////////////////////////////////////////////////////////////////////////
+//
+// Tool
+//
+///////////////////////////////////////////////////////////////////////////////////////////////////////////////////////////////////////////////////
+
+UBaseDynamicMeshComponent* UDataflowEditorVertexAttributePaintTool::GetSculptMeshComponent()
+{
+	return PreviewMesh ? Cast<UBaseDynamicMeshComponent>(PreviewMesh->GetRootComponent()) : nullptr;
+}
+
+void UDataflowEditorVertexAttributePaintTool::SetupBrushEditBehaviorSetup(ULocalTwoAxisPropertyEditInputBehavior& OutBehavior)
+{
+	// map the horizontal behavior to change the brush size
+	UMeshSculptToolBase::MapHorizontalBrushEditBehaviorToBrushSize(OutBehavior);
+
+	// map the vertical ehavior to change the attribiute value
+	OutBehavior.VerticalProperty.GetValueFunc = [this]()
+		{
+			return ToolProperties->BrushProperties.AttributeValue;
+		};
+
+	OutBehavior.VerticalProperty.SetValueFunc = [this](float NewValue)
+		{
+			ToolProperties->BrushProperties.AttributeValue = FMath::Min(1.0f, FMath::Max(NewValue, 0.f));
+#if WITH_EDITOR
+			FPropertyChangedEvent PropertyChangedEvent(FDataflowEditorVertexAttributePaintToolBrushProperties::StaticStruct()->FindPropertyByName(GET_MEMBER_NAME_CHECKED(FDataflowEditorVertexAttributePaintToolBrushProperties, AttributeValue)));
+			ToolProperties->PostEditChangeProperty(PropertyChangedEvent);
+#endif
+		};
+	OutBehavior.VerticalProperty.Name = LOCTEXT("AttributeValue", "Attribute Value");
+	OutBehavior.VerticalProperty.EditRate = 0.005f;
+	OutBehavior.VerticalProperty.bEnabled = true;
+}
+
+void UDataflowEditorVertexAttributePaintTool::Setup()
+{
+	UMeshSculptToolBase::Setup();
+
+	// todo(ccaillaud) : this should be a customized by the inherited class 
+	// Get the selected node to read and write the vertex attribute
+	FDataflowCollectionAddScalarVertexPropertyNode* NodeToUpdate = DataflowEditorContextObject->GetSelectedNodeOfType<FDataflowCollectionAddScalarVertexPropertyNode>();
+	checkf(NodeToUpdate, TEXT("No Node is currently selected, or more than one node is selected"));
+
+	// todo(ccaillaud) : this should be a parameter of the tool ? 
+	SetToolDisplayName(LOCTEXT("ToolName", "Paint Weight Maps"));
+
+	// Hide all meshes in the DataflowConstructionScene, as we will be painting onto our own Preview mesh
+	if (const TSharedPtr<FDataflowConstructionScene> Scene = Mode->GetDataflowConstructionScene().Pin())
+	{
+		Scene->SetVisibility(false);
+	}
+
+	//OnDynamicMeshComponentChangedHandle = DynamicMeshComponent->OnMeshVerticesChanged.AddUObject(this, &UDataflowEditorVertexAttributePaintTool::OnDynamicMeshComponentChanged);
+
+	// Create the preview mesh 
+	PreviewMesh = NewObject<UPreviewMesh>(this);
+	PreviewMesh->CreateInWorld(TargetWorld, FTransform::Identity);
+	PreviewMesh->SetTangentsMode(EDynamicMeshComponentTangentsMode::AutoCalculated);
+	ToolSetupUtil::ApplyRenderingConfigurationToPreview(PreviewMesh, Target);
+
+	// make sure we get a copy of the original dynamic mesh 
+	FGetMeshParameters GetMeshParams;
+	GetMeshParams.bWantMeshTangents = true;
+	FDynamicMesh3 DynamicMeshCopy = UE::ToolTarget::GetDynamicMeshCopy(Target, GetMeshParams);
+	PreviewMesh->UpdatePreview(MoveTemp(DynamicMeshCopy));
+
+	// Assign materials
+	FComponentMaterialSet MaterialSet = UE::ToolTarget::GetMaterialSet(Target);
+	PreviewMesh->SetMaterials(MaterialSet.Materials);
+
+	// make sure the dynamic mesh has all the necessary attributes
+	FDynamicMesh3* Mesh = GetSculptMesh();
+	Mesh->EnableVertexColors(FVector3f::One());
+	Mesh->Attributes()->EnablePrimaryColors();
+	Mesh->Attributes()->PrimaryColors()->CreateFromPredicate([](int ParentVID, int TriIDA, int TriIDB) { return true; }, 0.f);
+	UE::Geometry::FAxisAlignedBox3d Bounds = Mesh->GetBounds(true);
+
+	TFuture<void> PrecomputeFuture = Async(UE::DataflowEditorVertexAttributePaintTool::Private::WeightPaintToolAsyncExecTarget, [&]()
+		{
+			PrecomputeFilterData();
+		});
+
+	TFuture<void> OctreeFuture = Async(UE::DataflowEditorVertexAttributePaintTool::Private::WeightPaintToolAsyncExecTarget, [&]()
+		{
+			// initialize dynamic octree
+			if (Mesh->TriangleCount() > 100000)
+			{
+				Octree.RootDimension = Bounds.MaxDim() / 10.0;
+				Octree.SetMaxTreeDepth(4);
+			}
+			else
+			{
+				Octree.RootDimension = Bounds.MaxDim();
+				Octree.SetMaxTreeDepth(8);
+			}
+			Octree.Initialize(Mesh);
+		});
+
+	// Setup mesh selector
+	MeshSelector = NewObject<UToolMeshSelector>(this);
+	auto OnSelectionChangedLambda = [this]() { /*OnSelectionModified();*/ };
+	MeshSelector->InitialSetup(TargetWorld, this, OnSelectionChangedLambda);
+
+	// compute the mesh description for the mesh selector to use
+	{
+		using namespace UE::Geometry;
+		FStaticMeshAttributes Attributes(PreviewMeshDescription);
+		Attributes.Register();
+		FConversionToMeshDescriptionOptions ConverterOptions;
+		FDynamicMeshToMeshDescription Converter(ConverterOptions);
+		Converter.Convert(Mesh, PreviewMeshDescription);
+		MeshSelector->SetMesh(PreviewMesh, &PreviewMeshDescription, FTransform::Identity);
+	}
+
+	// initialize brush radius range interval, brush properties
+	UMeshSculptToolBase::InitializeBrushSizeRange(Bounds);
+
+	// initialize properties
+	ToolProperties = NewObject<UDataflowEditorVertexAttributePaintToolProperties>(this);
+	{
+		ToolProperties->WatchProperty(ToolProperties->BrushProperties.BrushSize,
+			[this](float NewSize) { UMeshSculptToolBase::BrushProperties->BrushSize.AdaptiveSize = NewSize; });
+		ToolProperties->BrushProperties.BrushSize = UMeshSculptToolBase::BrushProperties->BrushSize.AdaptiveSize;
+		CreateDefaultColorRamp(ToolProperties->DisplayProperties.ColorRamp);
+		ToolProperties->DisplayProperties.ColorRamp.ColorRamp.OnColorCurveChangedDelegate.AddUObject(this, &UDataflowEditorVertexAttributePaintTool::OnColorRampChanged);
+		ToolProperties->RestoreProperties(this);
+	}
+	AddToolPropertySource(ToolProperties);
+
+	// make sure the selector is in sync with the settings
+	MeshSelector->SetComponentSelectionMode(ToolProperties->SelectionProperties.ComponentSelectionMode);
+
+	InitializeIndicator();
+
+	// initialize our properties
+	AddToolPropertySource(UMeshSculptToolBase::BrushProperties);
+	SetToolPropertySourceEnabled(UMeshSculptToolBase::BrushProperties, false);
+	UMeshSculptToolBase::BrushProperties->bShowPerBrushProps = false;
+	UMeshSculptToolBase::BrushProperties->bShowFalloff = true;
+	UMeshSculptToolBase::BrushProperties->bShowLazyness = false;
+	UMeshSculptToolBase::BrushProperties->FlowRate = 0.0f;
+	UMeshSculptToolBase::BrushProperties->GetOnModified().AddLambda(
+		[this](UObject* InBrushProperties, FProperty* Property)
+		{
+			// make sure we update the settings back from updating the brush using B+Mouse drag
+			if (InBrushProperties == UMeshSculptToolBase::BrushProperties)
+			{
+				if (Property->GetFName() == GET_MEMBER_NAME_CHECKED(FDataflowEditorVertexAttributePaintToolBrushProperties, BrushSize))
+				{
+					ToolProperties->BrushProperties.BrushSize = UMeshSculptToolBase::BrushProperties->BrushSize.AdaptiveSize;
+				}
+			}
+		});
+	CalculateBrushRadius();
+
+	PaintBrushOpProperties = NewObject<UDataflowVertexAttributePaintBrushOpProps>(this);
+	RegisterBrushType(PaintBrushId, LOCTEXT("Paint", "Paint"),
+		MakeUnique<FLambdaMeshSculptBrushOpFactory>([this]() { return MakeUnique<FDataflowVertexAttributePaintBrushOp>(); }),
+		PaintBrushOpProperties);
+
+	// secondary brushes
+	EraseBrushOpProperties = NewObject<UDataflowWeightMapEraseBrushOpProps>(this);
+	RegisterSecondaryBrushType(EraseBrushId, LOCTEXT("Erase", "Erase"),
+		MakeUnique<TBasicMeshSculptBrushOpFactory<FDataflowWeightMapEraseBrushOp>>(),
+		EraseBrushOpProperties);
+
+	AddToolPropertySource(UMeshSculptToolBase::ViewProperties);
+	SetToolPropertySourceEnabled(UMeshSculptToolBase::ViewProperties, false);
+
+	AddToolPropertySource(UMeshSculptToolBase::GizmoProperties);
+	SetToolPropertySourceEnabled(UMeshSculptToolBase::GizmoProperties, false);
+
+	// must call before updating brush type so that we register all brush properties?
+	UMeshSculptToolBase::OnCompleteSetup();
+
+	UpdateBrushType(ToolProperties->BrushProperties.BrushMode);
+
+	// mesh element display is used to show seams and other features ofthe mesh
+	MeshElementsDisplay = NewObject<UMeshElementsVisualizer>(this);
+	MeshElementsDisplay->CreateInWorld(TargetWorld, FTransform::Identity);
+	if (ensure(MeshElementsDisplay->Settings))
+	{
+		MeshElementsDisplay->Settings->bShowNormalSeams = false;
+		MeshElementsDisplay->Settings->RestoreProperties(this, TEXT("DataflowEditorVertexAttributePaintTool"));
+		AddToolPropertySource(MeshElementsDisplay->Settings);
+	}
+	MeshElementsDisplay->SetMeshAccessFunction([this](UMeshElementsVisualizer::ProcessDynamicMeshFunc ProcessFunc) {
+		ProcessFunc(*GetSculptMesh());
+		});
+
+	// disable view properties
+	SetViewPropertiesEnabled(false);
+	UpdateMaterialMode(EMeshEditingMaterialModes::VertexColor);
+	UpdateWireframeVisibility(false);
+	UpdateFlatShadingSetting(false);
+
+	VertexData.Setup(*Mesh, NodeToUpdate, DataflowEditorContextObject);
+	
+	PrecomputeFuture.Wait();
+	OctreeFuture.Wait();
+
+	// update colors
+	UpdateVertexColorOverlay();
+	//DynamicMeshComponent->FastNotifyVertexAttributesUpdated(EMeshRenderAttributeFlags::VertexColors);
+	GetToolManager()->PostInvalidation();
+}
+
+void UDataflowEditorVertexAttributePaintTool::CreateDefaultColorRamp(FDataflowColorRamp& OutColorRamp)
+{
+	constexpr FLinearColor HeatMapColors[] =
+	{
+		FLinearColor(0.8f, 0.4f, 0.8f), // Purple
+		FLinearColor(0.0f, 0.0f, 0.5f), // Dark Blue
+		FLinearColor(0.2f, 0.2f, 1.0f), // Light Blue
+		FLinearColor(0.0f, 1.0f, 0.0f), // Green
+		FLinearColor(1.0f, 1.0f, 0.0f), // Yellow
+		FLinearColor(1.0f, 0.65f, 0.0f), // Orange
+		FLinearColor(1.0f, 0.0f, 0.0f, 0.0f), // Red
+	};
+	constexpr int32 NumColors = sizeof(HeatMapColors) / sizeof(FLinearColor);
+	ensure(NumColors > 1);
+
+	// Set ColorRamp
+	FDataflowColorCurveOwner& Ramp = OutColorRamp.ColorRamp;
+	for (int32 Index = 0; Index < NumColors; ++Index)
+	{
+		const float Time = (float)Index / ((float)NumColors - 1.f);
+		Ramp.GetCurves()[0].CurveToEdit->AddKey(Time, HeatMapColors[Index].R);
+		Ramp.GetCurves()[1].CurveToEdit->AddKey(Time, HeatMapColors[Index].G);
+		Ramp.GetCurves()[2].CurveToEdit->AddKey(Time, HeatMapColors[Index].B);
+	}
+	Ramp.GetCurves()[3].CurveToEdit->AddKey(0.f, 1.f); // single alpha  value
+}
+
+void UDataflowEditorVertexAttributePaintTool::SetDataflowEditorContextObject(TObjectPtr<UDataflowContextObject> InDataflowEditorContextObject)
+{
+	DataflowEditorContextObject = InDataflowEditorContextObject;
+}
+
+void UDataflowEditorVertexAttributePaintTool::Shutdown(EToolShutdownType ShutdownType)
+{
+	//if (DynamicMeshComponent != nullptr)
+	//{
+	//	DynamicMeshComponent->OnMeshChanged.Remove(OnDynamicMeshComponentChangedHandle);
+	//}
+
+	if (MeshElementsDisplay)
+	{
+		if (ensure(MeshElementsDisplay->Settings))
+		{
+			MeshElementsDisplay->Settings->SaveProperties(this, TEXT("DataflowEditorWeightMapPaintTool2"));
+		}
+		MeshElementsDisplay->Disconnect();
+	}
+
+	if (ToolProperties)
+	{
+		ToolProperties->SaveProperties(this);
+	}
+
+	if (MeshSelector)
+	{
+		MeshSelector->Shutdown();
+		MeshSelector = nullptr;
+	}
+
+	// we need to call this directly now because we use the preview mesh instead of the built-in component
+	if (ShutdownType == EToolShutdownType::Accept)
+	{
+		CommitResult(GetSculptMeshComponent(), false);
+	}
+
+	VertexData.Shutdown();
+
+	if (PreviewMesh)
+	{
+		PreviewMesh->SetVisible(false);
+		PreviewMesh->Disconnect();
+		PreviewMesh = nullptr;
+	}
+
+	UMeshSculptToolBase::Shutdown(ShutdownType);
+}
+
+
+void UDataflowEditorVertexAttributePaintTool::CommitResult(UBaseDynamicMeshComponent* Component, bool bModifiedTopology)
+{
+	GetToolManager()->BeginUndoTransaction(LOCTEXT("DataflowEditorVertexAttributePaintTool_TransactionName", "Paint Weights"));
+
+	VertexData.UpdateNode();
+
+	GetToolManager()->EndUndoTransaction();
+}
+
+
+void UDataflowEditorVertexAttributePaintTool::RegisterActions(FInteractiveToolActionSet& ActionSet)
+{
+	UMeshSculptToolBase::RegisterActions(ActionSet);
+
+	ActionSet.RegisterAction(this, (int32)EStandardToolActions::BaseClientDefinedActionID + 500,
+		TEXT("PickWeightValueUnderCursor"),
+		LOCTEXT("PickWeightValueUnderCursor", "Pick Weight Value"),
+		LOCTEXT("PickWeightValueUnderCursorTooltip", "Set the active weight painting value to that currently under the cursor"),
+		EModifierKey::Shift, EKeys::G,
+		[this]() { bPendingPickWeight = true; });
+};
+
+
+TUniquePtr<FMeshSculptBrushOp>& UDataflowEditorVertexAttributePaintTool::GetActiveBrushOp()
+{
+	if (GetInEraseStroke())
+	{
+		return SecondaryBrushOp;
+	}
+	else
+	{
+		return PrimaryBrushOp;
+	}
+}
+
+
+void UDataflowEditorVertexAttributePaintTool::OnPropertyModified(UObject* PropertySet, FProperty* Property)
+{
+	CalculateBrushRadius();
+}
+
+
+void UDataflowEditorVertexAttributePaintTool::IncreaseBrushRadiusAction()
+{
+	Super::IncreaseBrushRadiusAction();
+	ToolProperties->BrushProperties.BrushSize = BrushProperties->BrushSize.AdaptiveSize;
+	NotifyOfPropertyChangeByTool(ToolProperties);
+}
+
+void UDataflowEditorVertexAttributePaintTool::DecreaseBrushRadiusAction()
+{
+	Super::DecreaseBrushRadiusAction();
+	ToolProperties->BrushProperties.BrushSize = BrushProperties->BrushSize.AdaptiveSize;
+	NotifyOfPropertyChangeByTool(ToolProperties);
+}
+
+void UDataflowEditorVertexAttributePaintTool::IncreaseBrushRadiusSmallStepAction()
+{
+	Super::IncreaseBrushRadiusSmallStepAction();
+	ToolProperties->BrushProperties.BrushSize = BrushProperties->BrushSize.AdaptiveSize;
+	NotifyOfPropertyChangeByTool(ToolProperties);
+}
+
+void UDataflowEditorVertexAttributePaintTool::DecreaseBrushRadiusSmallStepAction()
+{
+	Super::DecreaseBrushRadiusSmallStepAction();
+	ToolProperties->BrushProperties.BrushSize = BrushProperties->BrushSize.AdaptiveSize;
+	NotifyOfPropertyChangeByTool(ToolProperties);
+}
+
+float UDataflowEditorVertexAttributePaintTool::GetBrushMinRadius() const
+{
+	return UMeshSculptToolBase::BrushProperties->BrushSize.WorldSizeRange.Min * 0.5f;
+}
+
+float UDataflowEditorVertexAttributePaintTool::GetBrushMaxRadius() const
+{
+	return UMeshSculptToolBase::BrushProperties->BrushSize.WorldSizeRange.Max * 0.5f;
+}
+
+bool UDataflowEditorVertexAttributePaintTool::IsInBrushSubMode() const
+{
+	return (ToolProperties->EditingMode == EDataflowEditorToolEditMode::Brush);
+}
+
+bool UDataflowEditorVertexAttributePaintTool::IsVolumetricBrush() const
+{
+	return (ToolProperties && ToolProperties->BrushProperties.BrushAreaMode == EMeshVertexPaintBrushAreaType::Volumetric);
+}
+
+void UDataflowEditorVertexAttributePaintTool::OnBeginStroke(const FRay& WorldRay)
+{
+	if (!VertexData.IsValid())
+	{
+		return;
+	}
+
+	UpdateBrushPosition(WorldRay);
+
+	if (PaintBrushOpProperties)
+	{
+		PaintBrushOpProperties->AttributeValue = ToolProperties->BrushProperties.AttributeValue;
+		PaintBrushOpProperties->Strength = ToolProperties->BrushProperties.Strength * ToolProperties->BrushProperties.Strength;
+		PaintBrushOpProperties->EditOperation = ToolProperties->BrushProperties.BrushMode;
+	}
+	if (EraseBrushOpProperties)
+	{
+		EraseBrushOpProperties->AttributeValue = 0.0;
+	}
+
+	// initialize first "Last Stamp", so that we can assume all stamps in stroke have a valid previous stamp
+	LastStamp.WorldFrame = GetBrushFrameWorld();
+	LastStamp.LocalFrame = GetBrushFrameLocal();
+	LastStamp.Radius = GetCurrentBrushRadius();
+	LastStamp.Falloff = GetCurrentBrushFalloff();
+	LastStamp.Direction = GetInInvertStroke() ? -1.0 : 1.0;
+	LastStamp.Depth = GetCurrentBrushDepth();
+	LastStamp.Power = GetActivePressure() * GetCurrentBrushStrength();
+	LastStamp.TimeStamp = FDateTime::Now();
+
+	FSculptBrushOptions SculptOptions;
+	SculptOptions.ConstantReferencePlane = GetCurrentStrokeReferencePlane();
+
+	TUniquePtr<FMeshSculptBrushOp>& UseBrushOp = GetActiveBrushOp();
+	UseBrushOp->ConfigureOptions(SculptOptions);
+	UseBrushOp->BeginStroke(GetSculptMesh(), LastStamp, VertexROI);
+
+	AccumulatedTriangleROI.Reset();
+
+	// begin change here? or wait for first stamp?
+	BeginChange();
+}
+
+void UDataflowEditorVertexAttributePaintTool::OnEndStroke()
+{
+	if (!VertexData.IsValid())
+	{
+		return;
+	}
+
+	GetActiveBrushOp()->EndStroke(GetSculptMesh(), LastStamp, VertexROI);
+
+	UpdateVertexColorOverlay(&TriangleROI);
+	//DynamicMeshComponent->FastNotifyVertexAttributesUpdated(EMeshRenderAttributeFlags::VertexColors);
+
+	// close change record
+	EndChange();
+}
+
+void UDataflowEditorVertexAttributePaintTool::UpdateROI(const FSculptBrushStamp& BrushStamp)
+{
+	SCOPE_CYCLE_COUNTER(VertexAttributePaintTool_UpdateROI);
+
+	const FVector3d& BrushPos = BrushStamp.LocalFrame.Origin;
+	const FDynamicMesh3* Mesh = GetSculptMesh();
+	const float RadiusSqr = GetCurrentBrushRadius() * GetCurrentBrushRadius();
+	UE::Geometry::FAxisAlignedBox3d BrushBox(
+		BrushPos - GetCurrentBrushRadius() * FVector3d::One(),
+		BrushPos + GetCurrentBrushRadius() * FVector3d::One());
+
+	TriangleROI.Reset();
+
+	int32 CenterTID = GetBrushTriangleID();
+	if (Mesh->IsTriangle(CenterTID))
+	{
+		TriangleROI.Add(CenterTID);
+	}
+
+	FVector3d CenterNormal = Mesh->IsTriangle(CenterTID) ? TriNormals[CenterTID] : FVector3d::One();		// One so that normal check always passes
+
+	const bool bVolumetric = IsVolumetricBrush();
+	const bool bUseAngleThreshold = ToolProperties->BrushProperties.AngleThreshold < 180.0f;
+	const double DotAngleThreshold = FMathd::Cos(ToolProperties->BrushProperties.AngleThreshold * FMathd::DegToRad);
+	const bool bStopAtUVSeams = ToolProperties->BrushProperties.bUVSeams;
+	const bool bStopAtNormalSeams = ToolProperties->BrushProperties.bNormalSeams;
+
+	auto CheckEdgeCriteria = [&](int32 t1, int32 t2) -> bool
+		{
+			if (bUseAngleThreshold == false || CenterNormal.Dot(TriNormals[t2]) > DotAngleThreshold)
+			{
+				int32 eid = Mesh->FindEdgeFromTriPair(t1, t2);
+				if (bStopAtUVSeams == false || UVSeamEdges[eid] == false)
+				{
+					if (bStopAtNormalSeams == false || NormalSeamEdges[eid] == false)
+					{
+						return true;
+					}
+				}
+			}
+			return false;
+		};
+
+	if (bVolumetric)
+	{
+		Octree.RangeQuery(BrushBox,
+			[&](int TriIdx) {
+				if ((Mesh->GetTriCentroid(TriIdx) - BrushPos).SquaredLength() < RadiusSqr)
+				{
+					TriangleROI.Add(TriIdx);
+				}
+			});
+	}
+	else
+	{
+		if (Mesh->IsTriangle(CenterTID))
+		{
+			TArray<int32> StartROI;
+			StartROI.Add(CenterTID);
+			UE::Geometry::FMeshConnectedComponents::GrowToConnectedTriangles(Mesh, StartROI, TriangleROI, &TempROIBuffer,
+				[&](int t1, int t2)
+				{
+					if ((Mesh->GetTriCentroid(t2) - BrushPos).SquaredLength() < RadiusSqr)
+					{
+						return CheckEdgeCriteria(t1, t2);
+					}
+					return false;
+				});
+		}
+	}
+
+	// Construct ROI vertex set
+	VertexSetBuffer.Reset();
+	for (int32 tid : TriangleROI)
+	{
+		UE::Geometry::FIndex3i Tri = Mesh->GetTriangle(tid);
+		VertexSetBuffer.Add(Tri.A);  VertexSetBuffer.Add(Tri.B);  VertexSetBuffer.Add(Tri.C);
+	}
+
+	// Apply visibility filter
+	if (ToolProperties->BrushProperties.VisibilityFilter != EDataflowEditorToolVisibilityType::None)
+	{
+		TArray<int32> ResultBuffer;
+		ApplyVisibilityFilter(VertexSetBuffer, TempROIBuffer, ResultBuffer);
+	}
+
+	VertexROI.SetNum(0, EAllowShrinking::No);
+	//TODO: If we paint a 2D projection of UVs, these will need to be the 2D vertices not the 3D original mesh vertices
+	BufferUtil::AppendElements(VertexROI, VertexSetBuffer);
+
+	// Construct ROI triangle and weight buffers
+	ROITriangleBuffer.Reserve(TriangleROI.Num());
+	ROITriangleBuffer.SetNum(0, EAllowShrinking::No);
+	for (int32 tid : TriangleROI)
+	{
+		ROITriangleBuffer.Add(tid);
+	}
+	ROIWeightValueBuffer.SetNum(VertexROI.Num(), EAllowShrinking::No);
+	SyncWeightBufferWithMesh(Mesh);
+}
+
+bool UDataflowEditorVertexAttributePaintTool::UpdateStampPosition(const FRay& WorldRay)
+{
+	CalculateBrushRadius();
+
+	TUniquePtr<FMeshSculptBrushOp>& UseBrushOp = GetActiveBrushOp();
+
+	ESculptBrushOpTargetType TargetType = UseBrushOp->GetBrushTargetType();
+	switch (TargetType)
+	{
+	case ESculptBrushOpTargetType::SculptMesh:
+	case ESculptBrushOpTargetType::TargetMesh:
+		UpdateBrushPositionOnSculptMesh(WorldRay, true);
+		break;
+	case ESculptBrushOpTargetType::ActivePlane:
+		check(false);
+		UpdateBrushPositionOnActivePlane(WorldRay);
+		break;
+	}
+
+	if (UseBrushOp->GetAlignStampToView() || IsVolumetricBrush())
+	{
+		AlignBrushToView();
+	}
+
+	CurrentStamp = LastStamp;
+	CurrentStamp.DeltaTime = FMathd::Min((FDateTime::Now() - LastStamp.TimeStamp).GetTotalSeconds(), 1.0);
+	CurrentStamp.WorldFrame = GetBrushFrameWorld();
+	CurrentStamp.LocalFrame = GetBrushFrameLocal();
+	CurrentStamp.Power = GetActivePressure() * GetCurrentBrushStrength();
+
+	CurrentStamp.PrevLocalFrame = LastStamp.LocalFrame;
+	CurrentStamp.PrevWorldFrame = LastStamp.WorldFrame;
+
+	FVector3d MoveDelta = CurrentStamp.LocalFrame.Origin - CurrentStamp.PrevLocalFrame.Origin;
+
+	if (UseBrushOp->IgnoreZeroMovements() && MoveDelta.SquaredLength() < 0.1 * CurrentBrushRadius)
+	{
+		return false;
+	}
+
+	return true;
+}
+
+bool UDataflowEditorVertexAttributePaintTool::ApplyStamp()
+{
+	SCOPE_CYCLE_COUNTER(VertexAttributePaintTool_ApplyStamp);
+
+	TUniquePtr<FMeshSculptBrushOp>& UseBrushOp = GetActiveBrushOp();
+
+	// yuck
+	FDataflowMeshVertexWeightMapEditBrushOp* WeightBrushOp = (FDataflowMeshVertexWeightMapEditBrushOp*)UseBrushOp.Get();
+	WeightBrushOp->bApplyRadiusLimit = IsInBrushSubMode();
+
+	FDynamicMesh3* Mesh = GetSculptMesh();
+	WeightBrushOp->ApplyStampByVertices(Mesh, CurrentStamp, VertexROI, ROIWeightValueBuffer);
+
+	bool bUpdated = SyncMeshWithWeightBuffer(Mesh);
+
+	LastStamp = CurrentStamp;
+	LastStamp.TimeStamp = FDateTime::Now();
+
+	return bUpdated;
+}
+
+bool UDataflowEditorVertexAttributePaintTool::SyncMeshWithWeightBuffer(FDynamicMesh3* Mesh)
+{
+	int32 NumModified = 0;
+	if (VertexData.IsValid())
+	{
+		// change update could be async here if we collected array of <idx,orig,new> and dispatched independenlty
+		const int32 NumVertexROI = VertexROI.Num();
+		for (int32 VertexROIIndex = 0; VertexROIIndex < NumVertexROI; ++VertexROIIndex)
+		{
+			const int VertIdx = VertexROI[VertexROIIndex];
+			const float NewValue = ROIWeightValueBuffer[VertexROIIndex];
+			VertexData.SetValue(VertIdx, 
+				[NewValue](UE::Geometry::FDynamicMeshWeightAttribute& AttributeMap, int32 AttributeIndex)
+				{
+					return NewValue;
+				});
+			NumModified++;
+		}
+	}
+	return (NumModified > 0);
+}
+
+bool UDataflowEditorVertexAttributePaintTool::SyncWeightBufferWithMesh(const FDynamicMesh3* Mesh)
+{
+	int32 NumModified = 0;
+	if (VertexData.IsValid())
+	{
+		// change update could be async here if we collected array of <idx,orig,new> and dispatched independenlty
+		const int32 NumVertexROI = VertexROI.Num();
+		for (int32 VertexROIIndex = 0; VertexROIIndex < NumVertexROI; ++VertexROIIndex)
+		{
+			const int VertIdx = VertexROI[VertexROIIndex];
+			const float CurWeight = VertexData.GetValue(VertIdx);
+			if (ROIWeightValueBuffer[VertexROIIndex] != CurWeight)
+			{
+				ROIWeightValueBuffer[VertexROIIndex] = CurWeight;
+				NumModified++;
+			}
+		}
+	}
+	return (NumModified > 0);
+}
+
+// TODO(ccaillaud) : move this to a common  place
+namespace UE::DataflowEditorVertexAttributePaintTool::Private
+{
+	template<typename RealType>
+	static bool FindPolylineSelfIntersection(
+		const TArray<UE::Math::TVector2<RealType>>& Polyline,
+		UE::Math::TVector2<RealType>& IntersectionPointOut,
+		UE::Geometry::FIndex2i& IntersectionIndexOut,
+		bool bParallel = true)
+	{
+		int32 N = Polyline.Num();
+		std::atomic<bool> bSelfIntersects(false);
+		ParallelFor(N - 1, [&](int32 i)
+			{
+				UE::Geometry::TSegment2<RealType> SegA(Polyline[i], Polyline[i + 1]);
+				for (int32 j = i + 2; j < N - 1 && bSelfIntersects == false; ++j)
+				{
+					UE::Geometry::TSegment2<RealType> SegB(Polyline[j], Polyline[j + 1]);
+					if (SegA.Intersects(SegB) && bSelfIntersects == false)
+					{
+						bool ExpectedValue = false;
+						if (std::atomic_compare_exchange_strong(&bSelfIntersects, &ExpectedValue, true))
+						{
+							UE::Geometry::TIntrSegment2Segment2<RealType> Intersection(SegA, SegB);
+							Intersection.Find();
+							IntersectionPointOut = Intersection.Point0;
+							IntersectionIndexOut = UE::Geometry::FIndex2i(i, j);
+							return;
+						}
+					}
+				}
+			}, (bParallel) ? EParallelForFlags::None : EParallelForFlags::ForceSingleThread);
+
+		return bSelfIntersects;
+	}
+
+
+
+	template<typename RealType>
+	static bool FindPolylineSegmentIntersection(
+		const TArray<UE::Math::TVector2<RealType>>& Polyline,
+		const UE::Geometry::TSegment2<RealType>& Segment,
+		UE::Math::TVector2<RealType>& IntersectionPointOut,
+		int& IntersectionIndexOut)
+	{
+
+		int32 N = Polyline.Num();
+		for (int32 i = 0; i < N - 1; ++i)
+		{
+			UE::Geometry::TSegment2<RealType> PolySeg(Polyline[i], Polyline[i + 1]);
+			if (Segment.Intersects(PolySeg))
+			{
+				UE::Geometry::TIntrSegment2Segment2<RealType> Intersection(Segment, PolySeg);
+				Intersection.Find();
+				IntersectionPointOut = Intersection.Point0;
+				IntersectionIndexOut = i;
+				return true;
+			}
+		}
+		return false;
+	}
+
+
+	bool ApproxSelfClipPolyline(TArray<FVector2f>& Polyline)
+	{
+		int32 N = Polyline.Num();
+
+		// handle already-closed polylines
+		if (UE::Geometry::Distance(Polyline[0], Polyline[N - 1]) < 0.0001f)
+		{
+			return true;
+		}
+
+		FVector2f IntersectPoint;
+		UE::Geometry::FIndex2i IntersectionIndex(-1, -1);
+		bool bSelfIntersects = FindPolylineSelfIntersection(Polyline, IntersectPoint, IntersectionIndex);
+		if (bSelfIntersects)
+		{
+			TArray<FVector2f> NewPolyline;
+			NewPolyline.Add(IntersectPoint);
+			for (int32 i = IntersectionIndex.A; i <= IntersectionIndex.B; ++i)
+			{
+				NewPolyline.Add(Polyline[i]);
+			}
+			NewPolyline.Add(IntersectPoint);
+			Polyline = MoveTemp(NewPolyline);
+			return true;
+		}
+
+
+		FVector2f StartDirOut = UE::Geometry::Normalized(Polyline[0] - Polyline[1]);
+		UE::Geometry::FLine2f StartLine(Polyline[0], StartDirOut);
+		FVector2f EndDirOut = UE::Geometry::Normalized(Polyline[N - 1] - Polyline[N - 2]);
+		UE::Geometry::FLine2f EndLine(Polyline[N - 1], EndDirOut);
+		UE::Geometry::FIntrLine2Line2f LineIntr(StartLine, EndLine);
+		bool bIntersects = false;
+		if (LineIntr.Find())
+		{
+			bIntersects = LineIntr.IsSimpleIntersection() && (LineIntr.Segment1Parameter > 0) && (LineIntr.Segment2Parameter > 0);
+			if (bIntersects)
+			{
+				Polyline.Add(StartLine.PointAt(LineIntr.Segment1Parameter));
+				Polyline.Add(StartLine.Origin);
+				return true;
+			}
+		}
+
+
+		UE::Geometry::FAxisAlignedBox2f Bounds;
+		for (const FVector2f& P : Polyline)
+		{
+			Bounds.Contain(P);
+		}
+		float Size = Bounds.DiagonalLength();
+
+		FVector2f StartPos = Polyline[0] + 0.001f * StartDirOut;
+		if (FindPolylineSegmentIntersection(Polyline, UE::Geometry::FSegment2f(StartPos, StartPos + 2 * Size * StartDirOut), IntersectPoint, IntersectionIndex.A))
+		{
+			return true;
+		}
+
+		FVector2f EndPos = Polyline[N - 1] + 0.001f * EndDirOut;
+		if (FindPolylineSegmentIntersection(Polyline, UE::Geometry::FSegment2f(EndPos, EndPos + 2 * Size * EndDirOut), IntersectPoint, IntersectionIndex.A))
+		{
+			return true;
+		}
+
+		return false;
+	}
+}
+
+
+void UDataflowEditorVertexAttributePaintTool::OnPolyLassoFinished(const FCameraPolyLasso& Lasso, bool bCanceled)
+{
+	// construct polyline
+	TArray<FVector2f> Polyline;
+	for (FVector2D Pos : Lasso.Polyline)
+	{
+		Polyline.Add((FVector2f)Pos);
+	}
+	int32 N = Polyline.Num();
+	if (N < 2)
+	{
+		return;
+	}
+
+	// Try to clip polyline to be closed, or closed-enough for winding evaluation to work.
+	// If that returns false, the polyline is "too open". In that case we will extend
+	// outwards from the endpoints and then try to create a closed very large polygon
+	if (UE::DataflowEditorVertexAttributePaintTool::Private::ApproxSelfClipPolyline(Polyline) == false)
+	{
+		FVector2f StartDirOut = UE::Geometry::Normalized(Polyline[0] - Polyline[1]);
+		UE::Geometry::FLine2f StartLine(Polyline[0], StartDirOut);
+		FVector2f EndDirOut = UE::Geometry::Normalized(Polyline[N - 1] - Polyline[N - 2]);
+		UE::Geometry::FLine2f EndLine(Polyline[N - 1], EndDirOut);
+
+		// if we did not intersect, we are in ambiguous territory. Check if a segment along either end-direction
+		// intersects the polyline. If it does, we have something like a spiral and will be OK. 
+		// If not, make a closed polygon by interpolating outwards from each endpoint, and then in perp-directions.
+		UE::Geometry::FPolygon2f Polygon(Polyline);
+		float PerpSign = Polygon.IsClockwise() ? -1.0 : 1.0;
+
+		Polyline.Insert(StartLine.PointAt(10000.0f), 0);
+		Polyline.Insert(Polyline[0] + 1000 * PerpSign * UE::Geometry::PerpCW(StartDirOut), 0);
+
+		Polyline.Add(EndLine.PointAt(10000.0f));
+		Polyline.Add(Polyline.Last() + 1000 * PerpSign * UE::Geometry::PerpCW(EndDirOut));
+		FVector2f StartPos = Polyline[0];
+		Polyline.Add(StartPos);		// close polyline (cannot use Polyline[0] in case Add resizes!)
+	}
+
+	N = Polyline.Num();
+
+	// project each mesh vertex to view plane and evaluate winding integral of polyline
+	const FDynamicMesh3* Mesh = GetSculptMesh();
+	TempROIBuffer.SetNum(Mesh->MaxVertexID());
+	ParallelFor(Mesh->MaxVertexID(), [&](int32 VertexIdx)
+		{
+			if (Mesh->IsVertex(VertexIdx))
+			{
+				FVector3d WorldPos = CurTargetTransform.TransformPosition(Mesh->GetVertex(VertexIdx));
+				FVector2f PlanePos = (FVector2f)Lasso.GetProjectedPoint((FVector)WorldPos);
+
+				double WindingSum = 0;
+				FVector2f a = Polyline[0] - PlanePos, b = FVector2f::Zero();
+				for (int32 i = 1; i < N; ++i)
+				{
+					b = Polyline[i] - PlanePos;
+					WindingSum += (double)FMathf::Atan2(a.X * b.Y - a.Y * b.X, a.X * b.X + a.Y * b.Y);
+					a = b;
+				}
+				WindingSum /= FMathd::TwoPi;
+				bool bInside = FMathd::Abs(WindingSum) > 0.3;
+				TempROIBuffer[VertexIdx] = bInside ? 1 : 0;
+			}
+			else
+			{
+				TempROIBuffer[VertexIdx] = -1;
+			}
+		});
+
+	// convert to vertex selection, and then select fully-enclosed faces
+	UE::Geometry::FMeshVertexSelection VertexSelection(Mesh);
+	VertexSelection.SelectByVertexID([&](int32 VertexIdx) { return TempROIBuffer[VertexIdx] == 1; });
+
+	const double SetWeightValue = GetInEraseStroke() ? 0.0 : ToolProperties->BrushProperties.AttributeValue;
+	SetVerticesToWeightMap(VertexSelection.AsSet(), SetWeightValue, GetInEraseStroke());
+}
+
+
+void UDataflowEditorVertexAttributePaintTool::ComputeGradient()
+{
+	//if (!ensure(ActiveWeightMap))
+	//{
+	//	UE_LOG(LogTemp, Warning, TEXT("No active weight map"));
+	//	return;
+	//}
+
+	//BeginChange();
+
+	//const FDynamicMesh3* const Mesh = DynamicMeshComponent->GetMesh();
+	//TempROIBuffer.SetNum(0, EAllowShrinking::No);
+	//for (int32 VertexIdx : Mesh->VertexIndicesItr())
+	//{
+	//	TempROIBuffer.Add(VertexIdx);
+	//}
+
+	//if (bHaveDynamicMeshToWeightConversion)
+	//{
+	//	for (int32 VertexIdx : TempROIBuffer)
+	//	{
+	//		for (const int32 Idx : WeightToDynamicMesh[DynamicMeshToWeight[VertexIdx]])
+	//		{
+	//			ActiveWeightEditChangeTracker->SaveVertexOneRingTriangles(Idx, true);
+	//		}
+	//	}
+	//}
+	//else
+	//{
+	//	for (int32 VertexIdx : TempROIBuffer)
+	//	{
+	//		ActiveWeightEditChangeTracker->SaveVertexOneRingTriangles(VertexIdx, true);
+	//	}
+	//}
+
+
+	//for (const int32 VertexIndex : TempROIBuffer)
+	//{
+	//	const FVector3d Vert = Mesh->GetVertex(VertexIndex);
+
+	//	// (Copied from FClothPaintTool_Gradient::ApplyGradient)
+
+	//	// Get distances
+	//	// TODO: Look into surface distance instead of 3D distance? May be necessary for some complex shapes
+	//	float DistanceToLowSq = MAX_flt;
+	//	for (const int32& LowIndex : LowValueGradientVertexSelection.SelectedCornerIDs)
+	//	{
+	//		const FVector3d LowPoint = Mesh->GetVertex(LowIndex);
+	//		const float DistanceSq = (LowPoint - Vert).SizeSquared();
+	//		if (DistanceSq < DistanceToLowSq)
+	//		{
+	//			DistanceToLowSq = DistanceSq;
+	//		}
+	//	}
+
+	//	float DistanceToHighSq = MAX_flt;
+	//	for (const int32& HighIndex : HighValueGradientVertexSelection.SelectedCornerIDs)
+	//	{
+	//		const FVector3d HighPoint = Mesh->GetVertex(HighIndex);
+	//		const float DistanceSq = (HighPoint - Vert).SizeSquared();
+	//		if (DistanceSq < DistanceToHighSq)
+	//		{
+	//			DistanceToHighSq = DistanceSq;
+	//		}
+	//	}
+
+	//	const float Value = FMath::LerpStable(FilterProperties->GradientLowValue, FilterProperties->GradientHighValue, DistanceToLowSq / (DistanceToLowSq + DistanceToHighSq));
+	//	if (bHaveDynamicMeshToWeightConversion)
+	//	{
+	//		for (const int32 Idx : WeightToDynamicMesh[DynamicMeshToWeight[VertexIndex]])
+	//		{
+	//			ActiveWeightMap->SetValue(Idx, &Value);
+	//		}
+	//	}
+	//	else
+	//	{
+	//		ActiveWeightMap->SetValue(VertexIndex, &Value);
+	//	}
+	//}
+
+	//// update colors
+	//UpdateVertexColorOverlay();
+	//DynamicMeshComponent->FastNotifyVertexAttributesUpdated(EMeshRenderAttributeFlags::VertexColors);
+	//GetToolManager()->PostInvalidation();
+	//EndChange();
+
+}
+
+bool UDataflowEditorVertexAttributePaintTool::HasSelection() const
+{
+	if (MeshSelector)
+	{
+		return MeshSelector->IsAnyComponentSelected();
+	}
+	return false;
+}
+
+void UDataflowEditorVertexAttributePaintTool::GrowSelection() const
+{
+	if (MeshSelector)
+	{
+		MeshSelector->GrowSelection();
+	}
+}
+
+void UDataflowEditorVertexAttributePaintTool::ShrinkSelection() const
+{
+	if (MeshSelector)
+	{
+		MeshSelector->ShrinkSelection();
+	}
+}
+
+void UDataflowEditorVertexAttributePaintTool::FloodSelection() const
+{
+	if (MeshSelector)
+	{
+		MeshSelector->FloodSelection();
+	}
+}
+
+void UDataflowEditorVertexAttributePaintTool::SelectBorder() const
+{
+	if (MeshSelector)
+	{
+		MeshSelector->SelectBorder();
+	}
+}
+
+
+
+void UDataflowEditorVertexAttributePaintTool::SetComponentSelectionMode(EComponentSelectionMode NewMode)
+{
+	if (ToolProperties)
+	{
+		ToolProperties->SelectionProperties.ComponentSelectionMode = NewMode;
+		if (MeshSelector)
+		{
+			MeshSelector->SetComponentSelectionMode(NewMode);
+		}
+	}
+}
+
+TArray<int32> UDataflowEditorVertexAttributePaintTool::GetSelectedVertices() const
+{
+	if (MeshSelector)
+	{
+		return MeshSelector->GetSelectedVertices();
+	}
+	return {};
+}
+
+void UDataflowEditorVertexAttributePaintTool::ApplyValueToSelection(EDataflowEditorToolEditOperation Operation, float InValue)
+{
+	if (!VertexData.IsValid())
+	{
+		return;
+	}
+
+	TArray<int32> SelectVertices = UDataflowEditorVertexAttributePaintTool::GetSelectedVertices();
+	if (SelectVertices.IsEmpty())
+	{
+		return;
+	}
+
+	BeginChange();
+
+	auto UpdateAttributeValue = 
+		[Operation, InValue](UE::Geometry::FDynamicMeshWeightAttribute& AttributeMap, int32 AttributeIndex) -> double
+		{
+			float CurrentValue;
+			AttributeMap.GetValue(AttributeIndex, &CurrentValue);
+			float NewValue = CurrentValue;
+			switch (Operation)
+			{
+			case EDataflowEditorToolEditOperation::Add:
+				NewValue = (CurrentValue + InValue);
+				break;
+			case EDataflowEditorToolEditOperation::Replace:
+				NewValue = InValue;
+				break;
+			case EDataflowEditorToolEditOperation::Multiply:
+				NewValue = (CurrentValue * InValue);;
+				break;
+			case EDataflowEditorToolEditOperation::Invert:
+				NewValue = (1.0f - CurrentValue);
+				break;
+			case EDataflowEditorToolEditOperation::Relax:
+				// we should not reach there since the relax operation should be handled by FDataflowWeightMapSmoothBrushOp
+				ensure(false);
+				break;
+			}
+			return FMath::Clamp(NewValue, 0.f, 1.f);
+		};
+
+	for (int32 VertexIdx : SelectVertices)
+	{
+		VertexData.SetValue(VertexIdx, UpdateAttributeValue);
+	}
+	
+	// TODO(ccaillaud) - make it work with the relax option - brush should be available other than through the general interface
+	//SmoothBrushOpProperties
+
+	// update colors
+	UpdateVertexColorOverlay();
+	GetToolManager()->PostInvalidation();
+	EndChange();
+}
+
+void UDataflowEditorVertexAttributePaintTool::PruneSelection(float Threshold)
+{
+	if (!VertexData.IsValid())
+	{
+		return;
+	}
+
+	TArray<int32> SelectVertices = UDataflowEditorVertexAttributePaintTool::GetSelectedVertices();
+	if (SelectVertices.IsEmpty())
+	{
+		return;
+	}
+
+	BeginChange();
+
+	auto PruneAttributeValue =
+		[Threshold](UE::Geometry::FDynamicMeshWeightAttribute& AttributeMap, int32 AttributeIndex) -> double
+		{
+			float CurrentValue;
+			AttributeMap.GetValue(AttributeIndex, &CurrentValue);
+			return (CurrentValue >= Threshold) ? CurrentValue : 0;
+		};
+
+	for (int32 VertexIdx : SelectVertices)
+	{
+		VertexData.SetValue(VertexIdx, PruneAttributeValue);
+	}
+
+	UpdateVertexColorOverlay();
+	GetToolManager()->PostInvalidation();
+	EndChange();
+}
+
+namespace UE::DataflowEditorVertexAttributePaintTool::Private
+{
+	const FString CopyAverageFromSelectionToClipboardIdentifier = TEXT("UE_DataflowEditorVertexAttributePaintTool_AverageValue:");
+}
+
+void UDataflowEditorVertexAttributePaintTool::CopyAverageFromSelectionToClipboard()
+{
+	using namespace UE::DataflowEditorVertexAttributePaintTool::Private;
+
+	TArray<int32> SelectVertices = GetSelectedVertices();
+	if (SelectVertices.IsEmpty())
+	{
+		const FText NotificationText = FText::FromString(TEXT("No vertices were selected. Nothing was copied to the clipboard."));
+		ShowEditorMessage(ELogVerbosity::Error, NotificationText);
+		return;
+	}
+
+	const float AverageValue = VertexData.GetAverageValue(SelectVertices);
+
+	// copy to clipboard
+	const FString ClipboardString = FString::Format(TEXT("{0}{1}"), { CopyAverageFromSelectionToClipboardIdentifier, FString::SanitizeFloat(AverageValue) });
+	FPlatformApplicationMisc::ClipboardCopy(*ClipboardString);
+
+	// notify user
+	const FText NotificationText = FText::FromString("Selection average value copied to clipboard.");
+	ShowEditorMessage(ELogVerbosity::Log, NotificationText);
+}
+
+void UDataflowEditorVertexAttributePaintTool::PasteValueToSelectionFromClipboard()
+{
+	using namespace UE::DataflowEditorVertexAttributePaintTool::Private;
+
+	if (!HasSelection())
+	{
+		const FText NotificationText = FText::FromString("No vertices were selected. Nothing was pasted.");
+		ShowEditorMessage(ELogVerbosity::Error, NotificationText);
+		return;
+	}
+
+	// get the clipboard content and check if it matches our format
+	FString ClipboardContent;
+	FPlatformApplicationMisc::ClipboardPaste(ClipboardContent);
+	if (!ClipboardContent.StartsWith(CopyAverageFromSelectionToClipboardIdentifier))
+	{
+		const FText NotificationText = FText::FromString("Failed to paste value from clipboard. Format is incompatible.");
+		ShowEditorMessage(ELogVerbosity::Fatal, NotificationText);
+		return;
+	}
+
+	// parse the value from the clipboard string
+	const FString StringValue = ClipboardContent.RightChop(CopyAverageFromSelectionToClipboardIdentifier.Len());
+	float ValueToPaste = 0.0f;
+	if (StringValue.IsNumeric())
+	{
+		LexFromString(ValueToPaste, *StringValue);
+		ValueToPaste = FMath::Clamp(ValueToPaste, 0.0f, 1.0f);
+
+		ApplyValueToSelection(EDataflowEditorToolEditOperation::Replace, ValueToPaste);
+	}
+	else
+	{
+		const FText NotificationText = FText::FromString("Failed to paste value from clipboard, Value from the clipboard is not numeric.");
+		ShowEditorMessage(ELogVerbosity::Warning, NotificationText);
+	}
+	
+	// notify user
+	const FText NotificationText = FText::FromString("Pasted weights.");
+	ShowEditorMessage(ELogVerbosity::Log, NotificationText);
+}
+
+void UDataflowEditorVertexAttributePaintTool::OnSelectionModified()
+{
+	//const bool bToolTypeIsGradient = (FilterProperties->SubToolType == EDataflowEditorWeightMapPaintInteractionType::Gradient);
+	//if (bToolTypeIsGradient && MeshSelector)
+	//{
+	//	const FGroupTopologySelection NewSelection = MeshSelector->GetSelectionMechanic()->GetActiveSelection(); // TODO(ccaillaud) : remove the new to access the mechanic object
+
+	//	const bool bSelectingLowValueGradientVertices = GetCtrlToggle();
+	//	if (bSelectingLowValueGradientVertices)
+	//	{
+	//		LowValueGradientVertexSelection = NewSelection;
+	//	}
+	//	else
+	//	{
+	//		HighValueGradientVertexSelection = NewSelection;
+	//	}
+
+	//	if (LowValueGradientVertexSelection.SelectedCornerIDs.Num() > 0 && HighValueGradientVertexSelection.SelectedCornerIDs.Num() > 0)
+	//	{
+	//		ComputeGradient();
+	//	}
+
+	//	constexpr bool bBroadcast = false;
+	//	MeshSelector->GetSelectionMechanic()->SetSelection(FGroupTopologySelection(), bBroadcast); // TODO(ccaillaud) should probably have a clear selection 
+	//}
+}
+
+
+void UDataflowEditorVertexAttributePaintTool::SetVerticesToWeightMap(const TSet<int32>& Vertices, double WeightValue, bool bIsErase)
+{
+	BeginChange();
+
+	TempROIBuffer.SetNum(0, EAllowShrinking::No);
+	for (int32 VertexIdx : Vertices)
+	{
+		TempROIBuffer.Add(VertexIdx);
+	}
+
+	if (HaveVisibilityFilter())
+	{
+		TArray<int32> VisibleVertices;
+		VisibleVertices.Reserve(TempROIBuffer.Num());
+		ApplyVisibilityFilter(TempROIBuffer, VisibleVertices);
+		TempROIBuffer = MoveTemp(VisibleVertices);
+	}
+
+	auto UpdateAttributeValue =
+		[WeightValue](UE::Geometry::FDynamicMeshWeightAttribute& AttributeMap, int32 AttributeIndex) -> double
+		{
+			return WeightValue;
+		};
+
+	for (int32 VertexIdx : TempROIBuffer)
+	{
+		VertexData.SetValue(VertexIdx, UpdateAttributeValue);
+	}
+
+	// update colors
+	UpdateVertexColorOverlay();
+	//DynamicMeshComponent->FastNotifyVertexAttributesUpdated(EMeshRenderAttributeFlags::VertexColors);
+	GetToolManager()->PostInvalidation();
+
+	EndChange();
+
+}
+
+bool UDataflowEditorVertexAttributePaintTool::HaveVisibilityFilter() const
+{
+	return (ToolProperties && ToolProperties->BrushProperties.VisibilityFilter != EDataflowEditorToolVisibilityType::None);
+}
+
+void UDataflowEditorVertexAttributePaintTool::ApplyVisibilityFilter(TSet<int32>& Vertices, TArray<int32>& ROIBuffer, TArray<int32>& OutputBuffer)
+{
+	ROIBuffer.SetNum(0, EAllowShrinking::No);
+	ROIBuffer.Reserve(Vertices.Num());
+	for (int32 VertexIdx : Vertices)
+	{
+		ROIBuffer.Add(VertexIdx);
+	}
+
+	OutputBuffer.Reset();
+	ApplyVisibilityFilter(TempROIBuffer, OutputBuffer);
+
+	Vertices.Reset();
+	for (int32 VertexIdx : OutputBuffer)
+	{
+		Vertices.Add(VertexIdx);
+	}
+}
+
+void UDataflowEditorVertexAttributePaintTool::ApplyVisibilityFilter(const TArray<int32>& Vertices, TArray<int32>& VisibleVertices)
+{
+	if (!HaveVisibilityFilter())
+	{
+		VisibleVertices = Vertices;
+		return;
+	}
+
+	FViewCameraState StateOut;
+	GetToolManager()->GetContextQueriesAPI()->GetCurrentViewState(StateOut);
+	FTransform3d LocalToWorld = UE::ToolTarget::GetLocalToWorldTransform(Target);
+	FVector3d LocalEyePosition(LocalToWorld.InverseTransformPosition(StateOut.Position));
+
+	const FDynamicMesh3* Mesh = GetSculptMesh();
+
+	int32 NumVertices = Vertices.Num();
+
+	VisibilityFilterBuffer.SetNum(NumVertices, EAllowShrinking::No);
+	ParallelFor(NumVertices, [&](int32 idx)
+		{
+			VisibilityFilterBuffer[idx] = true;
+			UE::Geometry::FVertexInfo VertexInfo;
+			Mesh->GetVertex(Vertices[idx], VertexInfo, true, false, false);
+			FVector3d Centroid = VertexInfo.Position;
+			FVector3d FaceNormal = (FVector3d)VertexInfo.Normal;
+			if (FaceNormal.Dot((Centroid - LocalEyePosition)) > 0)
+			{
+				VisibilityFilterBuffer[idx] = false;
+			}
+			if (ToolProperties->BrushProperties.VisibilityFilter == EDataflowEditorToolVisibilityType::Unoccluded)
+			{
+				int32 HitTID = Octree.FindNearestHitObject(FRay3d(LocalEyePosition, UE::Geometry::Normalized(Centroid - LocalEyePosition)));
+				if (HitTID != IndexConstants::InvalidID && Mesh->IsTriangle(HitTID))
+				{
+					// Check to see if our vertex has been occulded by another triangle.
+					UE::Geometry::FIndex3i TriVertices = Mesh->GetTriangle(HitTID);
+					if (TriVertices[0] != Vertices[idx] && TriVertices[1] != Vertices[idx] && TriVertices[2] != Vertices[idx])
+					{
+						VisibilityFilterBuffer[idx] = false;
+					}
+				}
+			}
+		});
+
+	VisibleVertices.Reset();
+	for (int32 k = 0; k < NumVertices; ++k)
+	{
+		if (VisibilityFilterBuffer[k])
+		{
+			VisibleVertices.Add(Vertices[k]);
+		}
+	}
+}
+
+
+
+int32 UDataflowEditorVertexAttributePaintTool::FindHitSculptMeshTriangleConst(const FRay3d& LocalRay) const
+{
+	if (!IsInBrushSubMode())
+	{
+		return IndexConstants::InvalidID;
+	}
+
+	if (GetBrushCanHitBackFaces())
+	{
+		return Octree.FindNearestHitObject(LocalRay);
+	}
+	else
+	{
+		const FDynamicMesh3* Mesh = GetSculptMesh();
+
+		FViewCameraState StateOut;
+		GetToolManager()->GetContextQueriesAPI()->GetCurrentViewState(StateOut);
+		FVector3d LocalEyePosition(CurTargetTransform.InverseTransformPosition((FVector3d)StateOut.Position));
+		int HitTID = Octree.FindNearestHitObject(LocalRay,
+			[this, Mesh, &LocalEyePosition](int TriangleID) {
+				FVector3d Normal, Centroid;
+				double Area;
+				Mesh->GetTriInfo(TriangleID, Normal, Area, Centroid);
+				return Normal.Dot((Centroid - LocalEyePosition)) < 0;
+			});
+		return HitTID;
+	}
+}
+
+int32 UDataflowEditorVertexAttributePaintTool::FindHitTargetMeshTriangleConst(const FRay3d& LocalRay) const
+{
+	check(false);
+	return IndexConstants::InvalidID;
+}
+
+void UDataflowEditorVertexAttributePaintTool::UpdateHitSculptMeshTriangle(int32 TriangleID, const FRay3d& LocalRay)
+{
+	// TODO: Figure out what the actual position on the triangle is when hit.
+	CurrentBaryCentricCoords = FVector3d(1.0 / 3.0, 1.0 / 3.0, 1.0 / 3.0);
+}
+
+bool UDataflowEditorVertexAttributePaintTool::UpdateBrushPosition(const FRay& WorldRay)
+{
+	TUniquePtr<FMeshSculptBrushOp>& UseBrushOp = GetActiveBrushOp();
+
+	bool bHit = false;
+	ESculptBrushOpTargetType TargetType = UseBrushOp->GetBrushTargetType();
+	switch (TargetType)
+	{
+	case ESculptBrushOpTargetType::SculptMesh:
+	case ESculptBrushOpTargetType::TargetMesh:
+		bHit = UpdateBrushPositionOnSculptMesh(WorldRay, false);
+		break;
+	case ESculptBrushOpTargetType::ActivePlane:
+		check(false);
+		bHit = UpdateBrushPositionOnSculptMesh(WorldRay, false);
+		break;
+	}
+
+	if (bHit && (UseBrushOp->GetAlignStampToView() || IsVolumetricBrush()))
+	{
+		AlignBrushToView();
+	}
+
+	return bHit;
+}
+
+
+
+
+bool UDataflowEditorVertexAttributePaintTool::OnUpdateHover(const FInputDeviceRay& DevicePos)
+{
+	if (ensure(InStroke() == false))
+	{
+		UpdateBrushPosition(DevicePos.WorldRay);
+	}
+	return true;
+}
+
+
+void UDataflowEditorVertexAttributePaintTool::DrawHUD(FCanvas* Canvas, IToolsContextRenderAPI* RenderAPI)
+{
+	if (PolyLassoMechanic)
+	{
+		// because the actual Weight change is deferred until mouse release, color the lasso to let the user know whether it will erase
+		PolyLassoMechanic->LineColor = GetInEraseStroke() ? FLinearColor::Red : FLinearColor::Green;
+		PolyLassoMechanic->DrawHUD(Canvas, RenderAPI);
+	}
+
+	if (ToolProperties->EditingMode == EDataflowEditorToolEditMode::Mesh)
+	{
+		if (MeshSelector)
+		{
+			MeshSelector->DrawHUD(Canvas, RenderAPI);
+		}
+	}
+
+	if (ToolProperties && ToolProperties->EditingMode == EDataflowEditorToolEditMode::Brush)
+	{
+		if (Canvas && RenderAPI && BrushIndicator)
+		{
+			// display the value under the brush
+			const float BrushRadius = BrushIndicator->BrushRadius;
+			const FVector Offset = (-RenderAPI->GetSceneView()->GetViewUp() * BrushRadius * 1.1f);
+			const FVector BrushPos = BrushIndicator->BrushPosition + Offset;
+
+			FVector4 ScreenBrushPos = RenderAPI->GetSceneView()->WorldToScreen(BrushPos);
+			FVector2D ScreenBrushPos2D = FVector2D(ScreenBrushPos.X, ScreenBrushPos.Y);
+
+			RenderAPI->GetSceneView()->ScreenToPixel(ScreenBrushPos, ScreenBrushPos2D);
+
+			const bool bShouldShowValueAtBrush = IsInBrushSubMode();
+			if (bShouldShowValueAtBrush)
+			{
+				const FText BaseText = LOCTEXT("DataflowWeightMapPaintTool_ValueAtBrush_TextFormat", "Value: {0}");
+				const FText ValueAtBrushText = FText::Format(BaseText, FText::AsNumber(ToolProperties->BrushProperties.ValueAtBrush));
+				FCanvasTextItem TextItem(ScreenBrushPos2D, ValueAtBrushText, GEngine->GetLargeFont(), FLinearColor::White);
+				TextItem.EnableShadow(FLinearColor::Black);
+				TextItem.bCentreX = true;
+				Canvas->DrawItem(TextItem);
+			}
+		}
+	}
+}
+
+void UDataflowEditorVertexAttributePaintTool::Render(IToolsContextRenderAPI* RenderAPI)
+{
+	UMeshSculptToolBase::Render(RenderAPI);
+
+	if (ToolProperties->EditingMode == EDataflowEditorToolEditMode::Mesh)
+	{
+		if (MeshSelector)
+		{
+			MeshSelector->Render(RenderAPI);
+		}
+	}
+}
+
+
+void UDataflowEditorVertexAttributePaintTool::UpdateMaterialMode(EMeshEditingMaterialModes MaterialMode)
+{
+	if (MaterialMode == EMeshEditingMaterialModes::VertexColor)
+	{
+		constexpr bool bUseTwoSidedMaterial = true;
+		ActiveOverrideMaterial = ToolSetupUtil::GetVertexColorMaterial(GetToolManager(), bUseTwoSidedMaterial);
+		if (ensure(ActiveOverrideMaterial != nullptr))
+		{
+			GetSculptMeshComponent()->SetOverrideRenderMaterial(ActiveOverrideMaterial);
+			ActiveOverrideMaterial->SetScalarParameterValue(TEXT("FlatShading"), (ViewProperties->bFlatShading) ? 1.0f : 0.0f);
+		}
+		GetSculptMeshComponent()->SetShadowsEnabled(false);
+	}
+	else
+	{
+		UMeshSculptToolBase::UpdateMaterialMode(MaterialMode);
+	}
+}
+
+void UDataflowEditorVertexAttributePaintTool::OnTick(float DeltaTime)
+{
+	UMeshSculptToolBase::OnTick(DeltaTime);
+	MeshElementsDisplay->OnTick(DeltaTime);
+
+	const bool bIsLasso = false; // TODO(ccaillaud) (FilterProperties->SubToolType == EDataflowEditorWeightMapPaintInteractionType::PolyLasso);
+	if (PolyLassoMechanic)
+	{
+		PolyLassoMechanic->SetIsEnabled(bIsLasso);
+	}
+
+	const bool bIsGradient = false; // TODO(ccaillaud) (FilterProperties->SubToolType == EDataflowEditorWeightMapPaintInteractionType::Gradient);
+	const bool bIsMeshEditMode = (ToolProperties->EditingMode == EDataflowEditorToolEditMode::Mesh);
+	if (MeshSelector)
+	{
+		MeshSelector->SetIsEnabled(bIsGradient || bIsMeshEditMode);
+	}
+
+	check(!(bIsLasso && bIsGradient));
+
+	ConfigureIndicator(IsVolumetricBrush());
+	SetIndicatorVisibility(!bIsLasso && !bIsGradient && !bIsMeshEditMode);
+
+	SCOPE_CYCLE_COUNTER(VertexAttributePaintTool_Tick);
+
+	// process the undo update
+	if (bUndoUpdatePending)
+	{
+		// wait for updates
+		WaitForPendingUndoRedoUpdate();
+
+		// post rendering update
+		//DynamicMeshComponent->FastNotifyTriangleVerticesUpdated(AccumulatedTriangleROI, EMeshRenderAttributeFlags::VertexColors);
+		GetToolManager()->PostInvalidation();
+
+		// ignore stamp and wait for next tick to do anything else
+		bUndoUpdatePending = false;
+		return;
+	}
+
+	// Get value at brush location
+	const bool bShouldPickWeight = bPendingPickWeight && IsStampPending() == false;
+	const bool bShouldUpdateValueAtBrush = IsInBrushSubMode();
+
+	if (bShouldPickWeight || bShouldUpdateValueAtBrush)
+	{
+		if (GetSculptMesh()->IsTriangle(GetBrushTriangleID()))
+		{
+			const double HitWeightValue = GetCurrentWeightValueUnderBrush();
+
+			if (bShouldPickWeight)
+			{
+				ToolProperties->BrushProperties.AttributeValue = HitWeightValue;
+				NotifyOfPropertyChangeByTool(ToolProperties);
+			}
+
+			if (bShouldUpdateValueAtBrush)
+			{
+				ToolProperties->BrushProperties.ValueAtBrush = HitWeightValue;
+			}
+		}
+		bPendingPickWeight = false;
+	}
+
+	auto ExecuteStampOperation = [this](int StampIndex, const FRay& StampRay)
+		{
+			SCOPE_CYCLE_COUNTER(VertexAttributePaintTool_Tick_ApplyStampBlock);
+
+			// update sculpt ROI
+			UpdateROI(CurrentStamp);
+
+			// append updated ROI to modified region (async)
+			FDynamicMesh3* Mesh = GetSculptMesh();
+			TFuture<void> AccumulateROI = Async(UE::DataflowEditorVertexAttributePaintTool::Private::WeightPaintToolAsyncExecTarget, [&]()
+				{
+					UE::Geometry::VertexToTriangleOneRing(Mesh, VertexROI, AccumulatedTriangleROI);
+				});
+
+			// apply the stamp
+			bool bWeightsModified = ApplyStamp();
+
+			if (bWeightsModified)
+			{
+				SCOPE_CYCLE_COUNTER(VertexAttributePaintTool_Tick_UpdateMeshBlock);
+				UpdateVertexColorOverlay(&TriangleROI);
+				//DynamicMeshComponent->FastNotifyTriangleVerticesUpdated(TriangleROI, EMeshRenderAttributeFlags::VertexColors);
+				GetToolManager()->PostInvalidation();
+			}
+
+			// we don't really need to wait for these to happen to end Tick()...
+			AccumulateROI.Wait();
+		};
+
+
+	if (IsInBrushSubMode())
+	{
+		ProcessPerTickStamps(
+			[this](const FRay& StampRay) -> bool {
+				return UpdateStampPosition(StampRay);
+			}, ExecuteStampOperation);
+	}
+
+}
+
+bool UDataflowEditorVertexAttributePaintTool::CanAccept() const
+{
+	return bAnyChangeMade;
+}
+
+void UDataflowEditorVertexAttributePaintTool::BeginChange()
+{
+	VertexData.BeginChange(Cast<UDynamicMeshComponent>(GetSculptMeshComponent()));
+	LongTransactions.Open(LOCTEXT("WeightPaintChange", "Weight Stroke"), GetToolManager());
+}
+
+void UDataflowEditorVertexAttributePaintTool::EndChange()
+{
+	bAnyChangeMade = true;
+
+	using namespace UE::DataflowEditorVertexAttributePaintTool;
+
+	TUniquePtr<Private::FMeshChange> MeshChange = VertexData.EndChange(Cast<UDynamicMeshComponent>(GetSculptMeshComponent()));
+
+	TUniquePtr<TWrappedToolCommandChange<Private::FMeshChange>> NewChange = MakeUnique<TWrappedToolCommandChange<Private::FMeshChange>>();
+	NewChange->WrappedChange = MoveTemp(MeshChange);
+	NewChange->BeforeModify = [this](bool bRevert)
+		{
+			this->WaitForPendingUndoRedoUpdate();
+		};
+
+	GetToolManager()->EmitObjectChange(GetSculptMeshComponent(), MoveTemp(NewChange), LOCTEXT("WeightPaintChange", "Weight Stroke"));
+	LongTransactions.Close(GetToolManager());
+}
+
+
+void UDataflowEditorVertexAttributePaintTool::WaitForPendingUndoRedoUpdate()
+{
+	if (bUndoUpdatePending)
+	{
+		bUndoUpdatePending = false;
+	}
+}
+
+void UDataflowEditorVertexAttributePaintTool::OnDynamicMeshComponentChanged(UDynamicMeshComponent* Component, const FMeshVertexChange* Change, bool bRevert)
+{
+	// update octree
+	FDynamicMesh3* Mesh = GetSculptMesh();
+
+	// make sure any previous async computations are done, and update the undo ROI
+	if (bUndoUpdatePending)
+	{
+		// we should never hit this anymore, because of pre-change calling WaitForPendingUndoRedoUpdate()
+		WaitForPendingUndoRedoUpdate();
+
+		// this is not right because now we are going to do extra recomputation, but it's very messy otherwise...
+		UE::Geometry::VertexToTriangleOneRing(Mesh, Change->Vertices, AccumulatedTriangleROI);
+	}
+	else
+	{
+		AccumulatedTriangleROI.Reset();
+		UE::Geometry::VertexToTriangleOneRing(Mesh, Change->Vertices, AccumulatedTriangleROI);
+	}
+
+	// note that we have a pending update
+	bUndoUpdatePending = true;
+}
+
+
+void UDataflowEditorVertexAttributePaintTool::PrecomputeFilterData()
+{
+	const FDynamicMesh3* Mesh = GetSculptMesh();
+
+	TriNormals.SetNum(Mesh->MaxTriangleID());
+	ParallelFor(Mesh->MaxTriangleID(), [&](int32 tid)
+		{
+			if (Mesh->IsTriangle(tid))
+			{
+				TriNormals[tid] = Mesh->GetTriNormal(tid);
+			}
+		});
+
+	const UE::Geometry::FDynamicMeshNormalOverlay* Normals = Mesh->Attributes()->PrimaryNormals();
+	const UE::Geometry::FDynamicMeshUVOverlay* UVs = Mesh->Attributes()->PrimaryUV();
+	UVSeamEdges.SetNum(Mesh->MaxEdgeID());
+	NormalSeamEdges.SetNum(Mesh->MaxEdgeID());
+	ParallelFor(Mesh->MaxEdgeID(), [&](int32 eid)
+		{
+			if (Mesh->IsEdge(eid))
+			{
+				UVSeamEdges[eid] = UVs->IsSeamEdge(eid);
+				NormalSeamEdges[eid] = Normals->IsSeamEdge(eid);
+			}
+		});
+}
+
+float UDataflowEditorVertexAttributePaintTool::GetCurrentWeightValueUnderBrush() const
+{
+	const int32 VertexID = GetBrushNearestVertex();
+	return VertexData.GetValue(VertexID);
+}
+
+int32 UDataflowEditorVertexAttributePaintTool::GetBrushNearestVertex() const
+{
+	int TriangleVertex = 0;
+
+	if (CurrentBaryCentricCoords.X >= CurrentBaryCentricCoords.Y && CurrentBaryCentricCoords.X >= CurrentBaryCentricCoords.Z)
+	{
+		TriangleVertex = 0;
+	}
+	else
+	{
+		if (CurrentBaryCentricCoords.Y >= CurrentBaryCentricCoords.X && CurrentBaryCentricCoords.Y >= CurrentBaryCentricCoords.Z)
+		{
+			TriangleVertex = 1;
+		}
+		else
+		{
+			TriangleVertex = 2;
+		}
+	}
+	const FDynamicMesh3* Mesh = GetSculptMesh();
+
+	int32 Tid = GetBrushTriangleID();
+	if (Tid == IndexConstants::InvalidID)
+	{
+		return IndexConstants::InvalidID;
+	}
+
+	UE::Geometry::FIndex3i Vertices = Mesh->GetTriangle(Tid);
+	return Vertices[TriangleVertex];
+}
+
+void UDataflowEditorVertexAttributePaintTool::UpdateBrushType(EDataflowEditorToolEditOperation BrushMode)
+{
+	static const FText BaseMessage = LOCTEXT("OnStartTool", "Hold Shift to Erase. [/] and S/D change Size (+Shift to small-step)");
+	FTextBuilder Builder;
+	Builder.AppendLine(BaseMessage);
+
+	if (ToolProperties)
+	{
+		ToolProperties->BrushProperties.BrushMode = BrushMode;
+	}
+	SetActivePrimaryBrushType(PaintBrushId);
+	SetActiveSecondaryBrushType(EraseBrushId);
+
+	SetToolPropertySourceEnabled(PaintBrushOpProperties, false);
+	SetToolPropertySourceEnabled(EraseBrushOpProperties, false);
+	SetToolPropertySourceEnabled(GizmoProperties, false);
+
+	GetToolManager()->DisplayMessage(Builder.ToText(), EToolMessageLevel::UserNotification);
+}
+
+void UDataflowEditorVertexAttributePaintTool::OnColorRampChanged(TArray<FRichCurve*> Curves)
+{
+	UpdateVertexColorOverlay();
+	GetToolManager()->PostInvalidation();
+}
+
+void UDataflowEditorVertexAttributePaintTool::SetColorMode(EDataflowEditorToolColorMode NewColorMode)
+{
+	if (ToolProperties)
+	{
+		ToolProperties->DisplayProperties.ColorMode = NewColorMode;
+		UpdateVertexColorOverlay();
+		GetToolManager()->PostInvalidation();
+	}
+}
+
+void UDataflowEditorVertexAttributePaintTool::UpdateVertexColorOverlay(const TSet<int>* TrianglesToUpdate)
+{
+	auto ValueToColor = [this](float Value) -> FVector4f
+		{
+			switch (ToolProperties->DisplayProperties.ColorMode)
+			{
+			case EDataflowEditorToolColorMode::Greyscale:
+				return FMath::Lerp(FLinearColor::Black, FLinearColor::White, Value);
+			case EDataflowEditorToolColorMode::Ramp:
+				return ToolProperties->DisplayProperties.ColorRamp.ColorRamp.GetLinearColorValue(Value);
+			case EDataflowEditorToolColorMode::FullMaterial:
+				return FLinearColor::White;
+			}
+			return FLinearColor::White;
+		};
+
+	auto SetColorsFromWeights = [&](FDynamicMesh3& Mesh, UE::Geometry::FDynamicMeshColorOverlay& ColorOverlay, int TriangleID)
+		{
+			const UE::Geometry::FIndex3i Tri = Mesh.GetTriangle(TriangleID);
+			const UE::Geometry::FIndex3i ColorElementTri = ColorOverlay.GetTriangle(TriangleID);
+
+			for (int TriVertIndex = 0; TriVertIndex < 3; ++TriVertIndex)
+			{
+				const float Value = VertexData.GetValue(Tri[TriVertIndex]);
+				ColorOverlay.SetElement(ColorElementTri[TriVertIndex], ValueToColor(Value));
+			}
+		};
+
+	// update mesh with new value colors
+	PreviewMesh->DeferredEditMesh(
+		[this, &TrianglesToUpdate, &SetColorsFromWeights](FDynamicMesh3& Mesh)
+		{
+			check(Mesh.HasAttributes());
+			check(Mesh.Attributes()->PrimaryColors());
+			
+			if (UE::Geometry::FDynamicMeshColorOverlay* ColorOverlay = Mesh.Attributes()->PrimaryColors())
+			{
+				if (TrianglesToUpdate)
+				{
+					for (const int TriangleID : *TrianglesToUpdate)
+					{
+						SetColorsFromWeights(Mesh, *ColorOverlay, TriangleID);
+					}
+				}
+				else
+				{
+					for (const int TriangleID : Mesh.TriangleIndicesItr())
+					{
+						SetColorsFromWeights(Mesh, *ColorOverlay, TriangleID);
+					}
+				}
+			}
+		}, false);
+	PreviewMesh->NotifyDeferredEditCompleted(UPreviewMesh::ERenderUpdateMode::FastUpdate, EMeshRenderAttributeFlags::VertexColors, false);
+}
+
+void UDataflowEditorVertexAttributePaintTool::AddReferencedObjects(UObject* InThis, FReferenceCollector& Collector)
+{
+	UDataflowEditorVertexAttributePaintTool* This = CastChecked<UDataflowEditorVertexAttributePaintTool>(InThis);
+	Collector.AddReferencedObject(This->PreviewMesh);
+	Collector.AddReferencedObject(This->MeshSelector);
+	Collector.AddReferencedObject(This->MeshElementsDisplay);
+	Collector.AddReferencedObject(This->DataflowEditorContextObject);
+	Super::AddReferencedObjects(InThis, Collector);
+}
+
+#undef LOCTEXT_NAMESPACE
+
diff --git a/Engine/Plugins/Experimental/Dataflow/Source/DataflowEditor/Private/DataflowEditorTools/DataflowEditorVertexAttributePaintTool.h b/Engine/Plugins/Experimental/Dataflow/Source/DataflowEditor/Private/DataflowEditorTools/DataflowEditorVertexAttributePaintTool.h
new file mode 100644
--- /dev/null
+++ b/Engine/Plugins/Experimental/Dataflow/Source/DataflowEditor/Private/DataflowEditorTools/DataflowEditorVertexAttributePaintTool.h
@@ -0,0 +1,473 @@
+// Copyright Epic Games, Inc. All Rights Reserved.
+
+#pragma once
+
+#include "CoreMinimal.h"
+#include "HAL/Platform.h"
+#include "UObject/NoExportTypes.h"
+#include "BaseTools/BaseBrushTool.h"
+#include "BaseTools/MeshSurfacePointMeshEditingTool.h"
+#include "Components/DynamicMeshComponent.h"
+#include "PropertySets/WeightMapSetProperties.h"
+#include "Mechanics/PolyLassoMarqueeMechanic.h"
+
+#include "Sculpting/MeshSculptToolBase.h"
+#include "Sculpting/MeshBrushOpBase.h"
+
+#include "DynamicMesh/DynamicMeshAABBTree3.h"
+#include "DynamicMesh/DynamicMeshOctree3.h"
+#include "DynamicMesh/DynamicMeshAttributeSet.h"
+#include "DynamicMesh/DynamicMeshChangeTracker.h"
+#include "DynamicMesh/MeshNormals.h"
+#include "TransformTypes.h"
+#include "ToolDataVisualizer.h"
+#include "GroupTopology.h"
+#include "MeshVertexPaintTool.h"
+#include "DataflowEditorTools/DataflowEditorToolBuilder.h"
+#include "Dataflow/DataflowColorRamp.h"
+#include "DataflowEditorTools/DataflowEditorToolEnums.h"
+#include "ToolMeshSelector.h"
+
+#include "DataflowEditorVertexAttributePaintTool.generated.h"
+
+
+#define UE_API DATAFLOWEDITOR_API
+
+class UMeshElementsVisualizer;
+class UDataflowWeightMapEraseBrushOpProps;
+class UDataflowVertexAttributePaintBrushOpProps;
+class UDataflowContextObject;
+struct FDataflowCollectionAddScalarVertexPropertyNode;
+class UDataflowEditorMode;
+
+namespace UE::DataflowEditorVertexAttributePaintTool::Private
+{
+	class FMeshChange;
+}
+
+DECLARE_STATS_GROUP(TEXT("Dataflow_VertexAttributePaintTool"), STATGROUP_VertexAttributePaintTool, STATCAT_Advanced);
+DECLARE_CYCLE_STAT(TEXT("Dataflow_VertexAttributePaintTool_UpdateROI"), VertexAttributePaintTool_UpdateROI, STATGROUP_VertexAttributePaintTool);
+DECLARE_CYCLE_STAT(TEXT("Dataflow_VertexAttributePaintTool_ApplyStamp"), VertexAttributePaintTool_ApplyStamp, STATGROUP_VertexAttributePaintTool);
+DECLARE_CYCLE_STAT(TEXT("Dataflow_VertexAttributePaintTool_Tick"), VertexAttributePaintTool_Tick, STATGROUP_VertexAttributePaintTool);
+DECLARE_CYCLE_STAT(TEXT("Dataflow_VertexAttributePaintTool_Tick_ApplyStampBlock"), VertexAttributePaintTool_Tick_ApplyStampBlock, STATGROUP_VertexAttributePaintTool);
+DECLARE_CYCLE_STAT(TEXT("Dataflow_VertexAttributePaintTool_Tick_UpdateMeshBlock"), VertexAttributePaintTool_Tick_UpdateMeshBlock, STATGROUP_VertexAttributePaintTool);
+
+///////////////////////////////////////////////////////////////////////////////////////////////////////////////////////////////////////////////////
+
+/**
+ * Tool Builder
+ */
+UCLASS(MinimalAPI)
+class UDataflowEditorVertexAttributePaintToolBuilder : public UMeshSurfacePointMeshEditingToolBuilder, public IDataflowEditorToolBuilder
+{
+	GENERATED_BODY()
+public:
+	void SetEditorMode(UDataflowEditorMode* InMode) { Mode = InMode; }
+
+private:
+	UE_API virtual void GetSupportedConstructionViewModes(const UDataflowContextObject& ContextObject, TArray<const UE::Dataflow::IDataflowConstructionViewMode*>& Modes) const override;
+	UE_API virtual bool CanBuildTool(const FToolBuilderState& SceneState) const override;
+	UE_API virtual UMeshSurfacePointTool* CreateNewTool(const FToolBuilderState& SceneState) const override;
+	virtual bool CanSetConstructionViewWireframeActive() const { return false; }
+	UDataflowEditorMode* Mode = nullptr;
+};
+
+///////////////////////////////////////////////////////////////////////////////////////////////////////////////////////////////////////////////////
+
+USTRUCT()
+struct FDataflowEditorVertexAttributePaintToolBrushProperties
+{
+	GENERATED_BODY()
+
+	UPROPERTY()
+	EDataflowEditorToolEditOperation BrushMode = EDataflowEditorToolEditOperation::Add;
+
+	/** Adaptive size of brush - Relative to the model size */
+	UPROPERTY()
+	float BrushSize = 0.25f;
+
+	/** Value to paint on the mesh (dependent on the BrushMode) */
+	UPROPERTY()
+	float AttributeValue = 1.0;
+
+	/** How quickly each brush stroke will drive mesh values towards the desired value */
+	UPROPERTY()
+	float Strength = 1.0;
+
+	/** Area Mode specifies the shape of the brush and which triangles will be included relative to the cursor */
+	UPROPERTY(EditAnywhere, Category = Filters, meta = (DisplayName = "Brush Area Mode"))
+	EMeshVertexPaintBrushAreaType BrushAreaMode = EMeshVertexPaintBrushAreaType::Connected;
+
+
+	/** The Region affected by the current operation will be bounded by edge angles larger than this threshold */
+	UPROPERTY(EditAnywhere, Category = Filters, meta = (UIMin = "0.0", ClampMin = "0.0", UIMax = "180.0", ClampMax = "180.0"))
+	float AngleThreshold = 180.0f;
+
+	/** The Region affected by the current operation will be bounded by UV borders/seams */
+	UPROPERTY(EditAnywhere, Category = Filters)
+	bool bUVSeams = false;
+
+	/** The Region affected by the current operation will be bounded by Hard Normal edges/seams */
+	UPROPERTY(EditAnywhere, Category = Filters)
+	bool bNormalSeams = false;
+
+	/** Control which triangles can be affected by the current operation based on visibility. Applied after all other filters. */
+	UPROPERTY(EditAnywhere, Category = Filters)
+	EDataflowEditorToolVisibilityType VisibilityFilter = EDataflowEditorToolVisibilityType::None;
+
+	/** The weight value at the brush indicator */
+	UPROPERTY(VisibleAnywhere, Transient, Category = Query, meta = (NoResetToDefault))
+	double ValueAtBrush = 0;
+};
+
+///////////////////////////////////////////////////////////////////////////////////////////////////////////////////////////////////////////////////
+
+USTRUCT()
+struct FDataflowEditorVertexAttributePaintToolGradientProperties
+{
+	GENERATED_BODY()
+
+	// Mesh selection mode (vertex, edge face)
+	/** The Gradient upper limit value */
+	UPROPERTY(EditAnywhere, Category = Gradient, meta = (UIMin = 0, ClampMin = 0, UIMax = 1, ClampMax = 1))
+	double GradientHighValue = 1.0;
+
+	/** The Gradient lower limit value */
+	UPROPERTY(EditAnywhere, Category = Gradient, meta = (UIMin = 0, ClampMin = 0, UIMax = 1, ClampMax = 1))
+	double GradientLowValue = 0.0;
+};
+
+///////////////////////////////////////////////////////////////////////////////////////////////////////////////////////////////////////////////////
+
+USTRUCT()
+struct FDataflowEditorVertexAttributePaintToolSelectionProperties
+{
+	GENERATED_BODY()
+
+	// Mesh selection mode (vertex, edge face)
+	UPROPERTY()
+	EComponentSelectionMode ComponentSelectionMode = EComponentSelectionMode::Vertices;
+};
+
+///////////////////////////////////////////////////////////////////////////////////////////////////////////////////////////////////////////////////
+
+USTRUCT()
+struct FDataflowEditorVertexAttributePaintToolDisplayProperties
+{
+	GENERATED_BODY()
+
+	// Vertex color display mode (Greyscale, Ramp, ... )
+	UPROPERTY()
+	EDataflowEditorToolColorMode ColorMode;
+
+	// Color ramp to use when the color mode is set to "Ramp"
+	UPROPERTY()
+	FDataflowColorRamp ColorRamp;
+};
+
+///////////////////////////////////////////////////////////////////////////////////////////////////////////////////////////////////////////////////
+
+UCLASS(MinimalAPI)
+class UDataflowEditorVertexAttributePaintToolProperties : public UInteractiveToolPropertySet
+{
+	GENERATED_BODY()
+
+public:
+
+	// Edit mode (brush,mesh...)
+	UPROPERTY()
+	EDataflowEditorToolEditMode EditingMode = EDataflowEditorToolEditMode::Brush;
+	
+	UPROPERTY()
+	FDataflowEditorVertexAttributePaintToolBrushProperties BrushProperties;
+
+	UPROPERTY()
+	FDataflowEditorVertexAttributePaintToolGradientProperties GradientProperties;
+
+	UPROPERTY()
+	FDataflowEditorVertexAttributePaintToolSelectionProperties SelectionProperties;
+
+	UPROPERTY()
+	FDataflowEditorVertexAttributePaintToolDisplayProperties DisplayProperties;
+};
+
+///////////////////////////////////////////////////////////////////////////////////////////////////////////////////////////////////////////////////
+
+struct FDataflowEditorVertexAttributePaintToolData
+{
+public:
+	void Setup(FDynamicMesh3& InMesh, FDataflowCollectionAddScalarVertexPropertyNode* InNodeToUpdate, TObjectPtr<UDataflowContextObject> InDataflowEditorContextObject);
+	void Shutdown();
+
+	void BeginChange(UDynamicMeshComponent* Component);
+	TUniquePtr<UE::DataflowEditorVertexAttributePaintTool::Private::FMeshChange> EndChange(UDynamicMeshComponent* Component);
+
+	bool IsValid() const { return (ActiveWeightMap != nullptr); }
+
+	float GetValue(int32 VertexIdx) const;
+	float GetAverageValue(const TArray<int32>& Vertices) const;
+
+	// lambda signature : float GetValue(UE::Geometry::FDynamicMeshWeightAttribute& AttributeMap, int32 AttributeIndex)
+	template <typename TLambda>
+	void SetValue(int32 VertexIdx, TLambda Lambda)
+	{
+		if (ActiveWeightMap)
+		{
+			if (bHaveDynamicMeshToWeightConversion)
+			{
+				for (const int32 Idx : WeightToDynamicMesh[DynamicMeshToWeight[VertexIdx]])
+				{
+					ActiveWeightEditChangeTracker->SaveVertexOneRingTriangles(Idx, true);
+					const float Value = Lambda(*ActiveWeightMap, Idx);
+					ActiveWeightMap->SetValue(Idx, &Value);
+				}
+			}
+			else
+			{
+				ActiveWeightEditChangeTracker->SaveVertexOneRingTriangles(VertexIdx, true);
+				const float Value = Lambda(*ActiveWeightMap, VertexIdx);
+				ActiveWeightMap->SetValue(VertexIdx, &Value);
+			}
+		}
+	}
+
+	void UpdateNode();
+
+private:
+	// DynamicMesh might be unwelded mesh, but weights are on the welded mesh.
+	bool bHaveDynamicMeshToWeightConversion = false;
+	TArray<int32> DynamicMeshToWeight;
+	TArray<TArray<int32>> WeightToDynamicMesh;
+
+	/** Setup weights used to store the initial values */
+	TArray<float> SetupWeights;
+
+	// Corresponding node to read from and write back to
+	FDataflowCollectionAddScalarVertexPropertyNode* NodeToUpdate = nullptr;
+
+	UE::Geometry::FDynamicMeshWeightAttribute* ActiveWeightMap = nullptr;
+
+	TUniquePtr<UE::Geometry::FDynamicMeshChangeTracker> ActiveWeightEditChangeTracker;
+};
+
+///////////////////////////////////////////////////////////////////////////////////////////////////////////////////////////////////////////////////
+
+/**
+ * Dataflow Editor tool to paint a vertex attribute as vertex colors
+ */
+UCLASS(MinimalAPI)
+class UDataflowEditorVertexAttributePaintTool : public UMeshSculptToolBase
+{
+	GENERATED_BODY()
+
+public:
+	UE_API virtual void RegisterActions(FInteractiveToolActionSet& ActionSet) override;
+
+	UE_API virtual void Setup() override;
+	UE_API virtual void Shutdown(EToolShutdownType ShutdownType) override;
+	UE_API virtual void DrawHUD(FCanvas* Canvas, IToolsContextRenderAPI* RenderAPI) override;
+	UE_API virtual void Render(IToolsContextRenderAPI* RenderAPI) override;
+	UE_API virtual void UpdateMaterialMode(EMeshEditingMaterialModes MaterialMode) override;
+	UE_API virtual void OnTick(float DeltaTime) override;
+
+	virtual bool HasCancel() const override { return true; }
+	virtual bool HasAccept() const override { return true; }
+	UE_API virtual bool CanAccept() const override;
+
+	UE_API virtual bool OnUpdateHover(const FInputDeviceRay& DevicePos) override;
+
+	UE_API virtual void OnPropertyModified(UObject* PropertySet, FProperty* Property) override;
+
+	UE_API bool IsInBrushSubMode() const;
+	UE_API bool IsVolumetricBrush() const;
+
+	UE_API virtual void CommitResult(UBaseDynamicMeshComponent* Component, bool bModifiedTopology) override;
+
+	UE_API void SetDataflowEditorContextObject(TObjectPtr<UDataflowContextObject> InDataflowEditorContextObject);
+	void SetEditorMode(UDataflowEditorMode* InMode) { Mode = InMode; }
+
+	virtual void SetupBrushEditBehaviorSetup(ULocalTwoAxisPropertyEditInputBehavior& OutBehavior) override;
+
+	//~ UObject interface
+
+	static UE_API void AddReferencedObjects(UObject* InThis, FReferenceCollector& Collector);
+
+public:
+	/** properties of the tool */
+	UPROPERTY()
+	TObjectPtr<UDataflowEditorVertexAttributePaintToolProperties> ToolProperties;
+
+	UPROPERTY()
+	TObjectPtr<UDataflowVertexAttributePaintBrushOpProps> PaintBrushOpProperties;
+
+	UPROPERTY()
+	TObjectPtr<UDataflowWeightMapEraseBrushOpProps> EraseBrushOpProperties;
+
+	UE_API bool HaveVisibilityFilter() const;
+	UE_API void ApplyVisibilityFilter(const TArray<int32>& Vertices, TArray<int32>& VisibleVertices);
+	UE_API void ApplyVisibilityFilter(TSet<int32>& Triangles, TArray<int32>& ROIBuffer, TArray<int32>& OutputBuffer);
+
+	// we override these so we can update the separate BrushSize property added for this tool
+	UE_API virtual void IncreaseBrushRadiusAction();
+	UE_API virtual void DecreaseBrushRadiusAction();
+	UE_API virtual void IncreaseBrushRadiusSmallStepAction();
+	UE_API virtual void DecreaseBrushRadiusSmallStepAction();
+
+	UE_API float GetBrushMinRadius() const;
+	UE_API float GetBrushMaxRadius() const;
+
+	UE_API void SetColorMode(EDataflowEditorToolColorMode NewColorMode);
+
+	UE_API bool HasSelection() const;
+	UE_API void GrowSelection() const;
+	UE_API void ShrinkSelection() const;
+	UE_API void FloodSelection() const;
+	UE_API void SelectBorder() const;
+	UE_API void SetComponentSelectionMode(EComponentSelectionMode NewMode);
+
+	UE_API void CopyAverageFromSelectionToClipboard();
+	UE_API void PasteValueToSelectionFromClipboard();
+	UE_API void PruneSelection(float Threshold);
+
+	// selection operations
+	UE_API void ApplyValueToSelection(EDataflowEditorToolEditOperation Operation, float InValue);
+
+	UE_API TArray<int32> GetSelectedVertices() const;
+
+protected:
+	// UMeshSculptToolBase API
+	virtual UBaseDynamicMeshComponent* GetSculptMeshComponent() override;
+	virtual FDynamicMesh3* GetBaseMesh()  override { check(false); return nullptr; }
+	virtual const FDynamicMesh3* GetBaseMesh() const  override { check(false); return nullptr; }
+
+	UE_API virtual int32 FindHitSculptMeshTriangleConst(const FRay3d& LocalRay) const override;
+	UE_API virtual int32 FindHitTargetMeshTriangleConst(const FRay3d& LocalRay) const override;
+	UE_API virtual void UpdateHitSculptMeshTriangle(int32 TriangleID, const FRay3d& LocalRay) override;
+
+	UE_API virtual void OnBeginStroke(const FRay& WorldRay) override;
+	UE_API virtual void OnEndStroke() override;
+
+	UE_API void OnColorRampChanged(TArray<FRichCurve*> Curves);
+
+	void CreateDefaultColorRamp(FDataflowColorRamp& OutColorRamp);
+
+	UE_API virtual TUniquePtr<FMeshSculptBrushOp>& GetActiveBrushOp();
+
+	virtual bool SharesBrushPropertiesChanges() const override { return false; }
+	// end UMeshSculptToolBase API
+
+protected:
+	//
+	// Marquee Support
+	//
+public:
+	UPROPERTY()
+	TObjectPtr<UPolyLassoMarqueeMechanic> PolyLassoMechanic;
+
+protected:
+	UE_API void OnPolyLassoFinished(const FCameraPolyLasso& Lasso, bool bCanceled);
+
+
+	// 
+	// Gradient Support
+	//
+	FToolDataVisualizer GradientSelectionRenderer;
+
+	UE::Geometry::FGroupTopologySelection LowValueGradientVertexSelection;
+	UE::Geometry::FGroupTopologySelection HighValueGradientVertexSelection;
+
+	UE_API void ComputeGradient();
+	UE_API void OnSelectionModified();
+
+	//
+	// Internals
+	//
+
+protected:
+
+	UPROPERTY()
+	TObjectPtr<UMeshElementsVisualizer> MeshElementsDisplay = nullptr;
+
+	UPROPERTY()
+	TObjectPtr<UDataflowContextObject> DataflowEditorContextObject = nullptr;
+
+	UPROPERTY()
+	TObjectPtr<UPreviewMesh> PreviewMesh = nullptr;
+	FMeshDescription PreviewMeshDescription;
+
+	UPROPERTY()
+	TObjectPtr<UToolMeshSelector> MeshSelector;
+
+	// realtime visualization
+	UE_API void OnDynamicMeshComponentChanged(UDynamicMeshComponent* Component, const FMeshVertexChange* Change, bool bRevert);
+	FDelegateHandle OnDynamicMeshComponentChangedHandle;
+
+	UE_API float GetCurrentWeightValueUnderBrush() const;
+	FVector3d CurrentBaryCentricCoords;
+	UE_API int32 GetBrushNearestVertex() const;
+
+	TSet<int32> AccumulatedTriangleROI;
+	bool bUndoUpdatePending = false;
+	TArray<int> NormalsBuffer;
+	UE_API void WaitForPendingUndoRedoUpdate();
+
+	TArray<int> TempROIBuffer;
+	TArray<int> VertexROI;
+	TArray<bool> VisibilityFilterBuffer;
+	TSet<int> VertexSetBuffer;
+	TSet<int> TriangleROI;
+	UE_API void UpdateROI(const FSculptBrushStamp& CurrentStamp);
+
+	UE_API bool UpdateStampPosition(const FRay& WorldRay);
+	UE_API bool ApplyStamp();
+
+	UE::Geometry::FDynamicMeshOctree3 Octree;
+
+	UE_API void UpdateBrushType(EDataflowEditorToolEditOperation EditOperation);
+	UE_API bool UpdateBrushPosition(const FRay& WorldRay);
+
+	bool GetInEraseStroke()
+	{
+		// Re-use the smoothing stroke key (shift) for erase stroke in the weight paint tool
+		return GetInSmoothingStroke();
+	}
+
+
+	bool bPendingPickWeight = false;
+
+	TArray<int32> ROITriangleBuffer;
+	TArray<double> ROIWeightValueBuffer;
+	UE_API bool SyncMeshWithWeightBuffer(FDynamicMesh3* Mesh);
+	UE_API bool SyncWeightBufferWithMesh(const FDynamicMesh3* Mesh);
+	UE_API void SetVerticesToWeightMap(const TSet<int32>& Vertices, double WeightValue, bool bIsErase);
+
+	UE_API void BeginChange();
+	UE_API void EndChange();
+
+	TArray<FVector3d> TriNormals;
+	TArray<int32> UVSeamEdges;
+	TArray<int32> NormalSeamEdges;
+	UE_API void PrecomputeFilterData();
+
+	FDataflowEditorVertexAttributePaintToolData VertexData;
+
+protected:
+	virtual bool ShowWorkPlane() const override { return false; }
+
+	friend class UDataflowEditorVertexAttributePaintToolBuilder;
+
+	bool bAnyChangeMade = false;
+
+	/** update the vertex color of the mesh from the attribute */
+	UE_API void UpdateVertexColorOverlay(const TSet<int>* TrianglesToUpdate = nullptr);
+
+	UDataflowEditorMode* Mode = nullptr;
+
+	static constexpr int32 PaintBrushId = 0;
+	static constexpr int32 EraseBrushId = 1;
+};
+
+
+
+#undef UE_API
+
diff --git a/Engine/Plugins/Experimental/Dataflow/Source/DataflowEditor/Private/DataflowEditorTools/DataflowEditorVertexAttributePaintToolDetailCustomization.cpp b/Engine/Plugins/Experimental/Dataflow/Source/DataflowEditor/Private/DataflowEditorTools/DataflowEditorVertexAttributePaintToolDetailCustomization.cpp
new file mode 100644
--- /dev/null
+++ b/Engine/Plugins/Experimental/Dataflow/Source/DataflowEditor/Private/DataflowEditorTools/DataflowEditorVertexAttributePaintToolDetailCustomization.cpp
@@ -0,0 +1,1216 @@
+// Copyright Epic Games, Inc. All Rights Reserved.
+
+#include "DataflowEditorTools/DataflowEditorVertexAttributePaintToolDetailCustomization.h"
+
+#include "DataflowEditorTools/DataflowEditorToolEnums.h"
+#include "DataflowEditorTools/DataflowEditorVertexAttributePaintTool.h"
+
+#include "DetailCategoryBuilder.h"
+#include "DetailLayoutBuilder.h"
+#include "DetailWidgetRow.h"
+#include "ModelingToolsEditorModeStyle.h"
+#include "SColorGradientEditor.h"
+#include "Widgets/Input/SSegmentedControl.h"
+#include "Widgets/Input/SSpinBox.h"
+#include "Widgets/Input/SButton.h"
+#include "Widgets/Input/SNumericEntryBox.h"
+#include "Framework/MultiBox/MultiBoxBuilder.h"
+#include "Drawing/MeshElementsVisualizer.h"
+
+#define LOCTEXT_NAMESPACE "VertexAttributePaintToolDetailCustomization"
+
+namespace UE::Dataflow::Editor
+{
+	// layout constants
+	float FVertexAttributePaintToolDetailCustomization::WeightSliderWidths = 150.0f;
+	float FVertexAttributePaintToolDetailCustomization::WeightEditingLabelsPercent = 0.40f;
+	float FVertexAttributePaintToolDetailCustomization::WeightEditVerticalPadding = 4.0f;
+	float FVertexAttributePaintToolDetailCustomization::WeightEditHorizontalPadding = 2.0f;
+
+	///////////////////////////////////////////////////////////////////////////////////////////////////////////////
+
+	TSharedRef<IDetailCustomization> FVertexAttributePaintToolDetailCustomization::MakeInstance()
+	{
+		return MakeShareable(new FVertexAttributePaintToolDetailCustomization);
+	}
+
+	///////////////////////////////////////////////////////////////////////////////////////////////////////////////
+
+	FVertexAttributePaintToolDetailCustomization::~FVertexAttributePaintToolDetailCustomization()
+	{
+		if (Tool.IsValid())
+		{
+			// TODO(ccaillaud) ==>>  Tool->OnSelectionChanged.RemoveAll(this);
+		}
+
+		Tool.Reset();
+		ToolProperties.Reset();
+	}
+
+	///////////////////////////////////////////////////////////////////////////////////////////////////////////////
+
+	void FVertexAttributePaintToolDetailCustomization::OnSelectionChanged()
+	{
+		// TODO(ccaillaud) ==>>  ToolProperties.Get()->DirectEditState.Reset();
+	};
+
+	///////////////////////////////////////////////////////////////////////////////////////////////////////////////
+
+	void FVertexAttributePaintToolDetailCustomization::AddEditingModeRow(IDetailCategoryBuilder& EditModeCategory) const
+	{
+		// Add segmented control toggle for editing modes ("Brush" or "Selection")
+		EditModeCategory.AddCustomRow(LOCTEXT("EditModeCategory", "Value Editing Mode"), false)
+		.WholeRowContent()
+		[
+			SNew(SBox)
+			.Padding(2.0f)
+			[
+				SNew(SSegmentedControl<EDataflowEditorToolEditMode>)
+				.ToolTipText(LOCTEXT("EditingModeTooltip",
+					"Brush: edit values by painting on mesh.\n"
+					"Mesh: select vertices/edges/faces to edit values directly.\n"
+				))
+				.Value_Lambda([this]()
+					{
+						return ToolProperties.IsValid() ? ToolProperties->EditingMode : EDataflowEditorToolEditMode::Brush;
+					})
+				.OnValueChanged_Lambda([this](EDataflowEditorToolEditMode Mode)
+					{
+						if (ToolProperties.IsValid())
+						{
+							ToolProperties->EditingMode = Mode;
+							// TODO(ccaillaud) ==>> ToolProperties->OwnerTool->ToggleEditingMode();
+							if (CurrentDetailBuilder)
+							{
+								CurrentDetailBuilder->ForceRefreshDetails();
+							}
+						}
+					})
+				+ SSegmentedControl<EDataflowEditorToolEditMode>::Slot(EDataflowEditorToolEditMode::Brush)
+				.Text(LOCTEXT("BrushEditMode", "Brush"))
+				+ SSegmentedControl<EDataflowEditorToolEditMode>::Slot(EDataflowEditorToolEditMode::Mesh)
+				.Text(LOCTEXT("MeshEditMode", "Mesh"))
+			]
+		];
+	}
+
+	///////////////////////////////////////////////////////////////////////////////////////////////////////////////
+
+	void FVertexAttributePaintToolDetailCustomization::AddColorModeRow(IDetailCategoryBuilder& MeshDisplayCategory) const
+	{
+		MeshDisplayCategory.AddCustomRow(LOCTEXT("ColorModeCategory", "Color Mode"), false)
+		.WholeRowContent()
+		[
+			SNew(SBox)
+			.Padding(2.0f)
+			.HAlign(HAlign_Center)
+			[
+				SNew(SSegmentedControl<EDataflowEditorToolColorMode>)
+				.ToolTipText(LOCTEXT("WeightColorTooltip",
+					"Adjust the value display in the viewport.\n\n"
+					"Greyscale: Displays values by blending from black (0) to white (1).\n"
+					"Ramp: Values 0 and 1 use the min and max colors. values in between 0 and 1 use the ramp colors.\n"
+					"Full Material: Displays normal mesh materials with textures.\n"))
+				.Value_Lambda([this]()
+					{
+						return ToolProperties.IsValid() ? ToolProperties->DisplayProperties.ColorMode : EDataflowEditorToolColorMode::Greyscale;
+					})
+				.OnValueChanged_Lambda([this](EDataflowEditorToolColorMode Mode)
+					{
+						if (Tool.IsValid())
+						{
+							Tool->SetColorMode(Mode);
+						}
+					})
+				+ SSegmentedControl<EDataflowEditorToolColorMode>::Slot(EDataflowEditorToolColorMode::Greyscale)
+				.Text(LOCTEXT("GreyscaleMode", "Greyscale"))
+				+ SSegmentedControl<EDataflowEditorToolColorMode>::Slot(EDataflowEditorToolColorMode::Ramp)
+				.Text(LOCTEXT("RampMode", "Ramp"))
+				+ SSegmentedControl<EDataflowEditorToolColorMode>::Slot(EDataflowEditorToolColorMode::FullMaterial)
+				.Text(LOCTEXT("MaterialMode", "Full Material"))
+			]
+		];
+	}
+
+	///////////////////////////////////////////////////////////////////////////////////////////////////////////////
+
+	void FVertexAttributePaintToolDetailCustomization::AddColorRampRow(IDetailCategoryBuilder& EditModeCategory) const
+	{
+		TSharedRef<SColorGradientEditor> GradientEditor = SNew(SColorGradientEditor)
+			.ViewMinInput(0.0f)
+			.ViewMaxInput(1.0f)
+			.ClampStopsToViewRange(true)
+			.DrawColorAndAlphaSeparate(false)
+			.IsEditingEnabled_Lambda([this]()
+				{
+					return ToolProperties.IsValid() && (ToolProperties->DisplayProperties.ColorMode == EDataflowEditorToolColorMode::Ramp);
+				});
+		if (ToolProperties.IsValid())
+		{
+			GradientEditor->SetCurveOwner(&ToolProperties->DisplayProperties.ColorRamp.ColorRamp);
+		}
+
+		EditModeCategory.AddCustomRow(LOCTEXT("ColorRampCategory", "Color Ramp"), false)
+		.WholeRowContent()
+		[
+			GradientEditor
+		];
+	}
+
+	///////////////////////////////////////////////////////////////////////////////////////////////////////////////
+
+	void FVertexAttributePaintToolDetailCustomization::CustomizeDetails(IDetailLayoutBuilder& DetailBuilder)
+	{
+		CurrentDetailBuilder = &DetailBuilder;
+	
+		TArray<TWeakObjectPtr<UObject>> DetailObjects;
+		DetailBuilder.GetObjectsBeingCustomized(DetailObjects);
+
+		// Should be impossible to get multiple settings objects for a single tool
+		ensure(DetailObjects.Num()==1);
+		ToolProperties = Cast<UDataflowEditorVertexAttributePaintToolProperties>(DetailObjects[0]);
+		Tool = Cast<UDataflowEditorVertexAttributePaintTool>(ToolProperties->GetOuter());
+		// TODO(ccaillaud) ==>> Tool->OnSelectionChanged.AddSP(this, &FVertexAttributePaintToolDetailCustomization::OnSelectionChanged);
+	
+		// Editing Mode category 
+		IDetailCategoryBuilder& EditModeCategory = DetailBuilder.EditCategory("Value Editing Mode", FText::GetEmpty(), ECategoryPriority::Important);
+		{
+			// Editing mode row [ Brush | Mesh ]
+			AddEditingModeRow(EditModeCategory);
+
+			switch (ToolProperties->EditingMode)
+			{
+			case EDataflowEditorToolEditMode::Brush:
+				AddBrushUI(DetailBuilder);
+				break;
+			case EDataflowEditorToolEditMode::Mesh:
+				AddSelectionUI(DetailBuilder);
+				break;
+			}
+		}
+
+		// Mesh Display category
+		IDetailCategoryBuilder& MeshDisplayCategory = DetailBuilder.EditCategory("MeshDisplay", FText::GetEmpty(), ECategoryPriority::Important);
+		MeshDisplayCategory.InitiallyCollapsed(false);
+		{
+			AddColorModeRow(MeshDisplayCategory);
+
+			AddColorRampRow(MeshDisplayCategory);
+		}
+
+		// Hide all customized properties 
+		HideProperty(DetailBuilder, GET_MEMBER_NAME_CHECKED(UDataflowEditorVertexAttributePaintToolProperties, EditingMode));
+		HideProperty(DetailBuilder, GET_MEMBER_NAME_CHECKED(FDataflowEditorVertexAttributePaintToolBrushProperties, BrushMode));
+		HideProperty(DetailBuilder, GET_MEMBER_NAME_CHECKED(FDataflowEditorVertexAttributePaintToolBrushProperties, BrushSize));
+		HideProperty(DetailBuilder, GET_MEMBER_NAME_CHECKED(FDataflowEditorVertexAttributePaintToolBrushProperties, BrushAreaMode));
+		HideProperty(DetailBuilder, GET_MEMBER_NAME_CHECKED(FDataflowEditorVertexAttributePaintToolBrushProperties, ValueAtBrush));
+		HideProperty(DetailBuilder, GET_MEMBER_NAME_CHECKED(FDataflowEditorVertexAttributePaintToolBrushProperties, AttributeValue));
+		HideProperty(DetailBuilder, GET_MEMBER_NAME_CHECKED(FDataflowEditorVertexAttributePaintToolBrushProperties, Strength));
+		HideProperty(DetailBuilder, GET_MEMBER_NAME_CHECKED(FDataflowEditorVertexAttributePaintToolDisplayProperties, ColorMode));
+
+		HideProperty(DetailBuilder, GET_MEMBER_NAME_CHECKED(UMeshElementsVisualizerProperties, bVisible));
+	}
+
+	///////////////////////////////////////////////////////////////////////////////////////////////////////////////
+
+	void FVertexAttributePaintToolDetailCustomization::AddBrushModeRow(IDetailCategoryBuilder& BrushCategory) const
+	{
+		// add segmented control toggle for brush behavior modes ("Add", "Replace", etc..)
+		BrushCategory.AddCustomRow(LOCTEXT("BrushModeCategory", "Brush Mode"), false)
+		.WholeRowContent()
+		[
+			SNew(SBox)
+			.Padding(2.0f)
+			[
+				SNew(SSegmentedControl<EDataflowEditorToolEditOperation>)
+				.ToolTipText(LOCTEXT("BrushTooltip",
+					"Select the operation to apply when using the brush.\n"
+					"Add: applies the current value plus the Strength value to the new value.\n"
+					"Replace: applies the current value minus the Strength value to the new value.\n"
+					"Multiply: applies the current value multiplied by the Strength value to the new value.\n"
+					"Relax: applies the average of the connected (by edge) vertex value to the new vertex value, blended by the Strength."))
+				.Value_Lambda([this]()
+					{
+						return ToolProperties.IsValid() ? ToolProperties->BrushProperties.BrushMode : EDataflowEditorToolEditOperation::Add;
+					})
+				.OnValueChanged_Lambda([this](EDataflowEditorToolEditOperation Mode)
+					{
+						if (ToolProperties.IsValid())
+						{
+							ToolProperties->BrushProperties.BrushMode = Mode;
+							FPropertyChangedEvent PropertyChangedEvent(UBrushBaseProperties::StaticClass()->FindPropertyByName(GET_MEMBER_NAME_CHECKED(FDataflowEditorVertexAttributePaintToolBrushProperties, BrushMode)));
+							ToolProperties->PostEditChangeProperty(PropertyChangedEvent);
+						}
+					})
+				+ SSegmentedControl<EDataflowEditorToolEditOperation>::Slot(EDataflowEditorToolEditOperation::Add)
+				.Text(LOCTEXT("BrushAddMode", "Add"))
+				+ SSegmentedControl<EDataflowEditorToolEditOperation>::Slot(EDataflowEditorToolEditOperation::Replace)
+				.Text(LOCTEXT("BrushReplaceMode", "Replace"))
+				+ SSegmentedControl<EDataflowEditorToolEditOperation>::Slot(EDataflowEditorToolEditOperation::Multiply)
+				.Text(LOCTEXT("BrushMultiplyMode", "Multiply"))
+				+ SSegmentedControl<EDataflowEditorToolEditOperation>::Slot(EDataflowEditorToolEditOperation::Relax)
+				.Text(LOCTEXT("BrushRelaxMode", "Relax"))
+			]
+		];
+	}
+
+	///////////////////////////////////////////////////////////////////////////////////////////////////////////////
+
+	void FVertexAttributePaintToolDetailCustomization::AddBrushFalloffModeRow(IDetailCategoryBuilder& BrushCategory) const
+	{
+		// add segmented control toggle for brush falloff modes ("Surface" or "Volume")
+		BrushCategory.AddCustomRow(LOCTEXT("BrushFalloffModeCategory", "Brush Falloff Mode"), false)
+		.WholeRowContent()
+		[
+			SNew(SBox)
+			.Padding(2.0f)
+			[
+				SNew(SSegmentedControl<EMeshVertexPaintBrushAreaType>)
+				.ToolTipText(LOCTEXT("BrushFalloffModeTooltip",
+					"Surface: falloff is based on the distance along the surface from the brush center to nearby connected vertices.\n"
+					"Volume: falloff is based on the straight-line distance from the brush center to surrounding vertices.\n"))
+				.Value_Lambda([this]()
+					{
+						return ToolProperties.IsValid() ? ToolProperties->BrushProperties.BrushAreaMode : EMeshVertexPaintBrushAreaType::Connected;
+					})
+				.OnValueChanged_Lambda([this](EMeshVertexPaintBrushAreaType Mode)
+					{
+						if (ToolProperties.IsValid())
+						{
+							ToolProperties->BrushProperties.BrushAreaMode = Mode; // TODO(ccaillaud) ==>> 
+							// TODO(ccaillaud) ==>> ToolProperties->SetFalloffMode(Mode);
+						}
+					})
+				+ SSegmentedControl<EMeshVertexPaintBrushAreaType>::Slot(EMeshVertexPaintBrushAreaType::Connected)
+				.Text(LOCTEXT("SurfaceMode", "Surface"))
+				+ SSegmentedControl<EMeshVertexPaintBrushAreaType>::Slot(EMeshVertexPaintBrushAreaType::Volumetric)
+				.Text(LOCTEXT("VolumeMode", "Volume"))
+			]
+		];
+	}
+
+	///////////////////////////////////////////////////////////////////////////////////////////////////////////////
+
+	void FVertexAttributePaintToolDetailCustomization::AddBrushRadiusRow(IDetailCategoryBuilder& BrushCategory) const
+	{
+		constexpr float MinBrushRadius = 0.01f;
+		constexpr float MaxBrushRadius = 20.f;
+		constexpr float DefaultBrushRadius = 10.0f;
+
+		BrushCategory.AddCustomRow(LOCTEXT("BrushSizeCategory", "Brush Radius"), false)
+		.NameContent()
+		[
+			SNew(STextBlock)
+			.Text(LOCTEXT("BrushRadiusLabel", "Radius"))
+			.Font(FAppStyle::GetFontStyle(TEXT("PropertyWindow.NormalFont")))
+			.ToolTipText(LOCTEXT("BrushRadiusTooltip", "The radius of the brush in scene units."))
+		]
+		.ValueContent()
+		[
+			SNew(SSpinBox<float>)
+			.MinValue_Lambda([this]() 
+				{ 
+					return Tool.IsValid() ? (Tool->GetBrushMinRadius()) : MinBrushRadius;
+				})
+			.MaxSliderValue_Lambda([this]()
+				{
+					return Tool.IsValid() ? (Tool->GetBrushMaxRadius()) : MaxBrushRadius;
+				})
+			.SupportDynamicSliderMaxValue(true)
+			.Value_Lambda([this]()
+				{
+					// BrushSize is adaptive ( [0,1] range ) , we need to scale it to the max value 
+					const float MaxRadius = Tool.IsValid() ? (Tool->GetBrushMaxRadius()) : MaxBrushRadius;
+					return (ToolProperties.IsValid()) ? (ToolProperties->BrushProperties.BrushSize * MaxRadius) : MaxRadius;
+				})
+			.OnValueChanged_Lambda([this](float NewValue)
+				{
+					if (ToolProperties.IsValid())
+					{
+						// BrushSize is adaptive ( [0,1] range ), we need to divide by max value to set it 
+						const float MaxRadius = FMath::Max((Tool.IsValid() ? (Tool->GetBrushMaxRadius()) : MaxBrushRadius), MinBrushRadius);
+						ToolProperties->BrushProperties.BrushSize = (NewValue / MaxRadius);
+						FPropertyChangedEvent PropertyChangedEvent(UBrushBaseProperties::StaticClass()->FindPropertyByName(GET_MEMBER_NAME_CHECKED(FDataflowEditorVertexAttributePaintToolBrushProperties, BrushSize)));
+						ToolProperties->PostEditChangeProperty(PropertyChangedEvent);
+					}
+				})
+			.OnValueCommitted_Lambda([this](float NewValue, ETextCommit::Type CommitType)
+				{
+					if (ToolProperties.IsValid())
+					{
+						// TODO(ccaillaud) ==>> ToolProperties->SaveConfig();
+					}
+				})
+		];
+	}
+
+	///////////////////////////////////////////////////////////////////////////////////////////////////////////////
+
+	void FVertexAttributePaintToolDetailCustomization::AddBrushStrengthRow(IDetailCategoryBuilder& BrushCategory) const
+	{
+		BrushCategory.AddCustomRow(LOCTEXT("BrushStrengthCategory", "Brush Strength"), false)
+		.NameContent()
+		[
+			SNew(STextBlock)
+			.Text(LOCTEXT("BrushStrengthLabel", "Strength"))
+			.Font(FAppStyle::GetFontStyle(TEXT("PropertyWindow.NormalFont")))
+			.ToolTipText(LOCTEXT("BrushStrengthTooltip", "The strength of the effect on the weights. Exact effect depends on the Brush mode."))
+		]
+		.ValueContent()
+		[
+			SNew(SSpinBox<float>)
+			.MinValue(0.f)
+			.MaxValue(2.0f)
+			.MaxSliderValue(1.f)
+			.Value(1.0f)
+			.SupportDynamicSliderMaxValue(true)
+			.Value_Lambda([this]()
+				{
+					return ToolProperties.IsValid() ? ToolProperties->BrushProperties.Strength : 1.f;
+				})
+			.OnValueChanged_Lambda([this](float NewValue)
+				{
+					if (ToolProperties.IsValid())
+					{
+						ToolProperties->BrushProperties.Strength = NewValue;
+						FPropertyChangedEvent PropertyChangedEvent(UBrushBaseProperties::StaticClass()->FindPropertyByName(GET_MEMBER_NAME_CHECKED(FDataflowEditorVertexAttributePaintToolBrushProperties, Strength)));
+						ToolProperties->PostEditChangeProperty(PropertyChangedEvent);
+					}
+				})
+			.OnValueCommitted_Lambda([this](float NewValue, ETextCommit::Type CommitType)
+				{
+					if (ToolProperties.IsValid())
+					{
+						// TODO(ccaillaud) ==>> ToolProperties->SaveConfig();
+					}
+				})
+		];
+	}
+
+	///////////////////////////////////////////////////////////////////////////////////////////////////////////////
+
+	void FVertexAttributePaintToolDetailCustomization::AddBrushValueRow(IDetailCategoryBuilder& BrushCategory) const
+	{
+		BrushCategory.AddCustomRow(LOCTEXT("AttributeValueCategory", "Attribute Value"), false)
+		.NameContent()
+		[
+			SNew(STextBlock)
+			.Text(LOCTEXT("AttributeValueLabel", "Attribute Value"))
+			.Font(FAppStyle::GetFontStyle(TEXT("PropertyWindow.NormalFont")))
+			.ToolTipText(LOCTEXT("AttributeValueTooltip", "The value of the attribute to be applied. Exact effect depends on the Brush strength and mode."))
+		]
+		.ValueContent()
+		[
+			SNew(SSpinBox<float>)
+			.MinValue(0.f)
+			.MaxValue(1.0f)
+			.MaxSliderValue(1.f)
+			.Value(1.0f)
+			.SupportDynamicSliderMaxValue(true)
+			.Value_Lambda([this]()
+				{
+					return ToolProperties.IsValid() ? ToolProperties->BrushProperties.AttributeValue : 1.f;
+				})
+			.OnValueChanged_Lambda([this](float NewValue)
+				{
+					if (ToolProperties.IsValid())
+					{
+						ToolProperties->BrushProperties.AttributeValue = FMath::Clamp(NewValue, 0.f, 1.f);
+						FPropertyChangedEvent PropertyChangedEvent(UBrushBaseProperties::StaticClass()->FindPropertyByName(GET_MEMBER_NAME_CHECKED(FDataflowEditorVertexAttributePaintToolBrushProperties, AttributeValue)));
+						ToolProperties->PostEditChangeProperty(PropertyChangedEvent);
+					}
+				})
+			.OnValueCommitted_Lambda([this](float NewValue, ETextCommit::Type CommitType)
+				{
+					if (ToolProperties.IsValid())
+					{
+						// TODO(ccaillaud) ==>> ToolProperties->SaveConfig();
+					}
+				})
+		];
+	}
+
+	///////////////////////////////////////////////////////////////////////////////////////////////////////////////
+
+	void FVertexAttributePaintToolDetailCustomization::AddBrushUI(IDetailLayoutBuilder& DetailBuilder) const
+	{
+		IDetailCategoryBuilder& BrushCategory = DetailBuilder.EditCategory("Brush", FText::GetEmpty(), ECategoryPriority::Important);
+		{
+			// Brush modes [ Add | Replace | Multiply | Relax ]
+			AddBrushModeRow(BrushCategory);
+
+			// Brush Falloff modes [ Surface | Volume ]
+			AddBrushFalloffModeRow(BrushCategory);
+
+			// Brush radius field
+			AddBrushRadiusRow(BrushCategory);
+
+			// Brush value field
+			AddBrushValueRow(BrushCategory);
+
+			// Brush strength field
+			AddBrushStrengthRow(BrushCategory);
+		}
+	}
+
+	///////////////////////////////////////////////////////////////////////////////////////////////////////////////
+
+	TSharedRef<SWidget> FVertexAttributePaintToolDetailCustomization::MakeSelectionElementsToolbar() const
+	{
+		FSlimHorizontalToolBarBuilder ToolbarBuilder(MakeShared<FUICommandList>(), FMultiBoxCustomization::None);
+		ToolbarBuilder.SetStyle(FModelingToolsEditorModeStyle::Get().Get(), "PolyEd.SelectionToolbar");
+		ToolbarBuilder.SetLabelVisibility(EVisibility::Collapsed);
+
+		ToolbarBuilder.BeginSection("SelectionFilter");
+		ToolbarBuilder.BeginBlockGroup();
+
+		auto AddToggleButtonForBool = [&ToolbarBuilder, this](EComponentSelectionMode Mode, const FText& Label, const FText& Tooltip, const FName IconName)
+			{
+				ToolbarBuilder.AddToolBarButton(FUIAction(
+					FExecuteAction::CreateLambda([this, Mode]()
+						{
+							if (Tool.IsValid())
+							{
+								Tool->SetComponentSelectionMode(Mode);
+							}
+						}),
+					FCanExecuteAction::CreateLambda([this]()
+						{
+							return ToolProperties.IsValid() ? ToolProperties->EditingMode == EDataflowEditorToolEditMode::Mesh : false;
+						}),
+					FIsActionChecked::CreateLambda([this, Mode]()
+						{
+							return ToolProperties.IsValid() ? (ToolProperties->SelectionProperties.ComponentSelectionMode == Mode) : false;
+						})),
+					NAME_None,	// Extension hook
+					Label,		// Label
+					Tooltip,	// Tooltip
+					FSlateIcon(FModelingToolsEditorModeStyle::Get()->GetStyleSetName(), IconName),
+					EUserInterfaceActionType::ToggleButton);
+			};
+
+		AddToggleButtonForBool(
+			EComponentSelectionMode::Vertices,
+			LOCTEXT("VerticesLabel", "Vertices"),
+			LOCTEXT("VerticesTooltip", "Select mesh vertices."),
+			"PolyEd.SelectCorners");
+		AddToggleButtonForBool(
+			EComponentSelectionMode::Edges,
+			LOCTEXT("EdgesLabel", "Edges"),
+			LOCTEXT("EdgesTooltip", "Select mesh edges."),
+			"PolyEd.SelectEdges");
+		AddToggleButtonForBool(
+			EComponentSelectionMode::Faces,
+			LOCTEXT("FacesLabel", "Faces"),
+			LOCTEXT("FacesTooltip", "Select mesh faces."),
+			"PolyEd.SelectFaces");
+
+		ToolbarBuilder.EndBlockGroup();
+		ToolbarBuilder.EndSection();
+
+		return ToolbarBuilder.MakeWidget();
+	}
+
+	TSharedRef<SWidget> FVertexAttributePaintToolDetailCustomization::MakeSelectionIsolationWidget() const
+	{
+#if 0
+		return SNew(SCheckBox)
+		.Style(FAppStyle::Get(), "ToggleButtonCheckbox")
+		.HAlign(HAlign_Center)
+		.ToolTipText(LOCTEXT("IsolateSelectedTooltip",
+			"Shows only the selected faces in the viewport.\n"
+			"Weight editing operations will not affect hidden vertices.\n"
+			"NOTE: This only works on the target (main) mesh."))
+		.IsEnabled_Lambda([this]()
+			{
+				if (Tool.IsValid())
+				{
+					// isolated selection only available on main mesh (for now)
+					const bool bHasSelection = Tool->GetMainMeshSelector()->IsAnyComponentSelected();
+					const bool bAlreadyIsolatingSelection = Tool->GetSelectionIsolator()->IsSelectionIsolated();
+					return bHasSelection || bAlreadyIsolatingSelection;
+				}
+				return false;
+			})
+		.IsChecked_Lambda([this]()
+			{
+				if (Tool.IsValid())
+				{
+					const bool bAlreadyIsolatingSelection = Tool->GetSelectionIsolator()->IsSelectionIsolated();
+					return bAlreadyIsolatingSelection ? ECheckBoxState::Checked : ECheckBoxState::Unchecked;
+				}
+				return ECheckBoxState::Unchecked;
+			})
+		.OnCheckStateChanged_Lambda([this](ECheckBoxState InCheckBoxState)
+			{
+				if (UWeightToolSelectionIsolator* SelectionIsolator = Tool.IsValid() ? Tool->GetSelectionIsolator() : nullptr)
+				{
+					if (InCheckBoxState == ECheckBoxState::Checked)
+					{
+						SelectionIsolator->IsolateSelectionAsTransaction();
+					}
+					else
+					{
+						SelectionIsolator->UnIsolateSelectionAsTransaction();
+					}
+				}
+			})
+		[
+			SNew(STextBlock)
+				.Text_Lambda([this]()
+					{
+						if (!Tool.IsValid() || Tool->GetSelectionIsolator()->IsSelectionIsolated())
+						{
+							return LOCTEXT("ShowAllButtonLabel", "Show Full Mesh");
+						}
+
+						return LOCTEXT("IsolateButtonLabel", "Isolate Selected");
+					})
+		];
+
+#else //0
+		return SNullWidget::NullWidget;
+
+#endif //0
+	}
+
+	TSharedRef<SWidget> FVertexAttributePaintToolDetailCustomization::MakeSelectionEditActionsToolbar() const
+	{
+		// [ GROW / SHRINK / FLOOD ]
+		return SNew(SVerticalBox)
+			+ SVerticalBox::Slot()
+			.AutoHeight()
+			[
+				SNew(SHorizontalBox)
+				+ SHorizontalBox::Slot()
+				.Padding(2.f, WeightEditVerticalPadding)
+				[
+					SNew(SButton)
+					.IsEnabled_Lambda([this]
+						{
+							return Tool.IsValid() ? Tool->HasSelection() : false;
+						})
+					.HAlign(HAlign_Center)
+					.VAlign(VAlign_Center)
+					.Text(LOCTEXT("GrowSelectionButtonLabel", "Grow"))
+					.ToolTipText(LOCTEXT("GrowSelectionTooltip",
+						"Grow the current selection by adding connected neighbors to current selection.\n"))
+					.OnClicked_Lambda([this]()
+						{
+							if (Tool.IsValid())
+							{
+								Tool->GrowSelection();
+							}
+							return FReply::Handled();
+						})
+				]
+				+ SHorizontalBox::Slot()
+				.Padding(2.f, WeightEditVerticalPadding)
+				[
+					SNew(SButton)
+					.IsEnabled_Lambda([this]
+						{
+							return Tool.IsValid() ? Tool->HasSelection() : false;
+						})
+					.HAlign(HAlign_Center)
+					.VAlign(VAlign_Center)
+					.Text(LOCTEXT("ShrinkSelectionButtonLabel", "Shrink"))
+					.ToolTipText(LOCTEXT("ShrinkSelectionTooltip",
+						"Shrink the current selection by removing components on the border of the current selection.\n"))
+					.OnClicked_Lambda([this]()
+						{
+							if (Tool.IsValid())
+							{
+								Tool->ShrinkSelection();
+							}
+							return FReply::Handled();
+						})
+				]
+			]
+			+ SVerticalBox::Slot()
+			.AutoHeight()
+			[
+				SNew(SHorizontalBox)
+				+ SHorizontalBox::Slot()
+				.Padding(2.f, WeightEditVerticalPadding)
+				[
+					SNew(SButton)
+					.IsEnabled_Lambda([this]
+						{
+							return Tool.IsValid() ? Tool->HasSelection() : false;
+						})
+					.HAlign(HAlign_Center)
+					.VAlign(VAlign_Center)
+					.Text(LOCTEXT("FloodSelectionButtonLabel", "Flood"))
+					.ToolTipText(LOCTEXT("FloodSelectionTooltip",
+						"Flood the current selection by adding all connected components to the current selection.\n"))
+					.OnClicked_Lambda([this]()
+						{
+							if (Tool.IsValid())
+							{
+								Tool->FloodSelection();
+							}
+							return FReply::Handled();
+						})
+				]
+				+ SHorizontalBox::Slot()
+				.Padding(2.f, WeightEditVerticalPadding)
+				[
+					SNew(SButton)
+					.IsEnabled_Lambda([this]
+						{
+							return Tool.IsValid() ? Tool->HasSelection() : false;
+						})
+					.HAlign(HAlign_Center)
+					.VAlign(VAlign_Center)
+					.Text(LOCTEXT("BorderSelectionButtonLabel", "Convert to Border"))
+					.ToolTipText(LOCTEXT("BorderSelectionTooltip", "Select vertices on the border of the current selection."))
+					.OnClicked_Lambda([this]()
+						{
+							if (Tool.IsValid())
+							{
+								Tool->SelectBorder();
+							}
+							return FReply::Handled();
+						})
+				]
+			];
+	}
+
+	void FVertexAttributePaintToolDetailCustomization::AddSelectionElementsRow(IDetailCategoryBuilder& EditSelectionCategory) const
+	{
+		TSharedRef<SWidget> ElementsToolBar = MakeSelectionElementsToolbar();
+		TSharedRef<SWidget> IsolationWidget = MakeSelectionIsolationWidget();
+		TSharedRef<SWidget> SelectionEditActionsToolBar = MakeSelectionEditActionsToolbar();
+
+		EditSelectionCategory.AddCustomRow(LOCTEXT("EditSelectionRow", "Edit Selection"), false)
+		.WholeRowContent()
+		[
+			SNew(SVerticalBox)
+			+ SVerticalBox::Slot()
+			.AutoHeight()
+			[
+				SNew(SHorizontalBox)
+				+ SHorizontalBox::Slot()
+				.HAlign(HAlign_Center)
+				.VAlign(VAlign_Center)
+				.Padding(2.0f, WeightEditVerticalPadding)
+				[
+					SNew(SHorizontalBox)
+					+ SHorizontalBox::Slot()
+					[
+						ElementsToolBar
+					]
+					+ SHorizontalBox::Slot()
+					.AutoWidth()
+					.VAlign(VAlign_Center)
+					[
+						IsolationWidget
+					]
+				]
+			]
+			+ SVerticalBox::Slot()
+			.AutoHeight()
+			[
+				SelectionEditActionsToolBar
+			]
+		];
+	}
+
+	void FVertexAttributePaintToolDetailCustomization::AddEmptySelectionWarningRow( IDetailCategoryBuilder& EditValuesCategory) const
+	{
+		EditValuesCategory.AddCustomRow(LOCTEXT("SelectMessageRow", "Select Vertices"), false)
+		.WholeRowContent()
+		.VAlign(VAlign_Center)
+		[
+			SNew(STextBlock)
+			.Visibility_Lambda([this]()
+				{
+					return Tool.IsValid() && !Tool->HasSelection() ? EVisibility::Visible : EVisibility::Collapsed;
+				})
+			.Text_Lambda([this]()
+				{
+					return LOCTEXT("NothingSelectedLabel", "Select vertices on target mesh to edit weights...");
+				})
+		];
+	}
+
+	void FVertexAttributePaintToolDetailCustomization::MakeSelectionOperationRow(
+		IDetailCategoryBuilder& EditValuesCategory,
+		const FText& RowName,
+		const TSharedRef<SWidget>& ButtonWidget, 
+		const TSharedRef<SWidget>& ValueWidget) const
+	{
+		EditValuesCategory.AddCustomRow(RowName, false)
+		.WholeRowContent()
+		[
+			SNew(SHorizontalBox)
+			.IsEnabled_Lambda([this]
+				{
+					return Tool.IsValid() ? Tool->HasSelection() : false;
+				})
+			+ SHorizontalBox::Slot()
+			[
+				SNew(SBox)[ButtonWidget]
+			]
+			+ SHorizontalBox::Slot()
+			[
+				ValueWidget
+			]
+		];
+	}
+
+	void FVertexAttributePaintToolDetailCustomization::MakeSelectionAddMultiplySliderOperationRow(IDetailCategoryBuilder& EditValuesCategory) const
+	{
+		TSharedRef<SSegmentedControl<EDataflowEditorToolEditOperation>> AddMultiplyButtonsWidget =
+			SNew(SSegmentedControl<EDataflowEditorToolEditOperation>)
+			.ToolTipText(LOCTEXT("InteractiveEditModeTooltip",
+				"Add: applies the current value plus the flood value to the new value.\n"
+				"Multiply: applies the current value multiplied by the flood value to the new value.\n"
+				"This operation applies interactively while dragging the slider. It operates on the currently selected vertices."))
+			.Value(EDataflowEditorToolEditOperation::Add)
+			+ SSegmentedControl<EDataflowEditorToolEditOperation>::Slot(EDataflowEditorToolEditOperation::Add)
+			.Text(LOCTEXT("BrushAddMode", "Add"))
+			+ SSegmentedControl<EDataflowEditorToolEditOperation>::Slot(EDataflowEditorToolEditOperation::Multiply)
+			.Text(LOCTEXT("BrushMultiplyMode", "Multiply"))
+			;
+
+		auto GetMinValue = [AddMultiplyButtonsWidget]() -> float
+			{
+				switch (AddMultiplyButtonsWidget->GetValue())
+				{
+				case EDataflowEditorToolEditOperation::Add: return -1.f;
+				case EDataflowEditorToolEditOperation::Multiply: return 0.f;
+				}
+				return 0.f;
+			};
+
+		auto GetMaxValue = [AddMultiplyButtonsWidget]() -> float
+			{
+				switch (AddMultiplyButtonsWidget->GetValue())
+				{
+				case EDataflowEditorToolEditOperation::Add: return 1.f;
+				case EDataflowEditorToolEditOperation::Multiply: return 10.f;
+				}
+				return 0.f;
+			};
+
+		struct FSliderState
+		{
+			float Value = 0.f;
+		};
+
+		TSharedPtr<FSliderState> SliderState = MakeShared<FSliderState>();
+
+		auto GetCurrentValue = [SliderState]()
+			{
+				return SliderState->Value;
+			};
+
+		auto ResetValue = [SliderState, AddMultiplyButtonsWidget]()
+			{
+				switch (AddMultiplyButtonsWidget->GetValue())
+				{
+				case EDataflowEditorToolEditOperation::Add: 
+					SliderState->Value = 0.f;
+					break;
+				case EDataflowEditorToolEditOperation::Multiply:
+					SliderState->Value = 1.f;
+					break;
+				}
+			};
+
+		auto ApplyOperation = [this, SliderState, AddMultiplyButtonsWidget](float NewValue)
+			{
+				if (Tool.IsValid())
+				{
+					SliderState->Value = NewValue;
+					Tool->ApplyValueToSelection(AddMultiplyButtonsWidget->GetValue(), NewValue);
+				}
+			};
+
+		TSharedRef<SWidget> SliderValueWidget = 
+			SNew(SSpinBox<float>)
+			.MinSliderValue_Lambda(GetMinValue)
+			.MaxSliderValue_Lambda(GetMaxValue)
+			.MinValue_Lambda(GetMinValue)
+			.MaxValue_Lambda(GetMaxValue)
+			.Value_Lambda(GetCurrentValue)
+			.OnValueChanged_Lambda([ApplyOperation](float NewValue)
+				{
+					ApplyOperation(NewValue);
+				})
+			.OnValueCommitted_Lambda([ApplyOperation](float NewValue, ETextCommit::Type CommitType)
+				{
+					ApplyOperation(NewValue);
+				})
+			.OnBeginSliderMovement_Lambda([this]()
+				{
+					// TODO(ccaillaud)
+					/*if (Tool.IsValid())
+					{
+						Tool->BeginChange();
+					}*/
+				})
+			.OnEndSliderMovement_Lambda([this, ResetValue](float)
+				{
+					const FText TransactionLabel = LOCTEXT("FloodWeightChange", "Flood values on vertices.");
+					// TODO(ccaillaud)
+					/*if (Tool.IsValid())
+					{
+						Tool->EndChange(TransactionLabel);
+					}*/
+
+					// reset the value
+					//SliderState-
+					ResetValue();
+				})
+			.ToolTipText(LOCTEXT("FloodWeightsToolTip", "Drag the slider to interactively adjust values on the selected vertices."))
+			;
+
+		MakeSelectionOperationRow(EditValuesCategory, LOCTEXT("FloodValuesRow", "Flood Values Slider"), AddMultiplyButtonsWidget, SliderValueWidget);
+	}
+
+	void FVertexAttributePaintToolDetailCustomization::MakeSelectionAddOperationRow(IDetailCategoryBuilder& EditValuesCategory) const
+	{
+		TSharedRef<SSpinBox<float>> AddValueWidget =
+			SNew(SSpinBox<float>)
+			.MinValue(0)
+			.MaxValue(1)
+			.Value(1)
+			.ToolTipText(LOCTEXT("AddValueSliderToolTip", "Adjust the value to Add to the selected vertices."));
+
+		TSharedRef<SWidget> AddButtonWidget =
+			SNew(SButton)
+			.HAlign(HAlign_Center)
+			.Text(LOCTEXT("AddValueButtonLabel", "Add"))
+			.ToolTipText(LOCTEXT("AddValueButtonTooltip",
+				"Add: applies the current value plus the flood value to the new weight.\n"
+				"This operation applies to the currently selected vertices."))
+			.OnClicked_Lambda([this, AddValueWidget]()
+				{
+					if (Tool.IsValid())
+					{
+						Tool->ApplyValueToSelection(EDataflowEditorToolEditOperation::Add, AddValueWidget->GetValue());
+					}
+					return FReply::Handled();
+				});
+
+		MakeSelectionOperationRow(EditValuesCategory, LOCTEXT("SelectionAddOperationRow", "Selection Add Operation"), AddButtonWidget, AddValueWidget);
+	}
+
+	void FVertexAttributePaintToolDetailCustomization::MakeSelectionReplaceOperationRow(IDetailCategoryBuilder& EditValuesCategory) const
+	{
+		TSharedRef<SSpinBox<float>> ReplaceValueWidget =
+			SNew(SSpinBox<float>)
+			.MinValue(0)
+			.MaxValue(1)
+			.Value(1)
+			.ToolTipText(LOCTEXT("ReplaceValueSliderToolTip", "Adjust the value to be replace on the selected vertices."));
+
+		TSharedRef<SWidget> ReplaceButtonWidget =
+			SNew(SButton)
+			.HAlign(HAlign_Center)
+			.Text(LOCTEXT("ReplaceValueButtonLabel", "Replace"))
+			.ToolTipText(LOCTEXT("ReplaceValueButtonTooltip",
+				"Replace: the value of selected vertices is replaced by the specified value.\n"
+				"This operation applies to the currently selected vertices."))
+			.OnClicked_Lambda([this, ReplaceValueWidget]()
+				{
+					if (Tool.IsValid())
+					{
+						Tool->ApplyValueToSelection(EDataflowEditorToolEditOperation::Replace, ReplaceValueWidget->GetValue());
+					}
+					return FReply::Handled();
+				});
+
+		MakeSelectionOperationRow(EditValuesCategory, LOCTEXT("SelectionReplaceOperationRow", "Selection Replace Operation"), ReplaceButtonWidget, ReplaceValueWidget);
+	}
+
+	void FVertexAttributePaintToolDetailCustomization::MakeSelectionInvertOperationRow(IDetailCategoryBuilder& EditValuesCategory) const
+	{
+		TSharedRef<SWidget> NoValueValueWidget = SNullWidget::NullWidget;
+
+		TSharedRef<SWidget> ReplaceButtonWidget =
+			SNew(SButton)
+			.HAlign(HAlign_Center)
+			.Text(LOCTEXT("InvertValueButtonLabel", "Invert"))
+			.ToolTipText(LOCTEXT("InvertValueButtonTooltip",
+				"Invert: the value of selected vertices is replaced by opne minus the specified value.\n"
+				"This operation applies to the currently selected vertices."))
+			.OnClicked_Lambda([this]()
+				{
+					if (Tool.IsValid())
+					{
+						Tool->ApplyValueToSelection(EDataflowEditorToolEditOperation::Invert, 0.0f);
+					}
+					return FReply::Handled();
+				});
+
+		MakeSelectionOperationRow(EditValuesCategory, LOCTEXT("SelectionInvertOperationRow", "Selection Invert Operation"), ReplaceButtonWidget, NoValueValueWidget);
+	}
+
+	void FVertexAttributePaintToolDetailCustomization::MakeSelectionRelaxOperationRow(IDetailCategoryBuilder& EditValuesCategory) const
+	{
+		TSharedRef<SSpinBox<float>> RelaxValueWidget =
+			SNew(SSpinBox<float>)
+			.MinValue(0)
+			.MaxValue(1)
+			.Value(0.5)
+			.ToolTipText(LOCTEXT("RelaxValueSliderToolTip", "Amount to blend when relaxing  the values"));
+
+		TSharedRef<SWidget> RelaxButtonWidget =
+			SNew(SButton)
+			.HAlign(HAlign_Center)
+			.Text(LOCTEXT("RelaxValueButtonLabel", "Relax"))
+			.ToolTipText(LOCTEXT("RelaxValueButtonTooltip",
+				"Relax: the value of each selected vertex is replaced by the average of it's neighbors. \n"
+				"This smooths values across the mesh."))
+			.OnClicked_Lambda([this, RelaxValueWidget]()
+				{
+					if (Tool.IsValid())
+					{
+						Tool->ApplyValueToSelection(EDataflowEditorToolEditOperation::Relax, RelaxValueWidget->GetValue());
+					}
+					return FReply::Handled();
+				});
+
+		MakeSelectionOperationRow(EditValuesCategory, LOCTEXT("SelectionRelaxOperationRow", "Selection Relax Operation"), RelaxButtonWidget, RelaxValueWidget);
+	}
+
+	void FVertexAttributePaintToolDetailCustomization::MakeSelectionPruneOperationRow(IDetailCategoryBuilder& EditValuesCategory) const
+	{
+		TSharedRef<SSpinBox<float>> PruneValueWidget =
+			SNew(SSpinBox<float>)
+			.MinValue(0.f)
+			.MaxValue(1.f)
+			.Value(0.01)
+			.ToolTipText(LOCTEXT("PruneValueSliderToolTip", "Prune Threshold - Values below this threshold will be set to zero"))
+			.OnValueCommitted_Lambda([this](float NewValue, ETextCommit::Type CommitType)
+				{
+					if (ToolProperties.IsValid())
+					{
+						// TOOD(ccaillaud) ==> ToolProperties->SaveConfig();
+					}
+				});
+
+		TSharedRef<SWidget> PruneButtonWidget =
+			SNew(SButton)
+			.HAlign(HAlign_Center)
+			.Text(LOCTEXT("PruneValuesButtonLabel", "Prune"))
+			.ToolTipText(LOCTEXT("PruneButtonTooltip",
+				"Zero values below the given threshold value.\n"
+				"This command operates on the selected vertices."))
+			.OnClicked_Lambda([this, PruneValueWidget]()
+				{
+					if (Tool.IsValid())
+					{
+						Tool->PruneSelection(PruneValueWidget->GetValue());
+					}
+					return FReply::Handled();
+				});
+
+		MakeSelectionOperationRow(EditValuesCategory, LOCTEXT("SelectionPruneOperationRow", "Selection Prune Operation"), PruneButtonWidget, PruneValueWidget);
+	}
+
+	void FVertexAttributePaintToolDetailCustomization::MakeSelectionCopyAndPasteRow(IDetailCategoryBuilder& EditValuesCategory) const
+	{
+		// COPY/PASTE WEIGHTS category
+		EditValuesCategory.AddCustomRow(LOCTEXT("CopyPasteValuesRow", "Copy Paste"), false)
+		.WholeRowContent()
+		[
+			SNew(SHorizontalBox)
+			.IsEnabled_Lambda([this]
+				{
+					return Tool.IsValid() ? Tool->HasSelection() : false;
+				})
+
+			+ SHorizontalBox::Slot()
+			.Padding(2.f, WeightEditVerticalPadding)
+			[
+				SNew(SBox)
+				[
+					SNew(SButton)
+					.HAlign(HAlign_Center)
+					.Text(LOCTEXT("CopyValuesButtonLabel", "Copy"))
+					.ToolTipText(LOCTEXT("CopyButtonTooltip",
+						"Copy the average of the selected values to the clipboard. \n"
+						"This is designed to work with the Paste command."))
+					.OnClicked_Lambda([this]()
+						{
+							if (Tool.IsValid())
+							{
+								Tool->CopyAverageFromSelectionToClipboard();
+							}
+							return FReply::Handled();
+						})
+				]
+			]
+
+			+ SHorizontalBox::Slot()
+			.Padding(2.f, WeightEditVerticalPadding)
+			[
+				SNew(SButton)
+				.HAlign(HAlign_Center)
+				.Text(LOCTEXT("PasteValuesButtonLabel", "Paste"))
+				.ToolTipText(LOCTEXT("PasteButtonTooltip",
+					"Paste the value on the selected vertices.\n"
+					"This command requires the clipboard contain the selection average value from the Copy command."))
+				.OnClicked_Lambda([this]()
+					{
+						if (ToolProperties.IsValid())
+						{
+							Tool->PasteValueToSelectionFromClipboard();
+						}
+						return FReply::Handled();
+					})
+			]
+		];
+	}
+
+	///////////////////////////////////////////////////////////////////////////////////////////////////////////////
+
+	void FVertexAttributePaintToolDetailCustomization::AddSelectionUI(IDetailLayoutBuilder& DetailBuilder) const
+	{
+		// custom display of selection editing tools
+		IDetailCategoryBuilder& EditSelectionCategory = DetailBuilder.EditCategory("Edit Selection", FText::GetEmpty(), ECategoryPriority::Important);
+		EditSelectionCategory.InitiallyCollapsed(false);
+		{
+			// create a toolbar for the selection filter [ Vtx |  Edges | Faces ] (Isolate)
+			AddSelectionElementsRow(EditSelectionCategory);
+		}
+	
+		IDetailCategoryBuilder& EditValuesCategory = DetailBuilder.EditCategory("Edit Values", FText::GetEmpty(), ECategoryPriority::Important);
+		EditValuesCategory.InitiallyCollapsed(false);
+		{
+			AddEmptySelectionWarningRow(EditValuesCategory);
+
+			MakeSelectionAddMultiplySliderOperationRow(EditValuesCategory);
+
+			MakeSelectionAddOperationRow(EditValuesCategory);
+
+			MakeSelectionReplaceOperationRow(EditValuesCategory);
+
+			MakeSelectionInvertOperationRow(EditValuesCategory);
+
+			MakeSelectionRelaxOperationRow(EditValuesCategory);
+
+			MakeSelectionPruneOperationRow(EditValuesCategory);
+
+			MakeSelectionCopyAndPasteRow(EditValuesCategory);
+		}
+
+#if 0
+		// MIRROR WEIGHTS category
+		EditWeightsCategory.AddCustomRow(LOCTEXT("MirrorWeightsRow", "Mirror"), false)
+		.WholeRowContent()
+		[
+			SNew(SVerticalBox)
+			.IsEnabled_Lambda([this]{ return Tool->HasActiveSelectionOnMainMesh(); })
+		
+			+ SVerticalBox::Slot()
+			.Padding(WeightEditHorizontalPadding, WeightEditVerticalPadding)
+			[
+				SNew(SHorizontalBox)
+			
+				+SHorizontalBox::Slot()
+				.VAlign(VAlign_Center)
+				.FillWidth(WeightEditingLabelsPercent)
+				[
+					SNew(STextBlock)
+					.Text(LOCTEXT("MirrorPlaneLabel", "Mirror Plane"))
+					.Font(FAppStyle::GetFontStyle(TEXT("PropertyWindow.NormalFont")))
+					.ToolTipText(LOCTEXT("MirrorPlaneTooltip", "The plane to copy weights across."))
+				]
+
+				+SHorizontalBox::Slot()
+				.FillWidth(1.f)
+				[
+					SNew(SHorizontalBox)
+					+SHorizontalBox::Slot()
+					[
+						SNew(SSegmentedControl<EAxis::Type>)
+						.ToolTipText(LOCTEXT("MirrorAxisTooltip",
+							"X: copies weights across the YZ plane.\n"
+							"Y: copies weights across the XZ plane.\n"
+							"Z: copies weights across the XY plane."))
+						.Value_Lambda([this]()
+						{
+							return ToolProperties.IsValid() ? ToolProperties->MirrorAxis : TEnumAsByte<EAxis::Type>(EAxis::X);
+						})
+						.OnValueChanged_Lambda([this](EAxis::Type Mode)
+						{
+							if (ToolProperties.IsValid())
+							{
+								ToolProperties->MirrorAxis = Mode;
+							}
+						})
+						+ SSegmentedControl<EAxis::Type>::Slot(EAxis::X)
+						.Text(LOCTEXT("MirrorXLabel", "X"))
+						+ SSegmentedControl<EAxis::Type>::Slot(EAxis::Y)
+						.Text(LOCTEXT("MirrorYLabel", "Y"))
+						+ SSegmentedControl<EAxis::Type>::Slot(EAxis::Z)
+						.Text(LOCTEXT("MirrorZLabel", "Z"))
+					]
+					
+					+SHorizontalBox::Slot()
+					[
+						SNew(SSegmentedControl<EMirrorDirection>)
+						.ToolTipText(LOCTEXT("MirrorDirectionTooltip", "The direction that determines what side of the plane to copy weights from."))
+						.Value_Lambda([this]()
+						{
+							return ToolProperties.IsValid() ? ToolProperties->MirrorDirection : EMirrorDirection::PositiveToNegative;
+						})
+						.OnValueChanged_Lambda([this](EMirrorDirection Mode)
+						{
+							if (ToolProperties.IsValid())
+							{
+								ToolProperties->MirrorDirection = Mode;
+							}
+						})
+						+ SSegmentedControl<EMirrorDirection>::Slot(EMirrorDirection::PositiveToNegative)
+						.Text(LOCTEXT("MirrorPosToNegLabel", "+ to -"))
+						+ SSegmentedControl<EMirrorDirection>::Slot(EMirrorDirection::NegativeToPositive)
+						.Text(LOCTEXT("MirrorNegToPosLabel", "- to +"))
+					]
+				]
+			]
+
+			+SVerticalBox::Slot()
+			.Padding(0.f, WeightEditVerticalPadding)
+			[
+				SNew(SBox)
+				[
+					SNew(SButton)
+					.HAlign(HAlign_Center)
+					.Text(LOCTEXT("MirrorWeightsButtonLabel", "Mirror"))
+					.ToolTipText(LOCTEXT("MirrorButtonTooltip",
+						"Weights are copied across the given plane in the given direction.\n"
+						"This command operates on the selected vertices."))
+					.OnClicked_Lambda([this]()
+					{
+						if (Tool.IsValid() && ToolProperties.IsValid())
+						{
+							Tool->MirrorWeights(ToolProperties->MirrorAxis, ToolProperties->MirrorDirection);
+						}
+						return FReply::Handled();
+					})
+				]
+			]
+		];
+
+		
+		
+
+	
+#endif // 0
+	}
+
+	void FVertexAttributePaintToolDetailCustomization::HideProperty(IDetailLayoutBuilder& DetailBuilder, FName PropertyName) const
+	{
+		TSharedRef<IPropertyHandle> Property = DetailBuilder.GetProperty(PropertyName);
+		DetailBuilder.HideProperty(Property);
+	}
+}
+
+#undef LOCTEXT_NAMESPACE
diff --git a/Engine/Plugins/Experimental/Dataflow/Source/DataflowEditor/Private/DataflowEditorTools/DataflowEditorVertexAttributePaintToolDetailCustomization.h b/Engine/Plugins/Experimental/Dataflow/Source/DataflowEditor/Private/DataflowEditorTools/DataflowEditorVertexAttributePaintToolDetailCustomization.h
new file mode 100644
--- /dev/null
+++ b/Engine/Plugins/Experimental/Dataflow/Source/DataflowEditor/Private/DataflowEditorTools/DataflowEditorVertexAttributePaintToolDetailCustomization.h
@@ -0,0 +1,70 @@
+// Copyright Epic Games, Inc. All Rights Reserved.
+
+#pragma once
+
+#include "CoreMinimal.h"
+#include "IDetailCustomization.h"
+#include "UObject/WeakObjectPtr.h"
+
+class IDetailLayoutBuilder;
+class IDetailCategoryBuilder;
+class FSlimHorizontalToolBarBuilder;
+class SWidget;
+class UDataflowEditorVertexAttributePaintTool;
+class UDataflowEditorVertexAttributePaintToolProperties;
+
+namespace UE::Dataflow::Editor
+{
+	class FVertexAttributePaintToolDetailCustomization : public IDetailCustomization
+	{
+	public:
+
+		virtual ~FVertexAttributePaintToolDetailCustomization() override;
+
+		static TSharedRef<IDetailCustomization> MakeInstance();
+
+		/** IDetailCustomization interface */
+		virtual void CustomizeDetails(IDetailLayoutBuilder& DetailBuilder) override;
+
+		void OnSelectionChanged();
+
+	private:
+
+		void AddEditingModeRow(IDetailCategoryBuilder& EditModeCategory) const;
+		void AddColorModeRow(IDetailCategoryBuilder& EditModeCategory) const;
+		void AddColorRampRow(IDetailCategoryBuilder& EditModeCategory) const;
+		
+		void AddBrushUI(IDetailLayoutBuilder& DetailBuilder) const;
+		void AddBrushModeRow(IDetailCategoryBuilder& BrushCategory) const;
+		void AddBrushFalloffModeRow(IDetailCategoryBuilder& BrushCategory) const;
+		void AddBrushRadiusRow(IDetailCategoryBuilder& BrushCategory) const;
+		void AddBrushValueRow(IDetailCategoryBuilder& BrushCategory) const;
+		void AddBrushStrengthRow(IDetailCategoryBuilder& BrushCategory) const;
+
+		void AddSelectionUI(IDetailLayoutBuilder& DetailBuilder) const;
+		TSharedRef<SWidget> MakeSelectionElementsToolbar() const;
+		TSharedRef<SWidget> MakeSelectionIsolationWidget() const;
+		TSharedRef<SWidget> MakeSelectionEditActionsToolbar() const;
+		void AddSelectionElementsRow(IDetailCategoryBuilder& BrushCategory) const;
+		void AddEmptySelectionWarningRow(IDetailCategoryBuilder& BrushCategory) const;
+		void MakeSelectionOperationRow(IDetailCategoryBuilder& EditValuesCategory, const FText& RowName, const TSharedRef<SWidget>& ButtonWidget, const TSharedRef<SWidget>& ValueWidget) const;
+		void MakeSelectionAddMultiplySliderOperationRow(IDetailCategoryBuilder& EditValuesCategory) const;
+		void MakeSelectionAddOperationRow(IDetailCategoryBuilder& EditValuesCategory) const;
+		void MakeSelectionReplaceOperationRow(IDetailCategoryBuilder& EditValuesCategory) const;
+		void MakeSelectionInvertOperationRow(IDetailCategoryBuilder& EditValuesCategory) const;
+		void MakeSelectionRelaxOperationRow(IDetailCategoryBuilder& EditValuesCategory) const;
+		void MakeSelectionPruneOperationRow(IDetailCategoryBuilder& EditValuesCategory) const;
+		void MakeSelectionCopyAndPasteRow(IDetailCategoryBuilder& EditValuesCategory) const;
+		
+		void HideProperty(IDetailLayoutBuilder& DetailBuilder, FName PropertyName) const;
+
+		IDetailLayoutBuilder* CurrentDetailBuilder = nullptr;
+		TWeakObjectPtr<UDataflowEditorVertexAttributePaintToolProperties> ToolProperties;
+		TWeakObjectPtr<UDataflowEditorVertexAttributePaintTool> Tool;
+
+		static float WeightSliderWidths;
+		static float WeightEditingLabelsPercent;
+		static float WeightEditVerticalPadding;
+		static float WeightEditHorizontalPadding;
+	};
+}
\ No newline at end of file
diff --git a/Engine/Plugins/Experimental/Dataflow/Source/DataflowEditor/Private/DataflowEditorTools/DataflowEditorWeightMapPaintBrushOps.cpp b/Engine/Plugins/Experimental/Dataflow/Source/DataflowEditor/Private/DataflowEditorTools/DataflowEditorWeightMapPaintBrushOps.cpp
--- a/Engine/Plugins/Experimental/Dataflow/Source/DataflowEditor/Private/DataflowEditorTools/DataflowEditorWeightMapPaintBrushOps.cpp
+++ b/Engine/Plugins/Experimental/Dataflow/Source/DataflowEditor/Private/DataflowEditorTools/DataflowEditorWeightMapPaintBrushOps.cpp
@@ -148,3 +148,63 @@ void FDataflowWeightMapSmoothBrushOp::ApplyStampByVertices(
 	}
 }
 
+////////////////////////////////////////////////////////////////////////////////////////////////////////////////////
+
+void FDataflowVertexAttributePaintBrushOp::ApplyStampByVertices(
+	const FDynamicMesh3* Mesh,
+	const FSculptBrushStamp& Stamp,
+	const TArray<int32>& Vertices,
+	TArray<double>& OutValues)
+{
+	check(OutValues.Num() == Vertices.Num());
+
+	if (UDataflowVertexAttributePaintBrushOpProps* Props = GetPropertySetAs<UDataflowVertexAttributePaintBrushOpProps>())
+	{
+		// special case for Relax option 
+		if (Props->EditOperation == EDataflowEditorToolEditOperation::Relax)
+		{
+			SmoothBrush.ApplyStampByVertices(Mesh, Stamp, Vertices, OutValues);
+			return;
+		}
+
+		// other operations
+		for (int32 VtxIdx = 0; VtxIdx < Vertices.Num(); ++VtxIdx)
+		{
+			if (bApplyRadiusLimit)
+			{
+				const FVector3d& StampPos = Stamp.LocalFrame.Origin;
+				const int32 MeshVertIdx = Vertices[VtxIdx];
+				const FVector3d VertexPos = Mesh->GetVertex(MeshVertIdx);
+				const double DistanceSquared = (VertexPos - StampPos).SquaredLength();
+				if (DistanceSquared >= Stamp.Radius * Stamp.Radius)
+				{
+					continue;
+				}
+			}
+			
+			const double ToolValue = Props->GetAttribute();
+			const double CurrentValue = OutValues[VtxIdx];
+			double NewValue = CurrentValue;
+
+			switch (Props->EditOperation)
+			{
+			case EDataflowEditorToolEditOperation::Add:
+				NewValue = CurrentValue + ToolValue;
+				break;
+			case EDataflowEditorToolEditOperation::Replace:
+				NewValue = ToolValue;
+				break;
+			case EDataflowEditorToolEditOperation::Multiply:
+				NewValue = CurrentValue * ToolValue;
+				break;
+			case EDataflowEditorToolEditOperation::Invert:
+				NewValue = 1.0 - CurrentValue;
+				break;
+			case EDataflowEditorToolEditOperation::Relax:
+				ensure(false); // already handled by SmoothBrush.ApplyStampByVertices
+			}
+
+			OutValues[VtxIdx] = FMath::Clamp(NewValue, 0.0, 1.0);
+		}
+	}
+}
\ No newline at end of file
diff --git a/Engine/Plugins/Experimental/Dataflow/Source/DataflowEditor/Private/DataflowEditorTools/DataflowEditorWeightMapPaintBrushOps.h b/Engine/Plugins/Experimental/Dataflow/Source/DataflowEditor/Private/DataflowEditorTools/DataflowEditorWeightMapPaintBrushOps.h
--- a/Engine/Plugins/Experimental/Dataflow/Source/DataflowEditor/Private/DataflowEditorTools/DataflowEditorWeightMapPaintBrushOps.h
+++ b/Engine/Plugins/Experimental/Dataflow/Source/DataflowEditor/Private/DataflowEditorTools/DataflowEditorWeightMapPaintBrushOps.h
@@ -3,11 +3,11 @@
 #pragma once
 
 #include "Sculpting/MeshBrushOpBase.h"
+#include "DataflowEditorTools/DataflowEditorToolEnums.h"
 
 #include "DataflowEditorWeightMapPaintBrushOps.generated.h"
 
-
-
+////////////////////////////////////////////////////////////////////////////////////////////////////////////////////
 
 class FDataflowMeshVertexWeightMapEditBrushOp : public FMeshSculptBrushOp
 {
@@ -36,7 +36,7 @@ public:
 };
 
 
-
+////////////////////////////////////////////////////////////////////////////////////////////////////////////////////
 //
 // Erase Brush
 // 
@@ -70,7 +70,7 @@ public:
 
 };
 
-
+////////////////////////////////////////////////////////////////////////////////////////////////////////////////////
 //
 // Paint Brush
 // 
@@ -116,6 +116,7 @@ public:
 
 };
 
+////////////////////////////////////////////////////////////////////////////////////////////////////////////////////
 //
 // Smooth Brush
 // 
@@ -153,3 +154,47 @@ public:
 		TArray<double>& NewAttributesOut) override;
 
 };
+
+////////////////////////////////////////////////////////////////////////////////////////////////////////////////////
+
+// this brush leverage EDataflowEditorToolEditOperation to decide how to paint the vertex attribute 
+UCLASS(MinimalAPI)
+class UDataflowVertexAttributePaintBrushOpProps : public UMeshSculptBrushOpProps
+{
+	GENERATED_BODY()
+
+public:
+	/** The Attribute that will be assigned to triangles within the brush region */
+	UPROPERTY(EditAnywhere, Category = PaintBrush, meta = (DisplayName = "Value", UIMin = 0))
+	double AttributeValue = 1;
+
+	UPROPERTY(EditAnywhere, Category = PaintBrush, meta = (DisplayName = "Strength", UIMin = "0.0", UIMax = "10.", ClampMin = "0.0", ClampMax = "10."))
+	float Strength = 1.0;
+
+	virtual double GetAttribute() { return AttributeValue; }
+	virtual float GetStrength() override { return Strength; }
+	virtual void SetStrength(float NewStrength) override { Strength = NewStrength; }
+	virtual float GetFalloff() override { return 0.0f; }
+
+	EDataflowEditorToolEditOperation EditOperation = EDataflowEditorToolEditOperation::Add;
+};
+
+class FDataflowVertexAttributePaintBrushOp : public FDataflowMeshVertexWeightMapEditBrushOp
+{
+public:
+
+	virtual bool IgnoreZeroMovements() const override
+	{
+		return false;
+	}
+
+	virtual void ApplyStampByVertices(
+		const FDynamicMesh3* Mesh,
+		const FSculptBrushStamp& Stamp,
+		const TArray<int32>& Vertices,
+		TArray<double>& NewAttributesOut) override;
+
+private:
+	// share the code with SmoothBrush when EditOperation is Relax
+	FDataflowWeightMapSmoothBrushOp SmoothBrush;
+};
diff --git a/Engine/Plugins/Experimental/Dataflow/Source/DataflowEditor/Public/DataflowEditorTools/DataflowEditorToolEnums.h b/Engine/Plugins/Experimental/Dataflow/Source/DataflowEditor/Public/DataflowEditorTools/DataflowEditorToolEnums.h
new file mode 100644
--- /dev/null
+++ b/Engine/Plugins/Experimental/Dataflow/Source/DataflowEditor/Public/DataflowEditorTools/DataflowEditorToolEnums.h
@@ -0,0 +1,44 @@
+// Copyright Epic Games, Inc. All Rights Reserved.
+
+#pragma once
+
+#include "DataflowEditorToolEnums.generated.h"
+
+UENUM()
+enum class EDataflowEditorToolEditMode : uint8
+{
+	Brush,
+	Mesh,
+};
+
+UENUM()
+enum class EDataflowEditorToolBrushAreaType : uint8
+{
+	Connected,
+	Volumetric
+};
+
+UENUM()
+enum class EDataflowEditorToolEditOperation : uint8
+{
+	Add,
+	Replace,
+	Multiply,
+	Invert,
+	Relax,
+};
+
+UENUM()
+enum class EDataflowEditorToolColorMode : uint8
+{
+	Greyscale,
+	Ramp,
+	FullMaterial,
+};
+
+UENUM()
+enum class EDataflowEditorToolVisibilityType : uint8
+{
+	None,
+	Unoccluded,
+};
