Dataflow vertex attribute Tool : fix uninitialized Color Mode in tool properties

diff --git a/Engine/Plugins/Experimental/Dataflow/Source/DataflowEditor/Private/DataflowEditorTools/DataflowEditorVertexAttributePaintTool.h b/Engine/Plugins/Experimental/Dataflow/Source/DataflowEditor/Private/DataflowEditorTools/DataflowEditorVertexAttributePaintTool.h
--- a/Engine/Plugins/Experimental/Dataflow/Source/DataflowEditor/Private/DataflowEditorTools/DataflowEditorVertexAttributePaintTool.h
+++ b/Engine/Plugins/Experimental/Dataflow/Source/DataflowEditor/Private/DataflowEditorTools/DataflowEditorVertexAttributePaintTool.h
@@ -158,7 +158,7 @@ struct FDataflowEditorVertexAttributePaintToolDisplayProperties
 
 	// Vertex color display mode (Greyscale, Ramp, ... )
 	UPROPERTY()
-	EDataflowEditorToolColorMode ColorMode;
+	EDataflowEditorToolColorMode ColorMode = EDataflowEditorToolColorMode::Greyscale;
 
 	// Color ramp to use when the color mode is set to "Ramp"
 	UPROPERTY()
