Dataflow edit skeleton bones tool

diff --git a/Engine/Plugins/Experimental/Dataflow/Dataflow.uplugin b/Engine/Plugins/Experimental/Dataflow/Dataflow.uplugin
--- a/Engine/Plugins/Experimental/Dataflow/Dataflow.uplugin
+++ b/Engine/Plugins/Experimental/Dataflow/Dataflow.uplugin
@@ -75,6 +75,11 @@
 		{
 			"Name": "EditorDataStorageFeatures",
 			"Enabled": true
+		},
+		{
+			"Name": "SkeletalMeshModelingTools",
+			"Enabled": true,
+			"TargetAllowList": [ "Editor" ]
 		}
 	]
 }
diff --git a/Engine/Plugins/Experimental/Dataflow/Source/DataflowEditor/DataflowEditor.Build.cs b/Engine/Plugins/Experimental/Dataflow/Source/DataflowEditor/DataflowEditor.Build.cs
--- a/Engine/Plugins/Experimental/Dataflow/Source/DataflowEditor/DataflowEditor.Build.cs
+++ b/Engine/Plugins/Experimental/Dataflow/Source/DataflowEditor/DataflowEditor.Build.cs
@@ -1,4 +1,5 @@
 // Copyright Epic Games, Inc. All Rights Reserved.
+using System.IO;
 
 namespace UnrealBuildTool.Rules
 {
@@ -6,11 +7,18 @@ namespace UnrealBuildTool.Rules
 	{
 		public DataflowEditor(ReadOnlyTargetRules Target) : base(Target)
 		{
+			PrivateIncludePaths.AddRange(
+				new string[] {
+					Path.Combine(GetModuleDirectory("GeometryCacheEd"), "Private"),
+					Path.Combine(GetModuleDirectory("SkeletalMeshModelingTools"), "Private")
+				}
+			);
 
 			PrivateDependencyModuleNames.AddRange(
 				new string[]
 				{
 					"TedsOutliner",
+					"SkeletalMeshModelingTools"
 				}
 			);
 
@@ -90,7 +98,11 @@ namespace UnrealBuildTool.Rules
 					"StructUtilsEditor",
 					"MessageLog",
 					"AppFramework",
-					"SkeletalMeshDescription"
+					"SkeletalMeshDescription",
+					"SequencerWidgets",
+					"TimeManagement",
+					"GeometryCacheEd",
+					"WidgetRegistration",
 				}
 			);
 		}
diff --git a/Engine/Plugins/Experimental/Dataflow/Source/DataflowEditor/Private/Dataflow/DataflowEditorMode.cpp b/Engine/Plugins/Experimental/Dataflow/Source/DataflowEditor/Private/Dataflow/DataflowEditorMode.cpp
--- a/Engine/Plugins/Experimental/Dataflow/Source/DataflowEditor/Private/Dataflow/DataflowEditorMode.cpp
+++ b/Engine/Plugins/Experimental/Dataflow/Source/DataflowEditor/Private/Dataflow/DataflowEditorMode.cpp
@@ -5,6 +5,8 @@
 #include "AssetEditorModeManager.h"
 #include "AttributeEditorTool.h"
 #include "BaseGizmos/TransformGizmoUtil.h"
+#include "SkeletalMeshEditorUtils.h"
+#include "SkeletalMeshGizmoUtils.h"
 #include "Components/DynamicMeshComponent.h"
 #include "ContextObjectStore.h"
 #include "DataflowEditorTools/DataflowEditorWeightMapPaintTool.h"
@@ -64,7 +66,7 @@ UDataflowEditorMode::UDataflowEditorMode()
 {
 	Info = FEditorModeInfo(
 		EM_DataflowEditorModeId,
-		LOCTEXT("DataflowEditorModeName", "Dataflow"),
+		LOCTEXT("DataflowEditorModeName", "Dataflow Toolbox"),
 		FSlateIcon(),
 		false);
 }
@@ -89,6 +91,8 @@ void UDataflowEditorMode::Enter()
 
 	// Register gizmo ContextObject for use inside interactive tools
 	UE::TransformGizmoUtil::RegisterTransformGizmoContextObject(GetInteractiveToolsContext());
+	UE::SkeletalMeshGizmoUtils::RegisterTransformGizmoContextObject(GetInteractiveToolsContext());
+	UE::SkeletalMeshEditorUtils::RegisterEditorContextObject(GetInteractiveToolsContext());
 
 	// Initialize view mode to a default
 	ConstructionViewMode = UE::Dataflow::FRenderingViewModeFactory::GetInstance().GetViewMode(UE::Dataflow::FDataflowConstruction3DViewMode::Name);
diff --git a/Engine/Plugins/Experimental/Dataflow/Source/DataflowEditor/Private/Dataflow/DataflowEditorModule.cpp b/Engine/Plugins/Experimental/Dataflow/Source/DataflowEditor/Private/Dataflow/DataflowEditorModule.cpp
--- a/Engine/Plugins/Experimental/Dataflow/Source/DataflowEditor/Private/Dataflow/DataflowEditorModule.cpp
+++ b/Engine/Plugins/Experimental/Dataflow/Source/DataflowEditor/Private/Dataflow/DataflowEditorModule.cpp
@@ -28,7 +28,9 @@
 #include "PropertyEditorModule.h"
 #include "EditorModeRegistry.h"
 #include "Dataflow/DataflowCollectionEditSkinWeightsNode.h"
+#include "Dataflow/DataflowCollectionEditSkeletonBonesNode.h"
 #include "DataflowEditorTools/DataflowEditorSkinWeightsPaintTool.h"
+#include "DataflowEditorTools/DataflowEditorEditSkeletonBonesTool.h"
 #include "DataflowEditorTools/DataflowEditorCorrectSkinWeightsNode.h"
 #include "Dataflow/DataflowColorRamp.h"
 #include "Dataflow/DataflowColorRampCustomization.h"
@@ -77,6 +79,23 @@ namespace UE::Dataflow::Private
 		}
 	};
 
+	class FDataflowEditorEditSkeletonBonesToolActionCommands : public TInteractiveToolCommands<FDataflowEditorEditSkeletonBonesToolActionCommands>
+	{
+	public:
+		FDataflowEditorEditSkeletonBonesToolActionCommands() : 
+			TInteractiveToolCommands<FDataflowEditorEditSkeletonBonesToolActionCommands>(
+				TEXT("DataflowEditorEditSkeletonBonesToolContext"),
+				LOCTEXT("DataflowEditorEditSkeletonBonesToolContext", "Dataflow skeleton edit Tool Context"),
+				NAME_None,
+				FAppStyle::GetAppStyleSetName())
+		{}
+
+		virtual void GetToolDefaultObjectList(TArray<UInteractiveTool*>& ToolCDOs) override
+		{
+			ToolCDOs.Add(GetMutableDefault<UDataflowEditorEditSkeletonBonesTool>());
+		}
+	};
+
 	class FDataflowToolActionCommandBindings : public UE::Dataflow::FDataflowToolRegistry::IDataflowToolActionCommands
 	{
 	public:
@@ -84,6 +103,7 @@ namespace UE::Dataflow::Private
 		{
 			FDataflowEditorWeightMapPaintToolActionCommands::Register();
 			FDataflowEditorSkinWeightPaintToolActionCommands::Register();
+			FDataflowEditorEditSkeletonBonesToolActionCommands::Register();
 		}
 
 		virtual void UnbindActiveCommands(const TSharedPtr<FUICommandList>& UICommandList) const override
@@ -93,6 +113,9 @@ namespace UE::Dataflow::Private
 
 			checkf(FDataflowEditorSkinWeightPaintToolActionCommands::IsRegistered(), TEXT("Expected SkinWeightPaintTool actions to have been registered"));
 			FDataflowEditorSkinWeightPaintToolActionCommands::Get().UnbindActiveCommands(UICommandList);
+
+			checkf(FDataflowEditorEditSkeletonBonesToolActionCommands::IsRegistered(), TEXT("Expected SkeletonEditTool actions to have been registered"));
+			FDataflowEditorEditSkeletonBonesToolActionCommands::Get().UnbindActiveCommands(UICommandList);
 		}
 
 		virtual void BindCommandsForCurrentTool(const TSharedPtr<FUICommandList>& UICommandList, UInteractiveTool* Tool) const override
@@ -107,6 +130,11 @@ namespace UE::Dataflow::Private
 				checkf(FDataflowEditorSkinWeightPaintToolActionCommands::IsRegistered(), TEXT("Expected SkinWeightPaintTool actions to have been registered"));
 				FDataflowEditorSkinWeightPaintToolActionCommands::Get().BindCommandsForCurrentTool(UICommandList, Tool);
 			}
+			if (ExactCast<UDataflowEditorEditSkeletonBonesTool>(Tool))
+			{
+				checkf(FDataflowEditorEditSkeletonBonesToolActionCommands::IsRegistered(), TEXT("Expected SkeletonEditTool actions to have been registered"));
+				FDataflowEditorEditSkeletonBonesToolActionCommands::Get().BindCommandsForCurrentTool(UICommandList, Tool);
+			}
 		}
 	};
 }
@@ -145,6 +173,10 @@ void FDataflowEditorModule::StartupModule()
 	ToolRegistry.AddNodeToToolMapping(FDataflowCollectionEditSkinWeightsNode::StaticType(), 
 		NewObject<UDataflowEditorSkinWeightsPaintToolBuilder>(), Actions,
 		FSlateIcon(FName("DataflowEditorStyle"), FName("Dataflow.EditSkinWeights")), LOCTEXT("AddSkinWeightNodeButtonText", "Edit Skin Weights"));
+
+	ToolRegistry.AddNodeToToolMapping(FDataflowCollectionEditSkeletonBonesNode::StaticType(), 
+		NewObject<UDataflowEditorEditSkeletonBonesToolBuilder>(), Actions,
+		FSlateIcon(FName("DataflowEditorStyle"), FName("Dataflow.EditSkeletonBones")), LOCTEXT("AddSkeletonEditNodeButtonText", "Edit Skeleton Bones"));
 		
 	DATAFLOW_NODE_REGISTER_CREATION_FACTORY(FDataflowCorrectSkinWeightsNode);
 	DATAFLOW_NODE_REGISTER_CREATION_FACTORY(FDataflowGetSkinningSelectionNode);
@@ -179,6 +211,7 @@ void FDataflowEditorModule::ShutdownModule()
 	UE::Dataflow::FDataflowToolRegistry& ToolRegistry = UE::Dataflow::FDataflowToolRegistry::Get();
 	ToolRegistry.RemoveNodeToToolMapping(FDataflowCollectionAddScalarVertexPropertyNode::StaticType());
 	ToolRegistry.RemoveNodeToToolMapping(FDataflowCollectionEditSkinWeightsNode::StaticType());
+	ToolRegistry.RemoveNodeToToolMapping(FDataflowCollectionEditSkeletonBonesNode::StaticType());
 
 	UE::Dataflow::FDataflowConstructionVisualizationRegistry& ConstructionVisualizationRegistry = UE::Dataflow::FDataflowConstructionVisualizationRegistry::GetInstance();
 	ConstructionVisualizationRegistry.DeregisterVisualization(UE::Dataflow::FMeshStatsConstructionVisualization::Name);
diff --git a/Engine/Plugins/Experimental/Dataflow/Source/DataflowEditor/Private/DataflowEditorTools/DataflowEditorEditSkeletonBonesTool.cpp b/Engine/Plugins/Experimental/Dataflow/Source/DataflowEditor/Private/DataflowEditorTools/DataflowEditorEditSkeletonBonesTool.cpp
new file mode 100644
--- /dev/null
+++ b/Engine/Plugins/Experimental/Dataflow/Source/DataflowEditor/Private/DataflowEditorTools/DataflowEditorEditSkeletonBonesTool.cpp
@@ -0,0 +1,164 @@
+﻿// Copyright Epic Games, Inc. All Rights Reserved.
+
+#include "DataflowEditorTools/DataflowEditorEditSkeletonBonesTool.h"
+
+#include "ContextObjectStore.h"
+#include "Dataflow/DataflowCollectionEditSkeletonBonesNode.h"
+#include "Dataflow/DataflowContent.h"
+#include "Dataflow/DataflowObjectInterface.h"
+#include "Dataflow/DataflowSNode.h"
+#include "Engine/World.h"
+#include "InteractiveToolManager.h"
+#include "ToolTargetManager.h"
+#include "ModelingToolTargetUtil.h"
+#include "Components/SkeletalMeshComponent.h"
+#include "DynamicMesh/DynamicVertexSkinWeightsAttribute.h"
+#include "TargetInterfaces/MeshDescriptionCommitter.h"
+#include "TargetInterfaces/MeshDescriptionProvider.h"
+#include "TargetInterfaces/PrimitiveComponentBackedTarget.h"
+
+#include UE_INLINE_GENERATED_CPP_BY_NAME(DataflowEditorEditSkeletonBonesTool)
+
+struct FDataflowCollectionEditSkeletonBonesNode;
+using namespace UE::Geometry;
+
+#define LOCTEXT_NAMESPACE "UDataflowEditorEditSkeletonBonesTool"
+
+void UDataflowEditorEditSkeletonBonesToolBuilder::GetSupportedConstructionViewModes(const UDataflowContextObject& ContextObject, TArray<const UE::Dataflow::IDataflowConstructionViewMode*>& Modes) const
+{}
+
+bool UDataflowEditorEditSkeletonBonesToolBuilder::CanBuildTool(const FToolBuilderState& SceneState) const
+{
+	auto HasManagedArrayCollection = [](const FDataflowNode* InDataflowNode, const TSharedPtr<UE::Dataflow::FEngineContext> Context)
+	{
+		if (InDataflowNode && Context)
+		{
+			for (const FDataflowOutput* const Output : InDataflowNode->GetOutputs())
+			{
+				if (Output->GetType() == FName("FManagedArrayCollection"))
+				{
+					return true;
+				}
+			}
+		}
+
+		return false;
+	};
+
+	if (USkeletonEditingToolBuilder::CanBuildTool(SceneState))
+	{
+		if (SceneState.SelectedComponents.Num() == 1 && SceneState.SelectedComponents[0]->IsA<USkeletalMeshComponent>())
+		{
+			if (UDataflowBaseContent* ContextObject = SceneState.ToolManager->GetContextObjectStore()->FindContext<UDataflowBaseContent>())
+			{
+				if (const TSharedPtr<UE::Dataflow::FEngineContext> EvaluationContext = ContextObject->GetDataflowContext())
+				{
+					if (const FDataflowNode* PrimarySelection = ContextObject->GetSelectedNodeOfType<FDataflowCollectionEditSkeletonBonesNode>())
+					{
+						return HasManagedArrayCollection(PrimarySelection, EvaluationContext);
+					}
+				}
+			}
+		}
+	}
+	return false;
+}
+
+const FToolTargetTypeRequirements& UDataflowEditorEditSkeletonBonesToolBuilder::GetTargetRequirements() const
+{
+	return USkeletonEditingToolBuilder::GetTargetRequirements();
+}
+
+UInteractiveTool* UDataflowEditorEditSkeletonBonesToolBuilder::BuildTool(const FToolBuilderState& SceneState) const
+{
+	UDataflowEditorEditSkeletonBonesTool* EditTool = NewObject<UDataflowEditorEditSkeletonBonesTool>(SceneState.ToolManager);
+
+	if (UDataflowBaseContent* ContextObject = SceneState.ToolManager->GetContextObjectStore()->FindContext<UDataflowBaseContent>())
+	{
+		if (FDataflowNode* PrimarySelection = ContextObject->GetSelectedNodeOfType<FDataflowCollectionEditSkeletonBonesNode>())
+		{
+			FDataflowCollectionEditSkeletonBonesNode* EditSkeletonBonesNode =
+				StaticCast<FDataflowCollectionEditSkeletonBonesNode*>(PrimarySelection);
+			TSharedPtr<FDataflowSkeletalMeshEditorBinding> SkeletalMeshBinding = MakeShared<FDataflowSkeletalMeshEditorBinding>();
+			
+			EditTool->EditSkeletonBonesNode = EditSkeletonBonesNode;
+			EditTool->BindTo(SkeletalMeshBinding);
+			
+			EditSkeletonBonesNode->OnBoneSelectionChanged.AddLambda(
+				[EditTool, SkeletalMeshBinding](const TArray<FName>& BoneNames)
+			{
+				EditTool->GetNotifier().HandleNotification(BoneNames, ESkeletalMeshNotifyType::BonesSelected);
+				SkeletalMeshBinding->BoneSelection = BoneNames;
+			});
+			EditTool->SetDataflowEditorContextObject(ContextObject);
+		}
+	}
+	
+	UToolTarget* Target = SceneState.TargetManager->BuildFirstSelectedTargetable(SceneState, GetTargetRequirements());
+	EditTool->SetTarget(Target);
+	EditTool->Init(SceneState);
+	
+	return EditTool;
+}
+
+void UDataflowEditorEditSkeletonBonesTool::Setup()
+{
+	USkeletonEditingTool::Setup();
+}
+
+void UDataflowEditorEditSkeletonBonesTool::Shutdown(EToolShutdownType ShutdownType)
+{
+	USkeletonEditingTool::Shutdown(ShutdownType);
+
+	if(Target && ShutdownType == EToolShutdownType::Accept)
+	{
+		EditSkeletonBonesNode->Invalidate();
+
+		// Avoid rebuilding the skeletal mesh after updating the skelton bones
+		EditSkeletonBonesNode->ValidateSkeletalMeshes();
+	}
+}
+
+FDataflowSkeletalMeshEditorBinding::FDataflowSkeletalMeshEditorBinding()
+	: ISkeletalMeshEditorBinding(),
+	DataflowNotifier()
+{}
+
+ISkeletalMeshNotifier& FDataflowSkeletalMeshEditorBinding::GetNotifier()
+{
+	return DataflowNotifier;	
+}
+
+ISkeletalMeshEditorBinding::NameFunction FDataflowSkeletalMeshEditorBinding::GetNameFunction()
+{
+	return [this](HHitProxy* InHitProxy) -> TOptional<FName>
+	{
+		if (const HDataflowElementHitProxy* ElementHitProxy = HitProxyCast<HDataflowElementHitProxy>(InHitProxy))
+		{
+			return ElementHitProxy->ElementName;
+		}
+
+		static const TOptional<FName> Dummy;
+		return Dummy;
+	};
+}
+
+TArray<FName> FDataflowSkeletalMeshEditorBinding::GetSelectedBones() const
+{
+	return BoneSelection;
+}
+
+TArray<FName> FDataflowSkeletalMeshEditorBinding::GetSelectedMorphTargets() const
+{
+	TArray<FName> Selection;
+	return Selection;
+}
+
+FDataflowSkeletalMeshEditorNotifier::FDataflowSkeletalMeshEditorNotifier()
+	: ISkeletalMeshNotifier()
+{}
+
+void FDataflowSkeletalMeshEditorNotifier::HandleNotification(const TArray<FName>& BoneNames, const ESkeletalMeshNotifyType InNotifyType)
+{}
+
+#undef LOCTEXT_NAMESPACE
diff --git a/Engine/Plugins/Experimental/Dataflow/Source/DataflowEditor/Public/DataflowEditorTools/DataflowEditorEditSkeletonBonesTool.h b/Engine/Plugins/Experimental/Dataflow/Source/DataflowEditor/Public/DataflowEditorTools/DataflowEditorEditSkeletonBonesTool.h
new file mode 100644
--- /dev/null
+++ b/Engine/Plugins/Experimental/Dataflow/Source/DataflowEditor/Public/DataflowEditorTools/DataflowEditorEditSkeletonBonesTool.h
@@ -0,0 +1,105 @@
+﻿// Copyright Epic Games, Inc. All Rights Reserved.
+
+#pragma once
+
+#include "CoreMinimal.h"
+#include "HAL/Platform.h"
+#include "DataflowEditorTools/DataflowEditorToolBuilder.h"
+#include "SkeletalMesh/SkeletonEditingTool.h"
+
+#include "DataflowEditorEditSkeletonBonesTool.generated.h"
+
+#define UE_API DATAFLOWEDITOR_API
+
+class UDataflowContextObject;
+class UDataflowEditorMode;
+
+/**
+ * Dataflow edit skeleton tool builder
+ */
+UCLASS(MinimalAPI)
+class UDataflowEditorEditSkeletonBonesToolBuilder : public USkeletonEditingToolBuilder, public IDataflowEditorToolBuilder
+{
+	GENERATED_BODY()
+
+private:
+	//~ Begin IDataflowEditorToolBuilder interface
+	UE_API virtual void GetSupportedConstructionViewModes(const UDataflowContextObject& ContextObject, TArray<const UE::Dataflow::IDataflowConstructionViewMode*>& Modes) const override;
+	virtual bool CanSetConstructionViewWireframeActive() const override { return false; }
+	//~ End IDataflowEditorToolBuilder interface
+
+	//~ Begin USkeletonEditToolBuilder interface
+	UE_API virtual bool CanBuildTool(const FToolBuilderState& SceneState) const override;
+	UE_API virtual UInteractiveTool* BuildTool(const FToolBuilderState& SceneState) const override;
+	//~ End USkeletonEditToolBuilder interface
+
+protected:
+	UE_API virtual const FToolTargetTypeRequirements& GetTargetRequirements() const override;
+};
+
+/**
+ * Dataflow edit skeleton tool
+ */
+UCLASS(MinimalAPI)
+class UDataflowEditorEditSkeletonBonesTool : public USkeletonEditingTool
+{
+	GENERATED_BODY()
+	
+	//~ Begin USkeletonEditingTool interface
+	UE_API virtual void Setup() override;
+	UE_API virtual void Shutdown(EToolShutdownType ShutdownType) override;
+	//~ End USkeletonEditingTool interface
+
+	friend class UDataflowEditorEditSkeletonBonesToolBuilder;
+
+private:
+
+	/** Set the dataflow context object */
+	void SetDataflowEditorContextObject(TObjectPtr<UDataflowContextObject> InDataflowEditorContextObject)
+	{
+		DataflowEditorContextObject = InDataflowEditorContextObject;
+	}
+	
+	/** Edit Skeleton node associated with that tool */
+	struct FDataflowCollectionEditSkeletonBonesNode* EditSkeletonBonesNode = nullptr;
+
+	/** Dataflow context object to be used to retrieve the node context data*/
+	TObjectPtr<UDataflowContextObject> DataflowEditorContextObject;
+};
+
+/** Notifier when skeleton selection is changing */
+class FDataflowSkeletalMeshEditorNotifier : public ISkeletalMeshNotifier
+{
+public:
+	FDataflowSkeletalMeshEditorNotifier();
+
+	//~ Begin ISkeletalMeshEditorBinding interface
+	virtual void HandleNotification(const TArray<FName>& BoneNames, const ESkeletalMeshNotifyType InNotifyType) override;
+	//~ End ISkeletalMeshEditorBinding interface
+
+private:
+};
+
+/** Binding used to retrieve the bone names (morph targets not used)*/
+class FDataflowSkeletalMeshEditorBinding : public ISkeletalMeshEditorBinding
+{
+public:
+	FDataflowSkeletalMeshEditorBinding();
+
+	//~ Begin ISkeletalMeshEditorBinding interface
+	virtual ISkeletalMeshNotifier& GetNotifier() override;
+	virtual NameFunction GetNameFunction() override;
+	virtual TArray<FName> GetSelectedBones() const override;
+	virtual TArray<FName> GetSelectedMorphTargets() const override;
+	//~ End ISkeletalMeshEditorBinding interface
+
+	TArray<FName> BoneSelection;
+
+private :
+
+	/** Notifier used when the tool selection is changing */
+	FDataflowSkeletalMeshEditorNotifier DataflowNotifier;
+};
+
+
+#undef UE_API
diff --git a/Engine/Plugins/Experimental/Dataflow/Source/DataflowNodes/Private/Dataflow/DataflowCollectionEditSkeletonBonesNode.cpp b/Engine/Plugins/Experimental/Dataflow/Source/DataflowNodes/Private/Dataflow/DataflowCollectionEditSkeletonBonesNode.cpp
new file mode 100644
--- /dev/null
+++ b/Engine/Plugins/Experimental/Dataflow/Source/DataflowNodes/Private/Dataflow/DataflowCollectionEditSkeletonBonesNode.cpp
@@ -0,0 +1,145 @@
+// Copyright Epic Games, Inc. All Rights Reserved.
+
+#include "Dataflow/DataflowCollectionEditSkeletonBonesNode.h"
+#include "Dataflow/DataflowInputOutput.h"
+#include "Dataflow/DataflowDebugDrawObject.h"
+#include "DataflowCollectionSkeletalMeshUtils.h"
+#include "GeometryCollection/Facades/CollectionRenderingFacade.h"
+#include "Rendering/SkeletalMeshModel.h"
+#include "Engine/SkeletalMesh.h"
+#include "Animation/Skeleton.h"
+#include "Components/SkeletalMeshComponent.h"
+
+#if WITH_EDITOR
+#include "StaticToSkeletalMeshConverter.h"
+#endif
+
+#include UE_INLINE_GENERATED_CPP_BY_NAME(DataflowCollectionEditSkeletonBonesNode)
+
+#define LOCTEXT_NAMESPACE "DataflowCollectionEditSkeleton"
+
+FDataflowCollectionEditSkeletonBonesNode::FDataflowCollectionEditSkeletonBonesNode(const UE::Dataflow::FNodeParameters& InParam, FGuid InGuid)
+	: FDataflowPrimitiveNode(InParam, InGuid)
+{
+	RegisterInputConnection(&Collection);
+	RegisterInputConnection(&Skeleton);
+	RegisterOutputConnection(&Collection, &Collection);
+	RegisterOutputConnection(&Skeleton, &Skeleton);
+}
+
+TArray<UE::Dataflow::FRenderingParameter> FDataflowCollectionEditSkeletonBonesNode::GetRenderParametersImpl() const
+{
+	return FDataflowAddScalarVertexPropertyCallbackRegistry::Get().GetRenderingParameters(VertexGroup.Name);
+}
+
+void FDataflowCollectionEditSkeletonBonesNode::AddPrimitiveComponents(UE::Dataflow::FContext& Context, const TSharedPtr<const FManagedArrayCollection> RenderCollection, TObjectPtr<UObject> NodeOwner, 
+	TObjectPtr<AActor> RootActor, TArray<TObjectPtr<UPrimitiveComponent>>& PrimitiveComponents)  
+{
+	if(RootActor)
+	{
+		if(TObjectPtr<USkeleton> TransientSkeleton = UpdateToolSkeleton(Context))
+		{
+			GeometryCollection::Facades::FRenderingFacade RenderingFacade(*RenderCollection);
+			const int32 NumGeometry = RenderingFacade.IsValid() ? RenderingFacade.NumGeometry() : 0;
+		
+			const bool bNeedsConstruction = (SkeletalMeshes.Num() != NumGeometry) || !bValidSkeletalMeshes;
+		
+			if(SkeletalMeshes.Num() != NumGeometry)
+			{
+				SkeletalMeshes.Reset();
+				for(int32 GeometryIndex = 0; GeometryIndex < NumGeometry; ++GeometryIndex)
+				{
+					FString SkeletalMeshName = FString("SK_DataflowSkeletalMesh_") + FString::FromInt(GeometryIndex);
+					SkeletalMeshName = MakeUniqueObjectName(NodeOwner, USkeletalMesh::StaticClass(), *SkeletalMeshName, EUniqueObjectNameOptions::GloballyUnique).ToString();
+					SkeletalMeshes.Add(NewObject<USkeletalMesh>(NodeOwner, FName(*SkeletalMeshName), RF_Transient));
+				}
+			}
+			if(bNeedsConstruction)
+			{
+				bValidSkeletalMeshes = UE::Dataflow::Private::BuildSkeletalMeshes(SkeletalMeshes, RenderCollection, TransientSkeleton);
+				if(!bValidSkeletalMeshes)
+				{
+					SkeletalMeshes.Reset();
+				}
+			}
+			if(!SkeletalMeshes.IsEmpty())
+			{
+				for(int32 GeometryIndex = 0; GeometryIndex < NumGeometry; ++GeometryIndex)
+				{
+					FName SkeletalMeshName(FString("Dataflow_SkeletalMesh") + FString::FromInt(GeometryIndex));
+					USkeletalMeshComponent* SkeletalMeshComponent = NewObject<USkeletalMeshComponent>(RootActor, SkeletalMeshName);
+					SkeletalMeshes[GeometryIndex]->SetSkeleton(TransientSkeleton);
+					SkeletalMeshComponent->SetSkeletalMesh(SkeletalMeshes[GeometryIndex]);
+					PrimitiveComponents.Add(SkeletalMeshComponent);
+				}
+			}
+		}
+	}
+}
+
+TObjectPtr<USkeleton> FDataflowCollectionEditSkeletonBonesNode::UpdateToolSkeleton(UE::Dataflow::FContext& Context)
+{
+	if (ToolSkeleton == nullptr)
+	{
+		TObjectPtr<USkeleton> InputSkeleton = GetValue(Context, &Skeleton, Skeleton);
+		ToolSkeleton = DuplicateObject<USkeleton>(InputSkeleton, nullptr);
+	}
+	return ToolSkeleton;
+}
+
+void FDataflowCollectionEditSkeletonBonesNode::Evaluate(UE::Dataflow::FContext& Context, const FDataflowOutput* Out) const
+{
+	using namespace UE::Dataflow;
+	
+	if (Out->IsA(&Collection))
+	{
+		SafeForwardInput(Context, &Collection, &Collection);
+	}
+	else if (Out->IsA(&Skeleton))
+	{
+		if(ToolSkeleton)
+		{ 
+			SetValue(Context, ToolSkeleton, &Skeleton);
+		}
+		else
+		{
+			SafeForwardInput(Context, &Skeleton, &Skeleton);
+		}
+	}
+}
+
+void FDataflowCollectionEditSkeletonBonesNode::OnInvalidate()
+{
+	Super::OnInvalidate();
+	
+	bValidSkeletalMeshes = false;
+}
+
+#if WITH_EDITOR
+
+void FDataflowCollectionEditSkeletonBonesNode::DebugDraw(UE::Dataflow::FContext& Context,
+		IDataflowDebugDrawInterface& DataflowRenderingInterface,
+		const FDebugDrawParameters& DebugDrawParameters) const
+{
+	if(TObjectPtr<USkeleton> InputSkeleton = GetValue(Context, &Skeleton, Skeleton))
+	{
+		TRefCountPtr<FDataflowDebugDrawSkeletonObject> SkeletonObject = MakeDebugDrawObject<FDataflowDebugDrawSkeletonObject>(
+			DataflowRenderingInterface.ModifyDataflowElements(), InputSkeleton->GetReferenceSkeleton(), false);
+		
+		DataflowRenderingInterface.DrawObject(TRefCountPtr<IDataflowDebugDrawObject>(SkeletonObject));
+		
+		SkeletonObject->OnBoneSelectionChanged.AddLambda([this](const TArray<FName>& BoneNames)
+		{
+			OnBoneSelectionChanged.Broadcast(BoneNames);
+		});
+	}
+}
+
+bool FDataflowCollectionEditSkeletonBonesNode::CanDebugDrawViewMode(const FName& ViewModeName) const
+{
+	return true;
+}
+
+#endif
+
+#undef LOCTEXT_NAMESPACE
diff --git a/Engine/Plugins/Experimental/Dataflow/Source/DataflowNodes/Private/Dataflow/DataflowCollectionEditSkinWeightsNode.cpp b/Engine/Plugins/Experimental/Dataflow/Source/DataflowNodes/Private/Dataflow/DataflowCollectionEditSkinWeightsNode.cpp
--- a/Engine/Plugins/Experimental/Dataflow/Source/DataflowNodes/Private/Dataflow/DataflowCollectionEditSkinWeightsNode.cpp
+++ b/Engine/Plugins/Experimental/Dataflow/Source/DataflowNodes/Private/Dataflow/DataflowCollectionEditSkinWeightsNode.cpp
@@ -5,7 +5,7 @@
 #include "Dataflow/DataflowDebugDrawObject.h"
 #include "Dataflow/DataflowSkeletalMeshNodes.h"
 #include "Dataflow/CollectionRenderingPatternUtility.h"
-
+#include "DataflowCollectionSkeletalMeshUtils.h"
 #include "DynamicMeshToMeshDescription.h"
 #include "MeshConversionOptions.h"
 #include "SkeletalMeshAttributes.h"
@@ -15,8 +15,6 @@
 #include "Engine/SkeletalMesh.h"
 #include "Animation/Skeleton.h"
 #include "Components/SkeletalMeshComponent.h"
-#include "Materials/Material.h"
-#include "MaterialDomain.h"
 
 #if WITH_EDITOR
 #include "StaticToSkeletalMeshConverter.h"
@@ -28,91 +26,6 @@
 
 namespace UE::Dataflow::Private
 {
-
-static bool CreateSkeletalMesh(USkeletalMesh* SkeletalMesh, const TArray<UE::Geometry::FDynamicMesh3>& DynamicMeshes, const FReferenceSkeleton& ReferenceSkeleton)
-{
-#if WITH_EDITOR
-	TArray<const FMeshDescription*> MeshDescriptions;
-
-#if WITH_EDITORONLY_DATA
-	SkeletalMesh->PreEditChange(nullptr);
-	SkeletalMesh->GetImportedModel()->LODModels.Empty();
-#endif
-	SkeletalMesh->SetNumSourceModels(0);
-
-	TArray<FMeshDescription> LocalDescriptions;
-	for (const UE::Geometry::FDynamicMesh3& DynamicMesh : DynamicMeshes)
-	{
-		// Create a mesh description
-		FMeshDescription& MeshDescription = LocalDescriptions.AddDefaulted_GetRef();
-	
-		// Add skeletal mesh attributes to the mesh description
-		FSkeletalMeshAttributes Attributes(MeshDescription);
-		Attributes.Register();
-
-		// Convert dynamic mesh to the mesh description
-		FConversionToMeshDescriptionOptions ConverterOptions;
-		FDynamicMeshToMeshDescription Converter(ConverterOptions);
-		Converter.Convert(&DynamicMesh, MeshDescription, false);
-
-		// Add the created description to the list
-		MeshDescriptions.Add(&MeshDescription);
-	}
-
-	TArray<FSkeletalMaterial> Materials;
-	TConstArrayView<FSkeletalMaterial> MaterialView;
-
-	// ensure there is at least one material
-	if (MaterialView.IsEmpty())
-	{
-		Materials.Add( UMaterial::GetDefaultMaterial(MD_Surface));
-		MaterialView = Materials;
-	}
-
-	static constexpr bool bRecomputeTangents = false;
-	static constexpr bool bRecomputeNormals = false;
-	static constexpr bool bCacheOptimize = false;
-
-	if (!FStaticToSkeletalMeshConverter::InitializeSkeletalMeshFromMeshDescriptions(
-		SkeletalMesh, MeshDescriptions, MaterialView, ReferenceSkeleton, bRecomputeNormals, bRecomputeTangents, bCacheOptimize))
-	{
-		return false;
-	}
-	return true;
-#else
-	return false;
-#endif
-}
-
-static bool BuildSkeletalMeshes(TArray<TObjectPtr<USkeletalMesh>>& SkeletalMeshes, 
-	const TSharedPtr<const FManagedArrayCollection>& RenderCollection, const TObjectPtr<USkeleton>& ObjectSkeleton)
-{
-	if(RenderCollection.IsValid() && ObjectSkeleton)
-	{
-		GeometryCollection::Facades::FRenderingFacade Facade(*RenderCollection);
-		if(Facade.IsValid())
-		{
-			bool bValidSkeletalMeshes = true;
-			const int32 NumGeometry = Facade.NumGeometry();
-
-			if(NumGeometry == SkeletalMeshes.Num())
-			{
-				for (int32 MeshIndex = 0; MeshIndex < NumGeometry; ++MeshIndex)
-				{
-					UE::Geometry::FDynamicMesh3 DynamicMesh;
-					UE::Dataflow::Conversion::RenderingFacadeToDynamicMesh(Facade, MeshIndex, DynamicMesh);
-
-					if(!CreateSkeletalMesh(SkeletalMeshes[MeshIndex], {DynamicMesh}, ObjectSkeleton->GetReferenceSkeleton()))
-					{
-						bValidSkeletalMeshes = false;
-					}
-				}
-			}
-			return bValidSkeletalMeshes;
-		}
-	}
-	return false;
-}
 	
 template<typename ScalarType, typename VectorType, int32 NumComponents>
 bool SetAttributeValues(FManagedArrayCollection& SelectedCollection,
@@ -286,7 +199,8 @@ void FDataflowCollectionEditSkinWeightsNode::AddPrimitiveComponents(UE::Dataflow
 			SkeletalMeshes.Reset();
 			for(int32 GeometryIndex = 0; GeometryIndex < NumGeometry; ++GeometryIndex)
 			{
-				FString SkeletalMeshName = FString("SK_DataflowSkeletalMesh_") + FString::FromInt(GeometryIndex);
+				const FString GeometryName = RenderingFacade.GetGeometryName()[GeometryIndex];
+				FString SkeletalMeshName = GeometryName.IsEmpty() ? FString("SK_DataflowSkeletalMesh_") + FString::FromInt(GeometryIndex) : GeometryName;
 				SkeletalMeshName = MakeUniqueObjectName(NodeOwner, USkeletalMesh::StaticClass(), *SkeletalMeshName, EUniqueObjectNameOptions::GloballyUnique).ToString();
 				SkeletalMeshes.Add(NewObject<USkeletalMesh>(NodeOwner, FName(*SkeletalMeshName), RF_Transient));
 			}
diff --git a/Engine/Plugins/Experimental/Dataflow/Source/DataflowNodes/Private/Dataflow/DataflowCollectionSkeletalMeshUtils.h b/Engine/Plugins/Experimental/Dataflow/Source/DataflowNodes/Private/Dataflow/DataflowCollectionSkeletalMeshUtils.h
new file mode 100644
--- /dev/null
+++ b/Engine/Plugins/Experimental/Dataflow/Source/DataflowNodes/Private/Dataflow/DataflowCollectionSkeletalMeshUtils.h
@@ -0,0 +1,110 @@
+﻿// Copyright Epic Games, Inc. All Rights Reserved.
+#pragma once
+
+#include "Dataflow/CollectionRenderingPatternUtility.h"
+#include "DynamicMeshToMeshDescription.h"
+#include "MeshConversionOptions.h"
+#include "SkeletalMeshAttributes.h"
+#include "DynamicMesh/DynamicMesh3.h"
+#include "GeometryCollection/Facades/CollectionRenderingFacade.h"
+#include "Rendering/SkeletalMeshModel.h"
+#include "Engine/SkeletalMesh.h"
+#include "Animation/Skeleton.h"
+#include "Components/SkeletalMeshComponent.h"
+#include "Materials/Material.h"
+#include "MaterialDomain.h"
+
+#if WITH_EDITOR
+#include "StaticToSkeletalMeshConverter.h"
+#endif
+
+namespace UE::Dataflow::Private
+{
+
+static bool CreateSkeletalMesh(USkeletalMesh* SkeletalMesh, const TArray<UE::Geometry::FDynamicMesh3>& DynamicMeshes, const FReferenceSkeleton& ReferenceSkeleton)
+{
+#if WITH_EDITOR
+	TArray<const FMeshDescription*> MeshDescriptions;
+
+#if WITH_EDITORONLY_DATA
+	SkeletalMesh->PreEditChange(nullptr);
+	SkeletalMesh->GetImportedModel()->LODModels.Empty();
+#endif
+	SkeletalMesh->SetNumSourceModels(0);
+
+	TArray<FMeshDescription> LocalDescriptions;
+	for (const UE::Geometry::FDynamicMesh3& DynamicMesh : DynamicMeshes)
+	{
+		// Create a mesh description
+		FMeshDescription& MeshDescription = LocalDescriptions.AddDefaulted_GetRef();
+	
+		// Add skeletal mesh attributes to the mesh description
+		FSkeletalMeshAttributes Attributes(MeshDescription);
+		Attributes.Register();
+
+		// Convert dynamic mesh to the mesh description
+		FConversionToMeshDescriptionOptions ConverterOptions;
+		FDynamicMeshToMeshDescription Converter(ConverterOptions);
+		Converter.Convert(&DynamicMesh, MeshDescription, false);
+
+		// Add the created description to the list
+		MeshDescriptions.Add(&MeshDescription);
+	}
+
+	TArray<FSkeletalMaterial> Materials;
+	TConstArrayView<FSkeletalMaterial> MaterialView;
+
+	// ensure there is at least one material
+	if (MaterialView.IsEmpty())
+	{
+		Materials.Add( UMaterial::GetDefaultMaterial(MD_Surface));
+		MaterialView = Materials;
+	}
+
+	static constexpr bool bRecomputeTangents = false;
+	static constexpr bool bRecomputeNormals = false;
+	static constexpr bool bCacheOptimize = false;
+
+	if (!FStaticToSkeletalMeshConverter::InitializeSkeletalMeshFromMeshDescriptions(
+		SkeletalMesh, MeshDescriptions, MaterialView, ReferenceSkeleton, bRecomputeNormals, bRecomputeTangents, bCacheOptimize))
+	{
+		return false;
+	}
+	return true;
+#else
+	return false;
+#endif
+}
+
+static bool BuildSkeletalMeshes(TArray<TObjectPtr<USkeletalMesh>>& SkeletalMeshes, 
+	const TSharedPtr<const FManagedArrayCollection>& RenderCollection, const TObjectPtr<USkeleton>& ObjectSkeleton)
+{
+	if(RenderCollection.IsValid() && ObjectSkeleton)
+	{
+		GeometryCollection::Facades::FRenderingFacade Facade(*RenderCollection);
+		if(Facade.IsValid())
+		{
+			bool bValidSkeletalMeshes = true;
+			const int32 NumGeometry = Facade.NumGeometry();
+
+			if(NumGeometry == SkeletalMeshes.Num())
+			{
+				for (int32 MeshIndex = 0; MeshIndex < NumGeometry; ++MeshIndex)
+				{
+					UE::Geometry::FDynamicMesh3 DynamicMesh;
+					UE::Dataflow::Conversion::RenderingFacadeToDynamicMesh(Facade, MeshIndex, DynamicMesh);
+
+					if(!CreateSkeletalMesh(SkeletalMeshes[MeshIndex], {DynamicMesh}, ObjectSkeleton->GetReferenceSkeleton()))
+					{
+						bValidSkeletalMeshes = false;
+					}
+				}
+			}
+			return bValidSkeletalMeshes;
+		}
+	}
+	return false;
+}
+	
+}
+
diff --git a/Engine/Plugins/Experimental/Dataflow/Source/DataflowNodes/Private/Dataflow/DataflowNodesPlugin.cpp b/Engine/Plugins/Experimental/Dataflow/Source/DataflowNodes/Private/Dataflow/DataflowNodesPlugin.cpp
--- a/Engine/Plugins/Experimental/Dataflow/Source/DataflowNodes/Private/Dataflow/DataflowNodesPlugin.cpp
+++ b/Engine/Plugins/Experimental/Dataflow/Source/DataflowNodes/Private/Dataflow/DataflowNodesPlugin.cpp
@@ -11,6 +11,8 @@
 #include "Dataflow/DataflowContextOverridesNodes.h"
 #include "Dataflow/DataflowCollectionAddScalarVertexPropertyNode.h"
 #include "Dataflow/DataflowCollectionEditSkinWeightsNode.h"
+#include "Dataflow/DataflowCollectionEditSkeletonBonesNode.h"
+#include "Dataflow/DataflowSkeletonAssetTerminalNode.h"
 
 #define LOCTEXT_NAMESPACE "DataflowNodes"
 
@@ -49,6 +51,8 @@ void IDataflowNodesPlugin::StartupModule()
 	UE::Dataflow::DataflowCollectionAttributeKeyNodes();
 	DATAFLOW_NODE_REGISTER_CREATION_FACTORY(FDataflowCollectionAddScalarVertexPropertyNode);
 	DATAFLOW_NODE_REGISTER_CREATION_FACTORY(FDataflowCollectionEditSkinWeightsNode);
+	DATAFLOW_NODE_REGISTER_CREATION_FACTORY(FDataflowCollectionEditSkeletonBonesNode);
+	DATAFLOW_NODE_REGISTER_CREATION_FACTORY(FSkeletonAssetTerminalNode);
 
 	UE::Dataflow::RegisterNodeFilter(FDataflowTerminalNode::StaticType());
 
diff --git a/Engine/Plugins/Experimental/Dataflow/Source/DataflowNodes/Private/Dataflow/DataflowSkeletalMeshNodes.cpp b/Engine/Plugins/Experimental/Dataflow/Source/DataflowNodes/Private/Dataflow/DataflowSkeletalMeshNodes.cpp
--- a/Engine/Plugins/Experimental/Dataflow/Source/DataflowNodes/Private/Dataflow/DataflowSkeletalMeshNodes.cpp
+++ b/Engine/Plugins/Experimental/Dataflow/Source/DataflowNodes/Private/Dataflow/DataflowSkeletalMeshNodes.cpp
@@ -137,28 +137,30 @@ void FDataflowDebugDrawSkeletonObject::DrawDataflowElements(FPrimitiveDrawInterf
 	{
 		OnBoneSelectionChanged.Broadcast(BoneNames);
 	}
+	if(bBonesVisible)
+	{
+		FSkelDebugDrawConfig DrawConfig;
+		DrawConfig.bUseMultiColorAsDefaultColor = true;
+		DrawConfig.BoneDrawMode = EBoneDrawMode::Selected;
+		DrawConfig.BoneDrawSize = 1.f;
+		DrawConfig.bAddHitProxy = true;
+		DrawConfig.bForceDraw = true;
+		DrawConfig.DefaultBoneColor = FLinearColor(0.0f,0.0f,0.025f,1.0f);
+		DrawConfig.AffectedBoneColor = FLinearColor(1.0f,1.0f,1.0f,1.0f);
+		DrawConfig.SelectedBoneColor = FLinearColor(0.2f,1.0f,0.2f,1.0f);
+		DrawConfig.ParentOfSelectedBoneColor = FLinearColor(0.85f,0.45f,0.12f,1.0f);
 	
-	FSkelDebugDrawConfig DrawConfig;
-	DrawConfig.bUseMultiColorAsDefaultColor = true;
-	DrawConfig.BoneDrawMode = EBoneDrawMode::Selected;
-	DrawConfig.BoneDrawSize = 1.f;
-	DrawConfig.bAddHitProxy = true;
-	DrawConfig.bForceDraw = true;
-	DrawConfig.DefaultBoneColor = FLinearColor(0.0f,0.0f,0.025f,1.0f);
-	DrawConfig.AffectedBoneColor = FLinearColor(1.0f,1.0f,1.0f,1.0f);
-	DrawConfig.SelectedBoneColor = FLinearColor(0.2f,1.0f,0.2f,1.0f);
-	DrawConfig.ParentOfSelectedBoneColor = FLinearColor(0.85f,0.45f,0.12f,1.0f);
-	
-	SkeletalDebugRendering::DrawBones(
-		PDI, FVector::Zero(),
-		RequiredBones,
-		ReferenceSkeleton,
-		WorldTransforms,
-		SelectedBones,
-		BoneColors,
-		HitProxies,
-		DrawConfig
-	);
+		SkeletalDebugRendering::DrawBones(
+			PDI, FVector::Zero(),
+			RequiredBones,
+			ReferenceSkeleton,
+			WorldTransforms,
+			SelectedBones,
+			BoneColors,
+			HitProxies,
+			DrawConfig
+		);
+	}
 }
 
 void FGetSkeletalMeshDataflowNode::Evaluate(UE::Dataflow::FContext& Context, const FDataflowOutput* Out) const
diff --git a/Engine/Plugins/Experimental/Dataflow/Source/DataflowNodes/Private/Dataflow/DataflowSkeletonAssetTerminalNode.cpp b/Engine/Plugins/Experimental/Dataflow/Source/DataflowNodes/Private/Dataflow/DataflowSkeletonAssetTerminalNode.cpp
new file mode 100644
--- /dev/null
+++ b/Engine/Plugins/Experimental/Dataflow/Source/DataflowNodes/Private/Dataflow/DataflowSkeletonAssetTerminalNode.cpp
@@ -0,0 +1,68 @@
+// Copyright Epic Games, Inc. All Rights Reserved.
+
+#include "Dataflow/DataflowSkeletonAssetTerminalNode.h"
+#include "Animation/Skeleton.h"
+#include "AnimationUtils.h"
+#include "AssetRegistry/AssetRegistryModule.h"
+#include "Misc/PackageName.h"
+#include "UObject/Package.h"
+#include "UObject/SavePackage.h"
+
+#include UE_INLINE_GENERATED_CPP_BY_NAME(DataflowSkeletonAssetTerminalNode)
+
+namespace UE::Dataflow::Private
+{
+	static void OverrideSkeletonData(TObjectPtr<USkeleton> TargetSkeleton, TObjectPtr<const USkeleton> SourceSkeleton)
+	{
+		if (!TargetSkeleton || !SourceSkeleton || TargetSkeleton == SourceSkeleton)
+		{
+			return;
+		}
+
+#if WITH_EDITOR
+
+		// Get the package and name
+		UPackage* Package = TargetSkeleton->GetOutermost();
+		FName ObjectName = TargetSkeleton->GetFName();
+
+		// Duplicate into same package with same name
+		USkeleton* NewSkeleton = DuplicateObject<USkeleton>(SourceSkeleton, Package, ObjectName);
+
+		// Mark the package dirty 
+		Package->MarkPackageDirty();
+		FAssetRegistryModule::AssetCreated(NewSkeleton);
+
+		// Save the package 
+		FString FilePath = FPackageName::LongPackageNameToFilename(Package->GetName(), FPackageName::GetAssetPackageExtension());
+		FSavePackageArgs SaveArgs;
+		SaveArgs.TopLevelFlags = EObjectFlags::RF_Public | EObjectFlags::RF_Standalone;
+		SaveArgs.SaveFlags = SAVE_NoError;
+
+		UPackage::SavePackage(Package, NewSkeleton, *FilePath, SaveArgs);
+
+#endif
+	}
+}
+
+FSkeletonAssetTerminalNode::FSkeletonAssetTerminalNode(const UE::Dataflow::FNodeParameters& InParam, FGuid InGuid)
+	: FDataflowTerminalNode(InParam, InGuid)
+{
+	RegisterInputConnection(&SourceSkeleton);
+	RegisterInputConnection(&SkeletonAsset);
+}
+
+void FSkeletonAssetTerminalNode::Evaluate(UE::Dataflow::FContext& Context) const
+{}
+
+void FSkeletonAssetTerminalNode::SetAssetValue(TObjectPtr<UObject> Asset, UE::Dataflow::FContext& Context) const
+{
+	USkeleton* AssetToSet = Cast<USkeleton>(Asset.Get());
+	if (!AssetToSet)
+	{
+		AssetToSet = GetValue(Context, &SkeletonAsset);
+	}
+	if(AssetToSet)
+	{
+		UE::Dataflow::Private::OverrideSkeletonData(AssetToSet, GetValue(Context, &SourceSkeleton));
+	}
+}
\ No newline at end of file
diff --git a/Engine/Plugins/Experimental/Dataflow/Source/DataflowNodes/Public/Dataflow/DataflowCollectionEditSkeletonBonesNode.h b/Engine/Plugins/Experimental/Dataflow/Source/DataflowNodes/Public/Dataflow/DataflowCollectionEditSkeletonBonesNode.h
new file mode 100644
--- /dev/null
+++ b/Engine/Plugins/Experimental/Dataflow/Source/DataflowNodes/Public/Dataflow/DataflowCollectionEditSkeletonBonesNode.h
@@ -0,0 +1,75 @@
+// Copyright Epic Games, Inc. All Rights Reserved.
+
+#pragma once 
+
+#include "DataflowCollectionAddScalarVertexPropertyNode.h"
+#include "DataflowPrimitiveNode.h"
+#include "DataflowSkeletalMeshNodes.h"
+#include "Dataflow/DataflowNode.h"
+#include "GeometryCollection/ManagedArrayCollection.h"
+
+#include "DataflowCollectionEditSkeletonBonesNode.generated.h"
+
+#define UE_API DATAFLOWNODES_API
+
+/** Edit skeleton bones properties. */
+USTRUCT(Meta = (Experimental,DataflowCollection))
+struct FDataflowCollectionEditSkeletonBonesNode : public FDataflowPrimitiveNode
+{
+	GENERATED_BODY()
+	DATAFLOW_NODE_DEFINE_INTERNAL(FDataflowCollectionEditSkeletonBonesNode, "EditSkeletonBones", "Collection", "Edit skeleton bones")
+
+public:
+
+	UE_API FDataflowCollectionEditSkeletonBonesNode(const UE::Dataflow::FNodeParameters& InParam, FGuid InGuid = FGuid::NewGuid());
+
+	UPROPERTY(Meta = (DataflowInput, DataflowOutput, DataflowPassthrough = "Collection"))
+	FManagedArrayCollection Collection;
+	
+	/** Target group in which the attributes are stored */
+	UPROPERTY(EditAnywhere, Category = "Vertex Attributes", meta = (DisplayName = "Vertex Group"))
+	FScalarVertexPropertyGroup VertexGroup;
+
+	/** Skeleton to be edited */
+	UPROPERTY(EditAnywhere, Category = "Skeleton Binding", Meta = (DataflowInput, DataflowOutput, DataflowPassthrough = "Skeleton"))
+	TObjectPtr<USkeleton> Skeleton = nullptr;
+
+	/** Skeleton used to store the tool result */
+	UPROPERTY()
+	TObjectPtr<USkeleton> ToolSkeleton = nullptr;
+	
+	/** Validate the skeletal mesh construction */
+	void ValidateSkeletalMeshes() {bValidSkeletalMeshes = true;}
+
+	/** Delegate to transfer the bone selection to the tool*/
+	FDataflowBoneSelectionChangedNotifyDelegate OnBoneSelectionChanged;
+	
+private:
+
+	/** Update the tool skeleton */
+	TObjectPtr<USkeleton> UpdateToolSkeleton(UE::Dataflow::FContext& Context);
+
+	//~ Begin FDataflowPrimitiveNode interface
+	virtual TArray<UE::Dataflow::FRenderingParameter> GetRenderParametersImpl() const override;
+	virtual void Evaluate(UE::Dataflow::FContext& Context, const FDataflowOutput* Out) const override;
+	virtual void AddPrimitiveComponents(UE::Dataflow::FContext& Context, const TSharedPtr<const FManagedArrayCollection> RenderCollection,
+		TObjectPtr<UObject> NodeOwner, TObjectPtr<AActor> RootActor, TArray<TObjectPtr<UPrimitiveComponent>>& PrimitiveComponents) override;
+	virtual void OnInvalidate() override;
+#if WITH_EDITOR
+	virtual bool CanDebugDraw() const override { return true; }
+	virtual bool CanDebugDrawViewMode(const FName& ViewModeName) const override;
+	virtual void DebugDraw(UE::Dataflow::FContext& Context,
+		IDataflowDebugDrawInterface& DataflowRenderingInterface,
+		const FDebugDrawParameters& DebugDrawParameters) const override;
+#endif
+	//~ End FDataflowPrimitiveNode interface
+
+	/** Transient skeletal mesh built from dataflow render collection */
+	UPROPERTY(Transient)
+	TArray<TObjectPtr<USkeletalMesh>> SkeletalMeshes;
+
+	/** Valid skeletal mesh boolean to trigger the construction */
+	bool bValidSkeletalMeshes = false;
+};
+
+#undef UE_API
diff --git a/Engine/Plugins/Experimental/Dataflow/Source/DataflowNodes/Public/Dataflow/DataflowSkeletalMeshNodes.h b/Engine/Plugins/Experimental/Dataflow/Source/DataflowNodes/Public/Dataflow/DataflowSkeletalMeshNodes.h
--- a/Engine/Plugins/Experimental/Dataflow/Source/DataflowNodes/Public/Dataflow/DataflowSkeletalMeshNodes.h
+++ b/Engine/Plugins/Experimental/Dataflow/Source/DataflowNodes/Public/Dataflow/DataflowSkeletalMeshNodes.h
@@ -24,8 +24,8 @@ DECLARE_MULTICAST_DELEGATE_OneParam(FDataflowBoneSelectionChangedNotifyDelegate,
 /** Debug draw skeleton object to debug skeletal meshes */
 struct FDataflowDebugDrawSkeletonObject : public FDataflowDebugDrawBaseObject
 {
-	FDataflowDebugDrawSkeletonObject(IDataflowDebugDrawInterface::FDataflowElementsType& InDataflowElements, const FReferenceSkeleton& InReferenceSkeleton) :
-		FDataflowDebugDrawBaseObject(InDataflowElements), ReferenceSkeleton(InReferenceSkeleton)
+	FDataflowDebugDrawSkeletonObject(IDataflowDebugDrawInterface::FDataflowElementsType& InDataflowElements, const FReferenceSkeleton& InReferenceSkeleton, const bool bBonesVisible = true) :
+		FDataflowDebugDrawBaseObject(InDataflowElements), ReferenceSkeleton(InReferenceSkeleton), bBonesVisible(bBonesVisible)
 	{}
 	
 	/** Populate dataflow elements */
@@ -46,6 +46,9 @@ private :
 	
 	/** Previous element selection */
 	TArray<bool> PreviousSelection;
+
+	/** Boolean to make the bones visible or not */
+	bool bBonesVisible = true;
 };
 
 USTRUCT()
diff --git a/Engine/Plugins/Experimental/Dataflow/Source/DataflowNodes/Public/Dataflow/DataflowSkeletonAssetTerminalNode.h b/Engine/Plugins/Experimental/Dataflow/Source/DataflowNodes/Public/Dataflow/DataflowSkeletonAssetTerminalNode.h
new file mode 100644
--- /dev/null
+++ b/Engine/Plugins/Experimental/Dataflow/Source/DataflowNodes/Public/Dataflow/DataflowSkeletonAssetTerminalNode.h
@@ -0,0 +1,38 @@
+﻿// Copyright Epic Games, Inc. All Rights Reserved.
+
+#pragma once 
+
+#include "CoreMinimal.h"
+#include "Dataflow/DataflowCore.h"
+#include "Dataflow/DataflowEngine.h"
+
+#include "DataflowSkeletonAssetTerminalNode.generated.h"
+
+class USkeleton;
+
+/*
+* Dataflow terminal node to save a skeleton asset
+*/
+USTRUCT(meta = (DataflowFlesh, DataflowTerminal))
+struct FSkeletonAssetTerminalNode : public FDataflowTerminalNode
+{
+	GENERATED_BODY()
+	DATAFLOW_NODE_DEFINE_INTERNAL(FSkeletonAssetTerminalNode, "SkeletonAssetTerminal", "Terminal", "")
+
+public:
+	FSkeletonAssetTerminalNode(const UE::Dataflow::FNodeParameters& InParam, FGuid InGuid = FGuid::NewGuid());
+
+private:
+	/** Source Skeleton used to override the skeleton asset */
+	UPROPERTY(meta = (DataflowInput))
+	TObjectPtr<const USkeleton> SourceSkeleton = nullptr;
+
+	/** Skeleton Asset to be saved */
+	UPROPERTY(EditAnywhere, Category = Asset, meta = (DataflowInput))
+	TObjectPtr<USkeleton> SkeletonAsset = nullptr;
+
+	//~ Begin FDataflowTerminalNode interface
+	virtual void Evaluate(UE::Dataflow::FContext& Context) const override;
+	virtual void SetAssetValue(TObjectPtr<UObject> Asset, UE::Dataflow::FContext& Context) const override;
+	//~ End FDataflowTerminalNode interface
+};
