Dataflow Editor : Make Vertex attribute paint tool more generic to adapt more nodes that need to use it
- Added a generic node editable type that other node can inherit from (FDataflowVertexAttributeEditableNode)
- This node type define an interface for the Vertex Attribute tool to get and set the vertex attributes
- Added a temporary way to define an fallback builder to allow enabling the new tool  on and off
Note : In the cloth Dataflow editor, the new tool still does not suppport 2D/3D Sim/Render mesh editing , this will be done in a later change

diff --git a/Engine/Plugins/ChaosClothAssetEditor/Source/ChaosClothAssetDataflowNodes/Private/ChaosClothAsset/WeightMapNode.cpp b/Engine/Plugins/ChaosClothAssetEditor/Source/ChaosClothAssetDataflowNodes/Private/ChaosClothAsset/WeightMapNode.cpp
--- a/Engine/Plugins/ChaosClothAssetEditor/Source/ChaosClothAssetDataflowNodes/Private/ChaosClothAsset/WeightMapNode.cpp
+++ b/Engine/Plugins/ChaosClothAssetEditor/Source/ChaosClothAssetDataflowNodes/Private/ChaosClothAsset/WeightMapNode.cpp
@@ -35,7 +35,7 @@ namespace UE::Chaos::ClothAsset::Private
 }
 
 FChaosClothAssetWeightMapNode::FChaosClothAssetWeightMapNode(const UE::Dataflow::FNodeParameters& InParam, FGuid InGuid)
-	: FDataflowNode(InParam, InGuid)
+	: FDataflowVertexAttributeEditableNode(InParam, InGuid)
 	, Transfer(FDataflowFunctionProperty::FDelegate::CreateRaw(this, &FChaosClothAssetWeightMapNode::OnTransfer))
 {
 	RegisterInputConnection(&Collection);
@@ -245,6 +245,65 @@ FName FChaosClothAssetWeightMapNode::GetInputName(UE::Dataflow::FContext& Contex
 	return InInputName != NAME_None ? InInputName : FName(OutputName.StringValue);
 }
 
+void FChaosClothAssetWeightMapNode::GetVertexAttributeValues(UE::Dataflow::FContext& Context, TArray<float>& OutValues) const
+{
+	using namespace UE::Chaos::ClothAsset;
+
+	const FName InInputName = GetInputName(Context);
+
+	FManagedArrayCollection InCollection = GetValue<FManagedArrayCollection>(Context, &Collection);
+	const TSharedRef<FManagedArrayCollection> ClothCollection = MakeShared<FManagedArrayCollection>(MoveTemp(InCollection));
+	FCollectionClothConstFacade ClothFacade(ClothCollection);
+	if (ClothFacade.IsValid())  // Can only act on the collection if it is a valid cloth collection
+	{
+		TConstArrayView<float> InputAttributeValues;
+		int32 NumValues = 0;
+		if (MeshTarget == EChaosClothAssetWeightMapMeshTarget::Simulation)
+		{
+			InputAttributeValues = ClothFacade.GetWeightMap(InInputName);
+			NumValues = ClothFacade.GetNumSimVertices3D();
+		}
+		else
+		{
+			InputAttributeValues = ClothFacade.GetUserDefinedAttribute<float>(InInputName, ClothCollectionGroup::RenderVertices);
+			NumValues = ClothFacade.GetNumRenderVertices();
+		}
+		TArray<float> FallbackInputValues;
+		if (InputAttributeValues.IsEmpty())
+		{
+			FallbackInputValues.Init(0, NumValues);
+			InputAttributeValues = FallbackInputValues;
+		}
+		OutValues.Init(0, NumValues);
+		UE::Chaos::ClothAsset::Private::CalculateFinalVertexWeightValues(InputAttributeValues, OutValues, MapOverrideType, GetVertexWeights());
+	}
+}
+
+void FChaosClothAssetWeightMapNode::SetVertexAttributeValues(UE::Dataflow::FContext& Context, const TArray<float>& InValues)
+{
+	using namespace UE::Chaos::ClothAsset;
+
+	const FName InInputName = GetInputName(Context);
+
+	FManagedArrayCollection InCollection = GetValue<FManagedArrayCollection>(Context, &Collection);
+	const TSharedRef<FManagedArrayCollection> ClothCollection = MakeShared<FManagedArrayCollection>(MoveTemp(InCollection));
+	FCollectionClothConstFacade ClothFacade(ClothCollection);
+	if (ClothFacade.IsValid())  // Can only act on the collection if it is a valid cloth collection
+	{
+		TConstArrayView<float> InputAttributeValues;
+		if (MeshTarget == EChaosClothAssetWeightMapMeshTarget::Simulation)
+		{
+			InputAttributeValues = ClothFacade.GetWeightMap(InInputName);
+		}
+		else
+		{
+			InputAttributeValues = ClothFacade.GetUserDefinedAttribute<float>(InInputName, ClothCollectionGroup::RenderVertices);
+		}
+
+		UE::Chaos::ClothAsset::Private::SetVertexWeights(InputAttributeValues, InValues, MapOverrideType, GetVertexWeights());
+	}
+}
+
 void FChaosClothAssetWeightMapNode::SetVertexWeights(const TConstArrayView<float> InputMap, const TArray<float>& FinalValues)
 {
 	UE::Chaos::ClothAsset::Private::SetVertexWeights(InputMap, FinalValues, MapOverrideType, GetVertexWeights());
diff --git a/Engine/Plugins/ChaosClothAssetEditor/Source/ChaosClothAssetDataflowNodes/Public/ChaosClothAsset/WeightMapNode.h b/Engine/Plugins/ChaosClothAssetEditor/Source/ChaosClothAssetDataflowNodes/Public/ChaosClothAsset/WeightMapNode.h
--- a/Engine/Plugins/ChaosClothAssetEditor/Source/ChaosClothAssetDataflowNodes/Public/ChaosClothAsset/WeightMapNode.h
+++ b/Engine/Plugins/ChaosClothAssetEditor/Source/ChaosClothAssetDataflowNodes/Public/ChaosClothAsset/WeightMapNode.h
@@ -6,6 +6,8 @@
 #include "Dataflow/DataflowFunctionProperty.h"
 #include "GeometryCollection/ManagedArrayCollection.h"
 #include "ChaosClothAsset/ConnectableValue.h"
+#include "Dataflow/DataflowCollectionAddScalarVertexPropertyNode.h"
+
 #include "WeightMapNode.generated.h"
 
 #if UE_ENABLE_INCLUDE_ORDER_DEPRECATED_IN_5_5
@@ -22,10 +24,12 @@ enum class EChaosClothAssetWeightMapMeshTarget : uint8
 	Render
 };
 
+
+
 /** Painted weight map attributes node. */
 PRAGMA_DISABLE_DEPRECATION_WARNINGS  // For Name implicit operators.
 USTRUCT(Meta = (DataflowCloth))
-struct FChaosClothAssetWeightMapNode : public FDataflowNode
+struct FChaosClothAssetWeightMapNode : public FDataflowVertexAttributeEditableNode
 {
 PRAGMA_ENABLE_DEPRECATION_WARNINGS
 	GENERATED_USTRUCT_BODY()
@@ -95,6 +99,11 @@ private:
 	virtual FDataflowOutput* RedirectSerializedOutput(const FName& MissingOutputName) override;
 	//~ End FDataflowNode interface
 
+	//~ Begin - FDataflowVertexAttributeEditableNode API
+	virtual void GetVertexAttributeValues(UE::Dataflow::FContext& Context, TArray<float>& OutValues) const override;
+	virtual void SetVertexAttributeValues(UE::Dataflow::FContext& Context, const TArray<float>& InValues) override;
+	//~ End - FDataflowVertexAttributeEditableNode API
+
 	void OnTransfer(UE::Dataflow::FContext& Context);
 
 	UPROPERTY()
diff --git a/Engine/Plugins/ChaosClothAssetEditor/Source/ChaosClothAssetEditorTools/Private/ChaosClothAsset/ChaosClothAssetEditorToolsModule.cpp b/Engine/Plugins/ChaosClothAssetEditor/Source/ChaosClothAssetEditorTools/Private/ChaosClothAsset/ChaosClothAssetEditorToolsModule.cpp
--- a/Engine/Plugins/ChaosClothAssetEditor/Source/ChaosClothAssetEditorTools/Private/ChaosClothAsset/ChaosClothAssetEditorToolsModule.cpp
+++ b/Engine/Plugins/ChaosClothAssetEditor/Source/ChaosClothAssetEditorTools/Private/ChaosClothAsset/ChaosClothAssetEditorToolsModule.cpp
@@ -3,6 +3,7 @@
 #include "Modules/ModuleManager.h"
 
 #include "Dataflow/DataflowToolRegistry.h"
+#include "DataflowEditorTools/DataflowEditorVertexAttributePaintTool.h"
 #include "ChaosClothAsset/ClothEditorToolBuilders.h"
 #include "ChaosClothAsset/SelectionNode.h"
 #include "ChaosClothAsset/TransferSkinWeightsNode.h"
@@ -22,8 +23,12 @@ namespace UE::Chaos::ClothAsset
 		
 			const TSharedRef<const FClothToolActionCommandBindings> ClothToolActions = MakeShared<FClothToolActionCommandBindings>();
 
+			UClothEditorWeightMapPaintToolBuilder* FallbackBuilder = NewObject<UClothEditorWeightMapPaintToolBuilder>();
+			UDataflowEditorVertexAttributePaintToolBuilder* WeightMapBuilder = NewObject<UDataflowEditorVertexAttributePaintToolBuilder>();
+			WeightMapBuilder->SetFallbackToolBuilder(FallbackBuilder);
+			
 			ToolRegistry.AddNodeToToolMapping(FChaosClothAssetWeightMapNode::StaticType(),          // Node type
-				NewObject<UClothEditorWeightMapPaintToolBuilder>(),                                 // Tool builder
+				WeightMapBuilder,                                 // Tool builder
 				ClothToolActions,                                                                   // Tool action bindings
 				FSlateIcon(FName("ClothStyle"), FName("ChaosClothAssetEditor.AddWeightMapNode")),   // Button icon (see FChaosClothAssetEditorStyle)
 				LOCTEXT("AddWeightMapNodeButtonText", "Cloth Weight Map"), FName("Cloth"));                               // Button text
diff --git a/Engine/Plugins/Experimental/Dataflow/Source/DataflowEditor/Private/DataflowEditorTools/DataflowEditorVertexAttributePaintTool.cpp b/Engine/Plugins/Experimental/Dataflow/Source/DataflowEditor/Private/DataflowEditorTools/DataflowEditorVertexAttributePaintTool.cpp
--- a/Engine/Plugins/Experimental/Dataflow/Source/DataflowEditor/Private/DataflowEditorTools/DataflowEditorVertexAttributePaintTool.cpp
+++ b/Engine/Plugins/Experimental/Dataflow/Source/DataflowEditor/Private/DataflowEditorTools/DataflowEditorVertexAttributePaintTool.cpp
@@ -160,6 +160,11 @@ void UDataflowEditorVertexAttributePaintToolBuilder::GetSupportedConstructionVie
 
 bool UDataflowEditorVertexAttributePaintToolBuilder::CanBuildTool(const FToolBuilderState& SceneState) const
 {
+	if (!UE::DataflowEditorVertexAttributePaintTool::CVars::DataflowEditorUseNewWeightMapTool && FallbackToolBuilder)
+	{
+		return FallbackToolBuilder->CanBuildTool(SceneState);
+	}
+	
 	if (UMeshSurfacePointMeshEditingToolBuilder::CanBuildTool(SceneState))
 	{
 		if (SceneState.SelectedComponents.Num() == 1)
@@ -172,7 +177,7 @@ bool UDataflowEditorVertexAttributePaintToolBuilder::CanBuildTool(const FToolBui
 					{
 						if (const TSharedPtr<UE::Dataflow::FEngineContext> EvaluationContext = ContextObject->GetDataflowContext())
 						{
-							if (const FDataflowNode* PrimarySelection = ContextObject->GetSelectedNodeOfType<FDataflowCollectionAddScalarVertexPropertyNode>())
+							if (const FDataflowNode* PrimarySelection = ContextObject->GetSelectedNodeOfType<FDataflowVertexAttributeEditableNode>())
 							{
 								return UE::DataflowEditorVertexAttributePaintTool::Private::HasManagedArrayCollection(PrimarySelection, EvaluationContext);
 							}
@@ -200,6 +205,10 @@ UMeshSurfacePointTool* UDataflowEditorVertexAttributePaintToolBuilder::CreateNew
 		}
 		return PaintTool;
 	}
+	else if (FallbackToolBuilder)
+	{
+		return FallbackToolBuilder->CreateNewTool(SceneState);
+	}
 
 	// fallback : use the old tool
 	// todo(ccaillaud) : remove this fall back when the new tool is on par with the old one 
@@ -218,7 +227,7 @@ UMeshSurfacePointTool* UDataflowEditorVertexAttributePaintToolBuilder::CreateNew
 // Data
 //
 ///////////////////////////////////////////////////////////////////////////////////////////////////////////////////////////////////////////////////
-void FDataflowEditorVertexAttributePaintToolData::Setup(FDynamicMesh3& InMesh, FDataflowCollectionAddScalarVertexPropertyNode* InNodeToUpdate, TObjectPtr<UDataflowContextObject> InDataflowEditorContextObject)
+void FDataflowEditorVertexAttributePaintToolData::Setup(FDynamicMesh3& InMesh, FDataflowVertexAttributeEditableNode* InNodeToUpdate, TObjectPtr<UDataflowContextObject> InDataflowEditorContextObject)
 {
 	// this should only be called once or have shutdown called in between 
 	ensure(ActiveWeightMap == nullptr);
@@ -265,19 +274,14 @@ void FDataflowEditorVertexAttributePaintToolData::Setup(FDynamicMesh3& InMesh, F
 		{
 			// Setup DynamicMeshToWeight conversion and get Input weight map (if it exists)
 			const int32 NumExpectedWeights = bHaveDynamicMeshToWeightConversion ? WeightToDynamicMesh.Num() : InMesh.MaxVertexID();
-			SetupWeights.SetNumZeroed(NumExpectedWeights);
 
+			TArray<float> CurrentWeights;
 			// Get the setup attribute values
 			if (TSharedPtr<UE::Dataflow::FEngineContext> DataflowContext = InDataflowEditorContextObject->GetDataflowContext())
 			{
-				const FCollectionAttributeKey AttributeKey = NodeToUpdate->GetWeightAttributeKey(*DataflowContext);
-				NodeToUpdate->FillAttributeWeights(InDataflowEditorContextObject->GetSelectedCollection(), AttributeKey, SetupWeights);
+				NodeToUpdate->GetVertexAttributeValues(*DataflowContext, CurrentWeights);
 			}
-
-			// Copy attribute  from selected node to the preview mesh
-			TArray<float> CurrentWeights;
-			CurrentWeights.SetNumZeroed(NumExpectedWeights);
-			NodeToUpdate->ExtractVertexWeights(SetupWeights, TArrayView<float>(CurrentWeights));
+			ensure(CurrentWeights.Num() == NumExpectedWeights);
 
 			if (bHaveDynamicMeshToWeightConversion)
 			{
@@ -374,31 +378,30 @@ float FDataflowEditorVertexAttributePaintToolData::GetAverageValue(const TArray<
 	return (float)AverageValue;
 }
 
-void FDataflowEditorVertexAttributePaintToolData::UpdateNode()
+void FDataflowEditorVertexAttributePaintToolData::UpdateNode(TObjectPtr<UDataflowContextObject> InDataflowEditorContextObject)
 {
 	if (ActiveWeightMap && NodeToUpdate)
 	{
 		if (const FDynamicMesh3* Mesh = ActiveWeightMap->GetParent())
 		{
-			TArray<float> WeightsToApply;
-
-			const int32 NumVertices = Mesh->VertexCount();
-			WeightsToApply.SetNumUninitialized(NumVertices);
-			for (int32 VertexID = 0; VertexID < NumVertices; ++VertexID)
-			{
-				ActiveWeightMap->GetValue(VertexID, &WeightsToApply[VertexID]);
-			}
-
-			if (bHaveDynamicMeshToWeightConversion)
-			{
-				NodeToUpdate->ReportVertexWeights(SetupWeights, WeightsToApply, DynamicMeshToWeight);
-			}
-			else
+			if (TSharedPtr<UE::Dataflow::FEngineContext> DataflowContext = InDataflowEditorContextObject->GetDataflowContext())
 			{
-				NodeToUpdate->ReportVertexWeights(SetupWeights, WeightsToApply, {});
+				TArray<float> WeightsToApply;
+				const int32 NumVertices = Mesh->VertexCount();
+				WeightsToApply.SetNumUninitialized(NumVertices);
+				for (int32 VertexID = 0; VertexID < NumVertices; ++VertexID)
+				{
+					if (bHaveDynamicMeshToWeightConversion)
+					{
+						ActiveWeightMap->GetValue(VertexID, &WeightsToApply[DynamicMeshToWeight[VertexID]]);
+					}
+					else
+					{
+						ActiveWeightMap->GetValue(VertexID, &WeightsToApply[VertexID]);
+					}
+				}
+				NodeToUpdate->SetVertexAttributeValues(*DataflowContext, WeightsToApply);
 			}
-
-			NodeToUpdate->Invalidate();
 		}
 	}
 }
@@ -444,7 +447,7 @@ void UDataflowEditorVertexAttributePaintTool::Setup()
 
 	// todo(ccaillaud) : this should be a customized by the inherited class 
 	// Get the selected node to read and write the vertex attribute
-	FDataflowCollectionAddScalarVertexPropertyNode* NodeToUpdate = DataflowEditorContextObject->GetSelectedNodeOfType<FDataflowCollectionAddScalarVertexPropertyNode>();
+	FDataflowVertexAttributeEditableNode* NodeToUpdate = DataflowEditorContextObject->GetSelectedNodeOfType<FDataflowVertexAttributeEditableNode>();
 	checkf(NodeToUpdate, TEXT("No Node is currently selected, or more than one node is selected"));
 
 	// todo(ccaillaud) : this should be a parameter of the tool ? 
@@ -693,7 +696,7 @@ void UDataflowEditorVertexAttributePaintTool::CommitResult(UBaseDynamicMeshCompo
 {
 	GetToolManager()->BeginUndoTransaction(LOCTEXT("DataflowEditorVertexAttributePaintTool_TransactionName", "Paint Weights"));
 
-	VertexData.UpdateNode();
+	VertexData.UpdateNode(DataflowEditorContextObject);
 
 	GetToolManager()->EndUndoTransaction();
 }
diff --git a/Engine/Plugins/Experimental/Dataflow/Source/DataflowEditor/Private/DataflowEditorTools/DataflowEditorVertexAttributePaintTool.h b/Engine/Plugins/Experimental/Dataflow/Source/DataflowEditor/Private/DataflowEditorTools/DataflowEditorVertexAttributePaintTool.h
deleted file mode 100644
diff --git a/Engine/Plugins/Experimental/Dataflow/Source/DataflowEditor/Public/DataflowEditorTools/DataflowEditorVertexAttributePaintTool.h b/Engine/Plugins/Experimental/Dataflow/Source/DataflowEditor/Public/DataflowEditorTools/DataflowEditorVertexAttributePaintTool.h
new file mode 100644
--- /dev/null
+++ b/Engine/Plugins/Experimental/Dataflow/Source/DataflowEditor/Public/DataflowEditorTools/DataflowEditorVertexAttributePaintTool.h
@@ -0,0 +1,480 @@
+// Copyright Epic Games, Inc. All Rights Reserved.
+
+#pragma once
+
+#include "CoreMinimal.h"
+#include "HAL/Platform.h"
+#include "UObject/NoExportTypes.h"
+#include "BaseTools/BaseBrushTool.h"
+#include "BaseTools/MeshSurfacePointMeshEditingTool.h"
+#include "Components/DynamicMeshComponent.h"
+#include "PropertySets/WeightMapSetProperties.h"
+#include "Mechanics/PolyLassoMarqueeMechanic.h"
+
+#include "Sculpting/MeshSculptToolBase.h"
+#include "Sculpting/MeshBrushOpBase.h"
+
+#include "DynamicMesh/DynamicMeshAABBTree3.h"
+#include "DynamicMesh/DynamicMeshOctree3.h"
+#include "DynamicMesh/DynamicMeshAttributeSet.h"
+#include "DynamicMesh/DynamicMeshChangeTracker.h"
+#include "DynamicMesh/MeshNormals.h"
+#include "TransformTypes.h"
+#include "ToolDataVisualizer.h"
+#include "GroupTopology.h"
+#include "MeshVertexPaintTool.h"
+#include "DataflowEditorTools/DataflowEditorToolBuilder.h"
+#include "Dataflow/DataflowColorRamp.h"
+#include "DataflowEditorTools/DataflowEditorToolEnums.h"
+#include "ToolMeshSelector.h"
+
+#include "DataflowEditorVertexAttributePaintTool.generated.h"
+
+
+#define UE_API DATAFLOWEDITOR_API
+
+class UMeshElementsVisualizer;
+class UDataflowWeightMapEraseBrushOpProps;
+class UDataflowVertexAttributePaintBrushOpProps;
+class UDataflowContextObject;
+struct FDataflowVertexAttributeEditableNode;
+struct FDataflowVertexAttributeProviderNode;
+class UDataflowEditorMode;
+
+namespace UE::DataflowEditorVertexAttributePaintTool::Private
+{
+	class FMeshChange;
+}
+
+namespace UE::DataflowEditorVertexAttributePaintTool::CVars
+{
+	extern bool DataflowEditorUseNewWeightMapTool;
+}
+
+DECLARE_STATS_GROUP(TEXT("Dataflow_VertexAttributePaintTool"), STATGROUP_VertexAttributePaintTool, STATCAT_Advanced);
+DECLARE_CYCLE_STAT(TEXT("Dataflow_VertexAttributePaintTool_UpdateROI"), VertexAttributePaintTool_UpdateROI, STATGROUP_VertexAttributePaintTool);
+DECLARE_CYCLE_STAT(TEXT("Dataflow_VertexAttributePaintTool_ApplyStamp"), VertexAttributePaintTool_ApplyStamp, STATGROUP_VertexAttributePaintTool);
+DECLARE_CYCLE_STAT(TEXT("Dataflow_VertexAttributePaintTool_Tick"), VertexAttributePaintTool_Tick, STATGROUP_VertexAttributePaintTool);
+DECLARE_CYCLE_STAT(TEXT("Dataflow_VertexAttributePaintTool_Tick_ApplyStampBlock"), VertexAttributePaintTool_Tick_ApplyStampBlock, STATGROUP_VertexAttributePaintTool);
+DECLARE_CYCLE_STAT(TEXT("Dataflow_VertexAttributePaintTool_Tick_UpdateMeshBlock"), VertexAttributePaintTool_Tick_UpdateMeshBlock, STATGROUP_VertexAttributePaintTool);
+
+///////////////////////////////////////////////////////////////////////////////////////////////////////////////////////////////////////////////////
+
+/**
+ * Tool Builder
+ */
+UCLASS(MinimalAPI)
+class UDataflowEditorVertexAttributePaintToolBuilder : public UMeshSurfacePointMeshEditingToolBuilder, public IDataflowEditorToolBuilder
+{
+	GENERATED_BODY()
+public:
+	void SetFallbackToolBuilder(UMeshSurfacePointMeshEditingToolBuilder* InFallbackToolBuilder) { FallbackToolBuilder = InFallbackToolBuilder; }
+	void SetEditorMode(UDataflowEditorMode* InMode) { Mode = InMode; }
+
+private:
+	UE_API virtual void GetSupportedConstructionViewModes(const UDataflowContextObject& ContextObject, TArray<const UE::Dataflow::IDataflowConstructionViewMode*>& Modes) const override;
+	UE_API virtual bool CanBuildTool(const FToolBuilderState& SceneState) const override;
+	UE_API virtual UMeshSurfacePointTool* CreateNewTool(const FToolBuilderState& SceneState) const override;
+	virtual bool CanSetConstructionViewWireframeActive() const { return false; }
+	UDataflowEditorMode* Mode = nullptr;
+
+	UPROPERTY()
+	TObjectPtr<UMeshSurfacePointMeshEditingToolBuilder> FallbackToolBuilder = nullptr;
+};
+
+///////////////////////////////////////////////////////////////////////////////////////////////////////////////////////////////////////////////////
+
+USTRUCT()
+struct FDataflowEditorVertexAttributePaintToolBrushProperties
+{
+	GENERATED_BODY()
+
+	UPROPERTY()
+	EDataflowEditorToolEditOperation BrushMode = EDataflowEditorToolEditOperation::Add;
+
+	/** Adaptive size of brush - Relative to the model size */
+	UPROPERTY()
+	float BrushSize = 0.25f;
+
+	/** Value to paint on the mesh (dependent on the BrushMode) */
+	UPROPERTY()
+	float AttributeValue = 1.0;
+
+	/** How quickly each brush stroke will drive mesh values towards the desired value */
+	UPROPERTY()
+	float Strength = 1.0;
+
+	/** Area Mode specifies the shape of the brush and which triangles will be included relative to the cursor */
+	UPROPERTY(EditAnywhere, Category = Filters, meta = (DisplayName = "Brush Area Mode"))
+	EMeshVertexPaintBrushAreaType BrushAreaMode = EMeshVertexPaintBrushAreaType::Connected;
+
+
+	/** The Region affected by the current operation will be bounded by edge angles larger than this threshold */
+	UPROPERTY(EditAnywhere, Category = Filters, meta = (UIMin = "0.0", ClampMin = "0.0", UIMax = "180.0", ClampMax = "180.0"))
+	float AngleThreshold = 180.0f;
+
+	/** The Region affected by the current operation will be bounded by UV borders/seams */
+	UPROPERTY(EditAnywhere, Category = Filters)
+	bool bUVSeams = false;
+
+	/** The Region affected by the current operation will be bounded by Hard Normal edges/seams */
+	UPROPERTY(EditAnywhere, Category = Filters)
+	bool bNormalSeams = false;
+
+	/** Control which triangles can be affected by the current operation based on visibility. Applied after all other filters. */
+	UPROPERTY(EditAnywhere, Category = Filters)
+	EDataflowEditorToolVisibilityType VisibilityFilter = EDataflowEditorToolVisibilityType::None;
+
+	/** The weight value at the brush indicator */
+	UPROPERTY(VisibleAnywhere, Transient, Category = Query, meta = (NoResetToDefault))
+	double ValueAtBrush = 0;
+};
+
+///////////////////////////////////////////////////////////////////////////////////////////////////////////////////////////////////////////////////
+
+USTRUCT()
+struct FDataflowEditorVertexAttributePaintToolGradientProperties
+{
+	GENERATED_BODY()
+
+	// Mesh selection mode (vertex, edge face)
+	/** The Gradient upper limit value */
+	UPROPERTY(EditAnywhere, Category = Gradient, meta = (UIMin = 0, ClampMin = 0, UIMax = 1, ClampMax = 1))
+	double GradientHighValue = 1.0;
+
+	/** The Gradient lower limit value */
+	UPROPERTY(EditAnywhere, Category = Gradient, meta = (UIMin = 0, ClampMin = 0, UIMax = 1, ClampMax = 1))
+	double GradientLowValue = 0.0;
+};
+
+///////////////////////////////////////////////////////////////////////////////////////////////////////////////////////////////////////////////////
+
+USTRUCT()
+struct FDataflowEditorVertexAttributePaintToolSelectionProperties
+{
+	GENERATED_BODY()
+
+	// Mesh selection mode (vertex, edge face)
+	UPROPERTY()
+	EComponentSelectionMode ComponentSelectionMode = EComponentSelectionMode::Vertices;
+};
+
+///////////////////////////////////////////////////////////////////////////////////////////////////////////////////////////////////////////////////
+
+USTRUCT()
+struct FDataflowEditorVertexAttributePaintToolDisplayProperties
+{
+	GENERATED_BODY()
+
+	// Vertex color display mode (Greyscale, Ramp, ... )
+	UPROPERTY()
+	EDataflowEditorToolColorMode ColorMode = EDataflowEditorToolColorMode::Greyscale;
+
+	// Color ramp to use when the color mode is set to "Ramp"
+	UPROPERTY()
+	FDataflowColorRamp ColorRamp;
+};
+
+///////////////////////////////////////////////////////////////////////////////////////////////////////////////////////////////////////////////////
+
+UCLASS(MinimalAPI)
+class UDataflowEditorVertexAttributePaintToolProperties : public UInteractiveToolPropertySet
+{
+	GENERATED_BODY()
+
+public:
+
+	// Edit mode (brush,mesh...)
+	UPROPERTY()
+	EDataflowEditorToolEditMode EditingMode = EDataflowEditorToolEditMode::Brush;
+	
+	UPROPERTY()
+	FDataflowEditorVertexAttributePaintToolBrushProperties BrushProperties;
+
+	UPROPERTY()
+	FDataflowEditorVertexAttributePaintToolGradientProperties GradientProperties;
+
+	UPROPERTY()
+	FDataflowEditorVertexAttributePaintToolSelectionProperties SelectionProperties;
+
+	UPROPERTY()
+	FDataflowEditorVertexAttributePaintToolDisplayProperties DisplayProperties;
+};
+
+///////////////////////////////////////////////////////////////////////////////////////////////////////////////////////////////////////////////////
+
+struct FDataflowEditorVertexAttributePaintToolData
+{
+public:
+	void Setup(FDynamicMesh3& InMesh, FDataflowVertexAttributeEditableNode* InNodeToUpdate, TObjectPtr<UDataflowContextObject> InDataflowEditorContextObject);
+	void Shutdown();
+
+	void BeginChange(UDynamicMeshComponent* Component);
+	TUniquePtr<UE::DataflowEditorVertexAttributePaintTool::Private::FMeshChange> EndChange(UDynamicMeshComponent* Component);
+
+	bool IsValid() const { return (ActiveWeightMap != nullptr); }
+
+	float GetValue(int32 VertexIdx) const;
+	float GetAverageValue(const TArray<int32>& Vertices) const;
+
+	// lambda signature : float GetValue(UE::Geometry::FDynamicMeshWeightAttribute& AttributeMap, int32 AttributeIndex)
+	template <typename TLambda>
+	void SetValue(int32 VertexIdx, TLambda Lambda)
+	{
+		if (ActiveWeightMap)
+		{
+			if (bHaveDynamicMeshToWeightConversion)
+			{
+				for (const int32 Idx : WeightToDynamicMesh[DynamicMeshToWeight[VertexIdx]])
+				{
+					ActiveWeightEditChangeTracker->SaveVertexOneRingTriangles(Idx, true);
+					const float Value = Lambda(*ActiveWeightMap, Idx);
+					ActiveWeightMap->SetValue(Idx, &Value);
+				}
+			}
+			else
+			{
+				ActiveWeightEditChangeTracker->SaveVertexOneRingTriangles(VertexIdx, true);
+				const float Value = Lambda(*ActiveWeightMap, VertexIdx);
+				ActiveWeightMap->SetValue(VertexIdx, &Value);
+			}
+		}
+	}
+
+	void UpdateNode(TObjectPtr<UDataflowContextObject> InDataflowEditorContextObject);
+
+private:
+	// DynamicMesh might be unwelded mesh, but weights are on the welded mesh.
+	bool bHaveDynamicMeshToWeightConversion = false;
+	TArray<int32> DynamicMeshToWeight;
+	TArray<TArray<int32>> WeightToDynamicMesh;
+
+	// Corresponding node to read from and write back to
+	FDataflowVertexAttributeEditableNode* NodeToUpdate = nullptr;
+
+	UE::Geometry::FDynamicMeshWeightAttribute* ActiveWeightMap = nullptr;
+
+	TUniquePtr<UE::Geometry::FDynamicMeshChangeTracker> ActiveWeightEditChangeTracker;
+};
+
+///////////////////////////////////////////////////////////////////////////////////////////////////////////////////////////////////////////////////
+
+/**
+ * Dataflow Editor tool to paint a vertex attribute as vertex colors
+ */
+UCLASS(MinimalAPI)
+class UDataflowEditorVertexAttributePaintTool : public UMeshSculptToolBase
+{
+	GENERATED_BODY()
+
+public:
+	UE_API virtual void RegisterActions(FInteractiveToolActionSet& ActionSet) override;
+
+	UE_API virtual void Setup() override;
+	UE_API virtual void Shutdown(EToolShutdownType ShutdownType) override;
+	UE_API virtual void DrawHUD(FCanvas* Canvas, IToolsContextRenderAPI* RenderAPI) override;
+	UE_API virtual void Render(IToolsContextRenderAPI* RenderAPI) override;
+	UE_API virtual void UpdateMaterialMode(EMeshEditingMaterialModes MaterialMode) override;
+	UE_API virtual void OnTick(float DeltaTime) override;
+
+	virtual bool HasCancel() const override { return true; }
+	virtual bool HasAccept() const override { return true; }
+	UE_API virtual bool CanAccept() const override;
+
+	UE_API virtual bool OnUpdateHover(const FInputDeviceRay& DevicePos) override;
+
+	UE_API virtual void OnPropertyModified(UObject* PropertySet, FProperty* Property) override;
+
+	UE_API bool IsInBrushSubMode() const;
+	UE_API bool IsVolumetricBrush() const;
+
+	UE_API virtual void CommitResult(UBaseDynamicMeshComponent* Component, bool bModifiedTopology) override;
+
+	UE_API void SetDataflowEditorContextObject(TObjectPtr<UDataflowContextObject> InDataflowEditorContextObject);
+	void SetEditorMode(UDataflowEditorMode* InMode) { Mode = InMode; }
+
+	virtual void SetupBrushEditBehaviorSetup(ULocalTwoAxisPropertyEditInputBehavior& OutBehavior) override;
+
+	//~ UObject interface
+
+	static UE_API void AddReferencedObjects(UObject* InThis, FReferenceCollector& Collector);
+
+public:
+	/** properties of the tool */
+	UPROPERTY()
+	TObjectPtr<UDataflowEditorVertexAttributePaintToolProperties> ToolProperties;
+
+	UPROPERTY()
+	TObjectPtr<UDataflowVertexAttributePaintBrushOpProps> PaintBrushOpProperties;
+
+	UPROPERTY()
+	TObjectPtr<UDataflowWeightMapEraseBrushOpProps> EraseBrushOpProperties;
+
+	UE_API bool HaveVisibilityFilter() const;
+	UE_API void ApplyVisibilityFilter(const TArray<int32>& Vertices, TArray<int32>& VisibleVertices);
+	UE_API void ApplyVisibilityFilter(TSet<int32>& Triangles, TArray<int32>& ROIBuffer, TArray<int32>& OutputBuffer);
+
+	// we override these so we can update the separate BrushSize property added for this tool
+	UE_API virtual void IncreaseBrushRadiusAction();
+	UE_API virtual void DecreaseBrushRadiusAction();
+	UE_API virtual void IncreaseBrushRadiusSmallStepAction();
+	UE_API virtual void DecreaseBrushRadiusSmallStepAction();
+
+	UE_API float GetBrushMinRadius() const;
+	UE_API float GetBrushMaxRadius() const;
+
+	UE_API void SetColorMode(EDataflowEditorToolColorMode NewColorMode);
+
+	UE_API bool HasSelection() const;
+	UE_API void GrowSelection() const;
+	UE_API void ShrinkSelection() const;
+	UE_API void FloodSelection() const;
+	UE_API void SelectBorder() const;
+	UE_API void SetComponentSelectionMode(EComponentSelectionMode NewMode);
+
+	UE_API void CopyAverageFromSelectionToClipboard();
+	UE_API void PasteValueToSelectionFromClipboard();
+	UE_API void PruneSelection(float Threshold);
+
+	// selection operations
+	UE_API void ApplyValueToSelection(EDataflowEditorToolEditOperation Operation, float InValue);
+
+	UE_API TArray<int32> GetSelectedVertices() const;
+
+protected:
+	// UMeshSculptToolBase API
+	virtual UBaseDynamicMeshComponent* GetSculptMeshComponent() override;
+	virtual FDynamicMesh3* GetBaseMesh()  override { check(false); return nullptr; }
+	virtual const FDynamicMesh3* GetBaseMesh() const  override { check(false); return nullptr; }
+
+	UE_API virtual int32 FindHitSculptMeshTriangleConst(const FRay3d& LocalRay) const override;
+	UE_API virtual int32 FindHitTargetMeshTriangleConst(const FRay3d& LocalRay) const override;
+	UE_API virtual void UpdateHitSculptMeshTriangle(int32 TriangleID, const FRay3d& LocalRay) override;
+
+	UE_API virtual void OnBeginStroke(const FRay& WorldRay) override;
+	UE_API virtual void OnEndStroke() override;
+
+	UE_API void OnColorRampChanged(TArray<FRichCurve*> Curves);
+
+	void CreateDefaultColorRamp(FDataflowColorRamp& OutColorRamp);
+
+	UE_API virtual TUniquePtr<FMeshSculptBrushOp>& GetActiveBrushOp();
+
+	virtual bool SharesBrushPropertiesChanges() const override { return false; }
+	// end UMeshSculptToolBase API
+
+protected:
+	//
+	// Marquee Support
+	//
+public:
+	UPROPERTY()
+	TObjectPtr<UPolyLassoMarqueeMechanic> PolyLassoMechanic;
+
+protected:
+	UE_API void OnPolyLassoFinished(const FCameraPolyLasso& Lasso, bool bCanceled);
+
+
+	// 
+	// Gradient Support
+	//
+	FToolDataVisualizer GradientSelectionRenderer;
+
+	UE::Geometry::FGroupTopologySelection LowValueGradientVertexSelection;
+	UE::Geometry::FGroupTopologySelection HighValueGradientVertexSelection;
+
+	UE_API void ComputeGradient();
+	UE_API void OnSelectionModified();
+
+	//
+	// Internals
+	//
+
+protected:
+
+	UPROPERTY()
+	TObjectPtr<UMeshElementsVisualizer> MeshElementsDisplay = nullptr;
+
+	UPROPERTY()
+	TObjectPtr<UDataflowContextObject> DataflowEditorContextObject = nullptr;
+
+	UPROPERTY()
+	TObjectPtr<UPreviewMesh> PreviewMesh = nullptr;
+	FMeshDescription PreviewMeshDescription;
+
+	UPROPERTY()
+	TObjectPtr<UToolMeshSelector> MeshSelector;
+
+	// realtime visualization
+	UE_API void OnDynamicMeshComponentChanged(UDynamicMeshComponent* Component, const FMeshVertexChange* Change, bool bRevert);
+	FDelegateHandle OnDynamicMeshComponentChangedHandle;
+
+	UE_API float GetCurrentWeightValueUnderBrush() const;
+	FVector3d CurrentBaryCentricCoords;
+	UE_API int32 GetBrushNearestVertex() const;
+
+	TSet<int32> AccumulatedTriangleROI;
+	bool bUndoUpdatePending = false;
+	TArray<int> NormalsBuffer;
+	UE_API void WaitForPendingUndoRedoUpdate();
+
+	TArray<int> TempROIBuffer;
+	TArray<int> VertexROI;
+	TArray<bool> VisibilityFilterBuffer;
+	TSet<int> VertexSetBuffer;
+	TSet<int> TriangleROI;
+	UE_API void UpdateROI(const FSculptBrushStamp& CurrentStamp);
+
+	UE_API bool UpdateStampPosition(const FRay& WorldRay);
+	UE_API bool ApplyStamp();
+
+	UE::Geometry::FDynamicMeshOctree3 Octree;
+
+	UE_API void UpdateBrushType(EDataflowEditorToolEditOperation EditOperation);
+	UE_API bool UpdateBrushPosition(const FRay& WorldRay);
+
+	bool GetInEraseStroke()
+	{
+		// Re-use the smoothing stroke key (shift) for erase stroke in the weight paint tool
+		return GetInSmoothingStroke();
+	}
+
+
+	bool bPendingPickWeight = false;
+
+	TArray<int32> ROITriangleBuffer;
+	TArray<double> ROIWeightValueBuffer;
+	UE_API bool SyncMeshWithWeightBuffer(FDynamicMesh3* Mesh);
+	UE_API bool SyncWeightBufferWithMesh(const FDynamicMesh3* Mesh);
+	UE_API void SetVerticesToWeightMap(const TSet<int32>& Vertices, double WeightValue, bool bIsErase);
+
+	UE_API void BeginChange();
+	UE_API void EndChange();
+
+	TArray<FVector3d> TriNormals;
+	TArray<int32> UVSeamEdges;
+	TArray<int32> NormalSeamEdges;
+	UE_API void PrecomputeFilterData();
+
+	FDataflowEditorVertexAttributePaintToolData VertexData;
+
+protected:
+	virtual bool ShowWorkPlane() const override { return false; }
+
+	friend class UDataflowEditorVertexAttributePaintToolBuilder;
+
+	bool bAnyChangeMade = false;
+
+	/** update the vertex color of the mesh from the attribute */
+	UE_API void UpdateVertexColorOverlay(const TSet<int>* TrianglesToUpdate = nullptr);
+
+	UDataflowEditorMode* Mode = nullptr;
+
+	static constexpr int32 PaintBrushId = 0;
+	static constexpr int32 EraseBrushId = 1;
+};
+
+
+
+#undef UE_API
+
diff --git a/Engine/Plugins/Experimental/Dataflow/Source/DataflowNodes/Private/Dataflow/DataflowCollectionAddScalarVertexPropertyNode.cpp b/Engine/Plugins/Experimental/Dataflow/Source/DataflowNodes/Private/Dataflow/DataflowCollectionAddScalarVertexPropertyNode.cpp
--- a/Engine/Plugins/Experimental/Dataflow/Source/DataflowNodes/Private/Dataflow/DataflowCollectionAddScalarVertexPropertyNode.cpp
+++ b/Engine/Plugins/Experimental/Dataflow/Source/DataflowNodes/Private/Dataflow/DataflowCollectionAddScalarVertexPropertyNode.cpp
@@ -137,12 +137,16 @@ TArray<UE::Dataflow::FRenderingParameter> DataflowAddScalarVertexPropertyCallbac
 PRAGMA_ENABLE_DEPRECATION_WARNINGS
 
 
+FDataflowVertexAttributeEditableNode::FDataflowVertexAttributeEditableNode(const UE::Dataflow::FNodeParameters& InParam, FGuid InGuid)
+	: FDataflowNode(InParam, InGuid)
+{}
+
 //
 // FDataflowCollectionAddScalarVertexPropertyNode
 //
 
 FDataflowCollectionAddScalarVertexPropertyNode::FDataflowCollectionAddScalarVertexPropertyNode(const UE::Dataflow::FNodeParameters& InParam, FGuid InGuid)
-	: FDataflowNode(InParam, InGuid)
+	: FDataflowVertexAttributeEditableNode(InParam, InGuid)
 {
 	RegisterInputConnection(&Collection);
 	RegisterInputConnection(&AttributeKey);
@@ -228,6 +232,39 @@ FCollectionAttributeKey FDataflowCollectionAddScalarVertexPropertyNode::GetWeigh
 	return Key;
 }
 
+void FDataflowCollectionAddScalarVertexPropertyNode::GetVertexAttributeValues(UE::Dataflow::FContext& Context, TArray<float>& OutValues) const
+{
+	const FCollectionAttributeKey InKey = GetWeightAttributeKey(Context);
+	const FName InAttribName(InKey.Attribute);
+	const FName InGroupName(InKey.Group);
+	const FManagedArrayCollection& InCollection = GetValue(Context, &Collection, Collection);
+	
+	if (const TManagedArray<float>* AttributeArray = InCollection.FindAttributeTyped<float>(InAttribName, InGroupName))
+	{
+		OutValues.SetNumUninitialized(AttributeArray->Num());
+		ExtractVertexWeights(AttributeArray->GetConstArray(), OutValues);
+	}
+	else
+	{
+		OutValues.Init(0.0f, InCollection.NumElements(InGroupName));
+	}
+}
+
+void FDataflowCollectionAddScalarVertexPropertyNode::SetVertexAttributeValues(UE::Dataflow::FContext& Context, const TArray<float>& InValues)
+{
+	const FCollectionAttributeKey InKey = GetWeightAttributeKey(Context);
+	const FName InAttribName(InKey.Attribute);
+	const FName InGroupName(InKey.Group);
+	const FManagedArrayCollection& InCollection = GetValue(Context, &Collection, Collection);
+
+	if (const TManagedArray<float>* AttributeArray = InCollection.FindAttributeTyped<float>(InAttribName, InGroupName))
+	{
+		ReportVertexWeights(AttributeArray->GetConstArray(), InValues, {});
+	}
+
+	Invalidate();
+}
+
 void FDataflowCollectionAddScalarVertexPropertyNode::ReportVertexWeights(const TArray<float>& SetupWeights, const TArray<float>& FinalWeights,
 		const TArray<int32>& WeightIndices)
 {
diff --git a/Engine/Plugins/Experimental/Dataflow/Source/DataflowNodes/Public/Dataflow/DataflowCollectionAddScalarVertexPropertyNode.h b/Engine/Plugins/Experimental/Dataflow/Source/DataflowNodes/Public/Dataflow/DataflowCollectionAddScalarVertexPropertyNode.h
--- a/Engine/Plugins/Experimental/Dataflow/Source/DataflowNodes/Public/Dataflow/DataflowCollectionAddScalarVertexPropertyNode.h
+++ b/Engine/Plugins/Experimental/Dataflow/Source/DataflowNodes/Public/Dataflow/DataflowCollectionAddScalarVertexPropertyNode.h
@@ -80,9 +80,26 @@ enum class EDataflowWeightMapOverrideType : uint8
 	ReplaceChanged,
 };
 
+USTRUCT()
+struct FDataflowVertexAttributeEditableNode : public FDataflowNode
+{
+	GENERATED_USTRUCT_BODY()
+
+	UE_API FDataflowVertexAttributeEditableNode() {};
+
+	UE_API FDataflowVertexAttributeEditableNode(const UE::Dataflow::FNodeParameters& InParam, FGuid InGuid = FGuid::NewGuid());
+
+public:
+	/** set the absolute vertex attribute values */
+	UE_API virtual void GetVertexAttributeValues(UE::Dataflow::FContext& Context, TArray<float>& OutValues) const { ensure(false); }
+
+	/** Get the absolute vertex attribute values */
+	UE_API virtual void SetVertexAttributeValues(UE::Dataflow::FContext& Context, const TArray<float>& InValues) { ensure(false); }
+};
+
 /** Scalar vertex properties. */
 USTRUCT(Meta = (DataflowCollection))
-struct FDataflowCollectionAddScalarVertexPropertyNode : public FDataflowNode
+struct FDataflowCollectionAddScalarVertexPropertyNode : public FDataflowVertexAttributeEditableNode
 {
 	GENERATED_USTRUCT_BODY()
 	DATAFLOW_NODE_DEFINE_INTERNAL(FDataflowCollectionAddScalarVertexPropertyNode, "PaintWeightMap", "Collection", "Paint a weight map and save it to collection")
@@ -132,6 +149,11 @@ private:
 	static constexpr float ReplaceChangedPassthroughValue = UE_BIG_NUMBER;
 	
 	UE_API virtual void Evaluate(UE::Dataflow::FContext& Context, const FDataflowOutput* Out) const override;
+
+	//~ Begin - FDataflowVertexAttributeEditableNode API
+	virtual void GetVertexAttributeValues(UE::Dataflow::FContext& Context, TArray<float>& OutValues) const override;
+	virtual void SetVertexAttributeValues(UE::Dataflow::FContext& Context, const TArray<float>& InValues) override;
+	//~ End - FDataflowVertexAttributeEditableNode API
 };
 
 #undef UE_API
