Dataflow Nodes : skin weight node fixes
- Use the skelmesh ref skeleton for the edit and exposed it as an input
- Exposed the skeletal mesh as input and output for the transfer groom skin weight  node
- Updated the skin weight node to take a context when writing and reading skin weights for the tool
- Update the tool to pass the context to the node weight read and write functions

diff --git a/Engine/Plugins/Experimental/Dataflow/Source/DataflowEditor/Private/DataflowEditorTools/DataflowEditorSkinWeightsPaintTool.cpp b/Engine/Plugins/Experimental/Dataflow/Source/DataflowEditor/Private/DataflowEditorTools/DataflowEditorSkinWeightsPaintTool.cpp
--- a/Engine/Plugins/Experimental/Dataflow/Source/DataflowEditor/Private/DataflowEditorTools/DataflowEditorSkinWeightsPaintTool.cpp
+++ b/Engine/Plugins/Experimental/Dataflow/Source/DataflowEditor/Private/DataflowEditorTools/DataflowEditorSkinWeightsPaintTool.cpp
@@ -144,7 +144,7 @@ bool UDataflowEditorSkinWeightsPaintTool::ExtractSkinWeights(TArray<TArray<int32
 				CurrentWeights.SetNumZeroed(SetupWeights.Num()); 
 				CurrentIndices.SetNumZeroed(SetupIndices.Num());
 						
-				SkinWeightNode->ExtractVertexWeights(SetupIndices, SetupWeights,
+				SkinWeightNode->ExtractVertexWeights(*DataflowContext, SetupIndices, SetupWeights,
 					TArrayView<TArray<int32>>(CurrentIndices), TArrayView<TArray<float>>(CurrentWeights));
 
 				return true;
@@ -202,7 +202,8 @@ void UDataflowEditorSkinWeightsPaintTool::OnShutdown(EToolShutdownType ShutdownT
 {
 	USkinWeightsPaintTool::OnShutdown(ShutdownType);
 
-	if(Target && ShutdownType == EToolShutdownType::Accept)
+	TSharedPtr<UE::Dataflow::FEngineContext> DataflowContext = DataflowEditorContextObject->GetDataflowContext();
+	if(DataflowContext && Target && ShutdownType == EToolShutdownType::Accept)
 	{
 		const int32 VertexOffset = GetVertexOffset();
 		if(VertexOffset != INDEX_NONE)
@@ -212,8 +213,8 @@ void UDataflowEditorSkinWeightsPaintTool::OnShutdown(EToolShutdownType ShutdownT
 
 			CurrentWeights.SetNumZeroed(SetupWeights.Num()); 
 			CurrentIndices.SetNumZeroed(SetupIndices.Num());
-						
-			SkinWeightNode->ExtractVertexWeights(SetupIndices, SetupWeights,
+					
+			SkinWeightNode->ExtractVertexWeights(*DataflowContext, SetupIndices, SetupWeights,
 				TArrayView<TArray<int32>>(CurrentIndices), TArrayView<TArray<float>>(CurrentWeights));
 			
 			if(FCleanedEditMesh* EditedMesh = EditedMeshes.Find(EMeshLODIdentifier::LOD0))
@@ -258,7 +259,7 @@ void UDataflowEditorSkinWeightsPaintTool::OnShutdown(EToolShutdownType ShutdownT
 					}
 				}
 			}
-			SkinWeightNode->ReportVertexWeights(SetupIndices, SetupWeights, CurrentIndices, CurrentWeights);
+			SkinWeightNode->ReportVertexWeights(*DataflowContext, SetupIndices, SetupWeights, CurrentIndices, CurrentWeights);
 			SkinWeightNode->Invalidate();
 
 			// Avoid rebuilding the skeletal mesh after updating the skin weights
diff --git a/Engine/Plugins/Experimental/Dataflow/Source/DataflowNodes/Private/Dataflow/DataflowCollectionEditSkeletonBonesNode.cpp b/Engine/Plugins/Experimental/Dataflow/Source/DataflowNodes/Private/Dataflow/DataflowCollectionEditSkeletonBonesNode.cpp
--- a/Engine/Plugins/Experimental/Dataflow/Source/DataflowNodes/Private/Dataflow/DataflowCollectionEditSkeletonBonesNode.cpp
+++ b/Engine/Plugins/Experimental/Dataflow/Source/DataflowNodes/Private/Dataflow/DataflowCollectionEditSkeletonBonesNode.cpp
@@ -56,7 +56,7 @@ void FDataflowCollectionEditSkeletonBonesNode::AddPrimitiveComponents(UE::Datafl
 			}
 			if(bNeedsConstruction)
 			{
-				bValidSkeletalMeshes = UE::Dataflow::Private::BuildSkeletalMeshes(SkeletalMeshes, RenderCollection, TransientSkeleton);
+				bValidSkeletalMeshes = UE::Dataflow::Private::BuildSkeletalMeshes(SkeletalMeshes, RenderCollection, TransientSkeleton->GetReferenceSkeleton());
 				if(!bValidSkeletalMeshes)
 				{
 					SkeletalMeshes.Reset();
diff --git a/Engine/Plugins/Experimental/Dataflow/Source/DataflowNodes/Private/Dataflow/DataflowCollectionEditSkinWeightsNode.cpp b/Engine/Plugins/Experimental/Dataflow/Source/DataflowNodes/Private/Dataflow/DataflowCollectionEditSkinWeightsNode.cpp
--- a/Engine/Plugins/Experimental/Dataflow/Source/DataflowNodes/Private/Dataflow/DataflowCollectionEditSkinWeightsNode.cpp
+++ b/Engine/Plugins/Experimental/Dataflow/Source/DataflowNodes/Private/Dataflow/DataflowCollectionEditSkinWeightsNode.cpp
@@ -174,6 +174,7 @@ FDataflowCollectionEditSkinWeightsNode::FDataflowCollectionEditSkinWeightsNode(c
 	RegisterInputConnection(&Collection);
 	RegisterInputConnection(&BoneIndicesKey);
 	RegisterInputConnection(&BoneWeightsKey);
+	RegisterInputConnection(&SkeletalMesh);
 	RegisterOutputConnection(&Collection, &Collection);
 	RegisterOutputConnection(&BoneIndicesKey, &BoneIndicesKey);
 	RegisterOutputConnection(&BoneWeightsKey, &BoneWeightsKey);
@@ -205,9 +206,11 @@ void FDataflowCollectionEditSkinWeightsNode::AddPrimitiveComponents(UE::Dataflow
 				SkeletalMeshes.Add(NewObject<USkeletalMesh>(NodeOwner, FName(*SkeletalMeshName), RF_Transient));
 			}
 		}
-		if(bNeedsConstruction)
+
+		TObjectPtr<USkeletalMesh> InSkeletalMesh = GetValue(Context, &SkeletalMesh, SkeletalMesh);
+		if(bNeedsConstruction && InSkeletalMesh)
 		{
-			bValidSkeletalMeshes = UE::Dataflow::Private::BuildSkeletalMeshes(SkeletalMeshes, RenderCollection, ObjectSkeleton);
+			bValidSkeletalMeshes = UE::Dataflow::Private::BuildSkeletalMeshes(SkeletalMeshes, RenderCollection, InSkeletalMesh->GetRefSkeleton());
 			if(!bValidSkeletalMeshes)
 			{
 				SkeletalMeshes.Reset();
@@ -226,6 +229,14 @@ void FDataflowCollectionEditSkinWeightsNode::AddPrimitiveComponents(UE::Dataflow
 	}
 }
 
+void FDataflowCollectionEditSkinWeightsNode::OnPropertyChanged(UE::Dataflow::FContext& Context, const FPropertyChangedEvent& InPropertyChangedEvent)
+{
+	if (InPropertyChangedEvent.GetPropertyName() == GET_MEMBER_NAME_CHECKED(FDataflowCollectionEditSkinWeightsNode, SkeletalMesh))
+	{
+		SkinWeights.Reset();
+	}
+}
+
 void FDataflowCollectionEditSkinWeightsNode::Evaluate(UE::Dataflow::FContext& Context, const FDataflowOutput* Out) const
 {
 	using namespace UE::Dataflow;
@@ -249,7 +260,7 @@ void FDataflowCollectionEditSkinWeightsNode::Evaluate(UE::Dataflow::FContext& Co
 				FinalIndices.SetNum(SetupIndices.Num());
 				FinalWeights.SetNum(SetupWeights.Num());
 				
-				ExtractVertexWeights(SetupIndices, SetupWeights, TArrayView<TArray<int32>>(FinalIndices.GetData(), FinalIndices.Num()),
+				ExtractVertexWeights(Context, SetupIndices, SetupWeights, TArrayView<TArray<int32>>(FinalIndices.GetData(), FinalIndices.Num()),
 				TArrayView<TArray<float>>(FinalWeights.GetData(), FinalWeights.Num()));
 
 				SetAttributeWeights(InCollection, BoneIndicesKeyValue, BoneWeightsKeyValue, FinalIndices, FinalWeights);
@@ -268,12 +279,26 @@ void FDataflowCollectionEditSkinWeightsNode::Evaluate(UE::Dataflow::FContext& Co
 	}
 }
 
-void FDataflowCollectionEditSkinWeightsNode::ReportVertexWeights(const TArray<TArray<int32>>& SetupIndices, const TArray<TArray<float>>& SetupWeights, const TArray<TArray<int32>>& FinalIndices, const TArray<TArray<float>>& FinalWeights)
+int32 FDataflowCollectionEditSkinWeightsNode::ComputeInputRefSkeletonHashFromInput(UE::Dataflow::FContext& Context) const
+{
+	const TObjectPtr<USkeletalMesh> InSkeletalMesh = GetValue(Context, &SkeletalMesh, SkeletalMesh);
+	if (InSkeletalMesh)
+	{
+		return UE::Dataflow::Private::ComputeHasFromRefSkeleton(InSkeletalMesh->GetRefSkeleton());
+	}
+	return 0;
+}
+
+void FDataflowCollectionEditSkinWeightsNode::ReportVertexWeights(UE::Dataflow::FContext& Context, const TArray<TArray<int32>>& SetupIndices, const TArray<TArray<float>>& SetupWeights, const TArray<TArray<int32>>& FinalIndices, const TArray<TArray<float>>& FinalWeights)
 {
 	check(SetupWeights.Num() == FinalWeights.Num());
 	check(SetupWeights.Num() == SetupIndices.Num());
 	check(FinalWeights.Num() == FinalIndices.Num());
 	
+	// update ref skeletal hash for the current stored skinweights
+	SkinWeightsRefSkeletonHash = ComputeInputRefSkeletonHashFromInput(Context);
+
+	// update the weights themselves
 	SkinWeights.SetNumZeroed(FinalWeights.Num());
 	
 	for (int32 VertexIndex = 0, NumVertices = FinalWeights.Num(); VertexIndex < NumVertices; ++VertexIndex)
@@ -283,14 +308,17 @@ void FDataflowCollectionEditSkinWeightsNode::ReportVertexWeights(const TArray<TA
 	}
 }
 
-void FDataflowCollectionEditSkinWeightsNode::ExtractVertexWeights(const TArray<TArray<int32>>& SetupIndices, const TArray<TArray<float>>& SetupWeights,
+void FDataflowCollectionEditSkinWeightsNode::ExtractVertexWeights(UE::Dataflow::FContext& Context, const TArray<TArray<int32>>& SetupIndices, const TArray<TArray<float>>& SetupWeights,
 	TArrayView<TArray<int32>> FinalIndices, TArrayView<TArray<float>> FinalWeights) const
 {
 	check(SetupWeights.Num() == FinalWeights.Num());
 	check(SetupWeights.Num() == SetupIndices.Num());
 	check(FinalWeights.Num() == FinalIndices.Num());
 
-	if(SkinWeights.Num() == FinalWeights.Num())
+	// we only returned the stored weight if they match the input refskeleton, other wise, returned the original weights
+	const int32 RefSkeletonHashFromInput = ComputeInputRefSkeletonHashFromInput(Context);
+
+	if(SkinWeights.Num() == FinalWeights.Num() && RefSkeletonHashFromInput == SkinWeightsRefSkeletonHash)
 	{
 		for (int32 VertexIndex = 0, NumVertices = FinalWeights.Num(); VertexIndex < NumVertices; ++VertexIndex)
 		{
@@ -389,10 +417,11 @@ void FDataflowCollectionEditSkinWeightsNode::DebugDraw(UE::Dataflow::FContext& C
 		IDataflowDebugDrawInterface& DataflowRenderingInterface,
 		const FDebugDrawParameters& DebugDrawParameters) const
 {
-	if(ObjectSkeleton)
+	TObjectPtr<USkeletalMesh> InSkeletalMesh = GetValue(Context, &SkeletalMesh, SkeletalMesh);
+	if(InSkeletalMesh)
 	{
 		TRefCountPtr<FDataflowDebugDrawSkeletonObject> SkeletonObject = MakeDebugDrawObject<FDataflowDebugDrawSkeletonObject>(
-			DataflowRenderingInterface.ModifyDataflowElements(), ObjectSkeleton->GetReferenceSkeleton());
+			DataflowRenderingInterface.ModifyDataflowElements(), InSkeletalMesh->GetRefSkeleton());
 	
 		DataflowRenderingInterface.DrawObject(TRefCountPtr<IDataflowDebugDrawObject>(SkeletonObject));
 
@@ -410,13 +439,13 @@ bool FDataflowCollectionEditSkinWeightsNode::CanDebugDrawViewMode(const FName& V
 
 #endif
 
-int32 FDataflowCollectionEditSkinWeightsNode::GetSkeletalMeshOffset(const TObjectPtr<USkeletalMesh>& SkeletalMesh) const
+int32 FDataflowCollectionEditSkinWeightsNode::GetSkeletalMeshOffset(const TObjectPtr<USkeletalMesh>& InSkeletalMesh) const
 {
 #if WITH_EDITORONLY_DATA
 	int32 SkeletalMeshOffset = 0;
 	for(int32 SkeletalMeshIndex = 0, NumMeshes = SkeletalMeshes.Num(); SkeletalMeshIndex < NumMeshes; ++SkeletalMeshIndex)
 	{
-		if (SkeletalMeshes[SkeletalMeshIndex] == SkeletalMesh)
+		if (SkeletalMeshes[SkeletalMeshIndex] == InSkeletalMesh)
 		{
 			return SkeletalMeshOffset;
 		}
diff --git a/Engine/Plugins/Experimental/Dataflow/Source/DataflowNodes/Private/Dataflow/DataflowCollectionSkeletalMeshUtils.h b/Engine/Plugins/Experimental/Dataflow/Source/DataflowNodes/Private/Dataflow/DataflowCollectionSkeletalMeshUtils.h
--- a/Engine/Plugins/Experimental/Dataflow/Source/DataflowNodes/Private/Dataflow/DataflowCollectionSkeletalMeshUtils.h
+++ b/Engine/Plugins/Experimental/Dataflow/Source/DataflowNodes/Private/Dataflow/DataflowCollectionSkeletalMeshUtils.h
@@ -1,4 +1,4 @@
-﻿// Copyright Epic Games, Inc. All Rights Reserved.
+// Copyright Epic Games, Inc. All Rights Reserved.
 #pragma once
 
 #include "Dataflow/CollectionRenderingPatternUtility.h"
@@ -77,9 +77,9 @@ static bool CreateSkeletalMesh(USkeletalMesh* SkeletalMesh, const TArray<UE::Geo
 }
 
 static bool BuildSkeletalMeshes(TArray<TObjectPtr<USkeletalMesh>>& SkeletalMeshes, 
-	const TSharedPtr<const FManagedArrayCollection>& RenderCollection, const TObjectPtr<USkeleton>& ObjectSkeleton)
+	const TSharedPtr<const FManagedArrayCollection>& RenderCollection, const FReferenceSkeleton& ReferenceSkeleton)
 {
-	if(RenderCollection.IsValid() && ObjectSkeleton)
+	if(RenderCollection.IsValid())
 	{
 		GeometryCollection::Facades::FRenderingFacade Facade(*RenderCollection);
 		if(Facade.IsValid())
@@ -94,7 +94,7 @@ static bool BuildSkeletalMeshes(TArray<TObjectPtr<USkeletalMesh>>& SkeletalMeshe
 					UE::Geometry::FDynamicMesh3 DynamicMesh;
 					UE::Dataflow::Conversion::RenderingFacadeToDynamicMesh(Facade, MeshIndex, DynamicMesh);
 
-					if(!CreateSkeletalMesh(SkeletalMeshes[MeshIndex], {DynamicMesh}, ObjectSkeleton->GetReferenceSkeleton()))
+					if(!CreateSkeletalMesh(SkeletalMeshes[MeshIndex], {DynamicMesh}, ReferenceSkeleton))
 					{
 						bValidSkeletalMeshes = false;
 					}
@@ -106,5 +106,20 @@ static bool BuildSkeletalMeshes(TArray<TObjectPtr<USkeletalMesh>>& SkeletalMeshe
 	return false;
 }
 	
+static int32 ComputeHasFromRefSkeleton(const FReferenceSkeleton& ReferenceSkeleton)
+{
+	const TArray<FMeshBoneInfo>& BonesInfo = ReferenceSkeleton.GetRefBoneInfo();
+	int32 Hash = 0;
+	for (const FMeshBoneInfo& BoneInfo : BonesInfo)
+	{
+		Hash = HashCombine(Hash, 
+					HashCombine(
+						GetTypeHash(BoneInfo.Name), 
+						GetTypeHash(BoneInfo.ParentIndex)
+					));
+	}
+	return Hash;
+}
+
 }
 
diff --git a/Engine/Plugins/Experimental/Dataflow/Source/DataflowNodes/Public/Dataflow/DataflowCollectionEditSkinWeightsNode.h b/Engine/Plugins/Experimental/Dataflow/Source/DataflowNodes/Public/Dataflow/DataflowCollectionEditSkinWeightsNode.h
--- a/Engine/Plugins/Experimental/Dataflow/Source/DataflowNodes/Public/Dataflow/DataflowCollectionEditSkinWeightsNode.h
+++ b/Engine/Plugins/Experimental/Dataflow/Source/DataflowNodes/Public/Dataflow/DataflowCollectionEditSkinWeightsNode.h
@@ -62,8 +62,8 @@ public:
 	FCollectionAttributeKey BoneWeightsKey;
 
 	/** Skeletal mesh to extract the skeleton from for the skinning */
-	UPROPERTY(EditAnywhere, Category = "Skeleton Binding")
-	TObjectPtr<USkeleton> ObjectSkeleton = nullptr;
+	UPROPERTY(EditAnywhere, Category = "Skeleton Binding", meta = (DataflowInput))
+	TObjectPtr<USkeletalMesh> SkeletalMesh = nullptr;
 
 	/** Boolean to use a compressed format (FVector4f, FIntVector) to store the skin weights */
 	UPROPERTY(EditAnywhere, Category = "Skeleton Binding")
@@ -73,14 +73,21 @@ public:
 	UPROPERTY()
 	TArray<FDataflowSkinWeightData> SkinWeights;
 	
+	/** 
+	* Hash of the ref skeleton the stored skin weight are based on 
+	* This allow invalidating the SkinWeights when the ref skeleton does not match anymore instead of having broken weights
+	*/
+	UPROPERTY()
+	int32 SkinWeightsRefSkeletonHash;
+
 	/** Delegate to transfer the bone selection to the tool*/
 	FDataflowBoneSelectionChangedNotifyDelegate OnBoneSelectionChanged;
 
 	/** Report the vertex weights back onto the property ones */
-	UE_API void ReportVertexWeights(const TArray<TArray<int32>>& SetupIndices, const TArray<TArray<float>>& SetupWeights, const TArray<TArray<int32>>& FinalIndices, const TArray<TArray<float>>& FinalWeights);
+	UE_API void ReportVertexWeights(UE::Dataflow::FContext& Context, const TArray<TArray<int32>>&SetupIndices, const TArray<TArray<float>>& SetupWeights, const TArray<TArray<int32>>& FinalIndices, const TArray<TArray<float>>& FinalWeights);
 
 	/** Extract the vertex weights back from the property ones */
-	UE_API void ExtractVertexWeights(const TArray<TArray<int32>>& SetupIndices, const TArray<TArray<float>>& SetupWeights, TArrayView<TArray<int32>> FinalIndices, TArrayView<TArray<float>> FinalWeights) const;
+	UE_API void ExtractVertexWeights(UE::Dataflow::FContext& Context, const TArray<TArray<int32>>& SetupIndices, const TArray<TArray<float>>& SetupWeights, TArrayView<TArray<int32>> FinalIndices, TArrayView<TArray<float>> FinalWeights) const;
 
 	/** Fill the weight attribute values from the collection */
 	UE_API static bool FillAttributeWeights(const FManagedArrayCollection& SelectedCollection,
@@ -113,6 +120,7 @@ private:
 	//~ Begin FDataflowPrimitiveNode interface
 	UE_API virtual TArray<UE::Dataflow::FRenderingParameter> GetRenderParametersImpl() const override;
 	UE_API virtual void Evaluate(UE::Dataflow::FContext& Context, const FDataflowOutput* Out) const override;
+	UE_API virtual void OnPropertyChanged(UE::Dataflow::FContext& Context, const FPropertyChangedEvent& InPropertyChangedEvent) override;
 	UE_API virtual void AddPrimitiveComponents(UE::Dataflow::FContext& Context, const TSharedPtr<const FManagedArrayCollection> RenderCollection,
 		TObjectPtr<UObject> NodeOwner, TObjectPtr<AActor> RootActor, TArray<TObjectPtr<UPrimitiveComponent>>& PrimitiveComponents) override;
 	UE_API virtual void OnInvalidate() override;
@@ -125,6 +133,8 @@ private:
 #endif
 	//~ End FDataflowPrimitiveNode interface
 
+	int32 ComputeInputRefSkeletonHashFromInput(UE::Dataflow::FContext& Context) const;
+
 	/** Transient skeletal mesh built from dataflow render collection */
 	UPROPERTY(Transient)
 	TArray<TObjectPtr<USkeletalMesh>> SkeletalMeshes;
diff --git a/Engine/Plugins/Runtime/HairStrands/Source/HairStrandsDataflow/Private/BuildGroomSkinningNodes.cpp b/Engine/Plugins/Runtime/HairStrands/Source/HairStrandsDataflow/Private/BuildGroomSkinningNodes.cpp
--- a/Engine/Plugins/Runtime/HairStrands/Source/HairStrandsDataflow/Private/BuildGroomSkinningNodes.cpp
+++ b/Engine/Plugins/Runtime/HairStrands/Source/HairStrandsDataflow/Private/BuildGroomSkinningNodes.cpp
@@ -230,11 +230,22 @@ namespace UE::Groom::Private
 	}
 }
 
+FTransferSkinWeightsGroomNode::FTransferSkinWeightsGroomNode(const UE::Dataflow::FNodeParameters& InParam, FGuid InGuid)
+	: FDataflowNode(InParam, InGuid)
+{
+	RegisterInputConnection(&Collection);
+	RegisterInputConnection(&SkeletalMesh);
+	RegisterOutputConnection(&Collection, &Collection);
+	RegisterOutputConnection(&BoneIndicesKey);
+	RegisterOutputConnection(&BoneWeightsKey);
+	RegisterOutputConnection(&SkeletalMesh, &SkeletalMesh);
+}
+
 void FTransferSkinWeightsGroomNode::Evaluate(UE::Dataflow::FContext& Context, const FDataflowOutput* Out) const
 {
-	if (Out->IsA<FManagedArrayCollection>(&Collection))
+	if (Out->IsA(&Collection))
 	{
-		FManagedArrayCollection GroomCollection = GetValue<FManagedArrayCollection>(Context, &Collection);
+		FManagedArrayCollection GroomCollection = GetValue(Context, &Collection);
 
 		if(CurvesType == EGroomCollectionType::Guides)
 		{
@@ -246,14 +257,18 @@ void FTransferSkinWeightsGroomNode::Evaluate(UE::Dataflow::FContext& Context, co
 		}
 		SetValue(Context, MoveTemp(GroomCollection), &Collection);
 	}
-	else if (Out->IsA<FCollectionAttributeKey>(&BoneIndicesKey))
+	else if (Out->IsA(&BoneIndicesKey))
 	{
 		SetValue(Context, GetBoneIndicesKey(), &BoneIndicesKey);
 	}
-	else if (Out->IsA<FCollectionAttributeKey>(&BoneWeightsKey))
+	else if (Out->IsA(&BoneWeightsKey))
 	{
 		SetValue(Context, GetBoneWeightsKey(), &BoneWeightsKey);
 	}
+	else if (Out->IsA(&SkeletalMesh))
+	{
+		SafeForwardInput(Context, &SkeletalMesh, &SkeletalMesh);
+	}
 }
 
 FCollectionAttributeKey FTransferSkinWeightsGroomNode::GetBoneIndicesKey() const
diff --git a/Engine/Plugins/Runtime/HairStrands/Source/HairStrandsDataflow/Public/BuildGroomSkinningNodes.h b/Engine/Plugins/Runtime/HairStrands/Source/HairStrandsDataflow/Public/BuildGroomSkinningNodes.h
--- a/Engine/Plugins/Runtime/HairStrands/Source/HairStrandsDataflow/Public/BuildGroomSkinningNodes.h
+++ b/Engine/Plugins/Runtime/HairStrands/Source/HairStrandsDataflow/Public/BuildGroomSkinningNodes.h
@@ -21,14 +21,7 @@ struct FTransferSkinWeightsGroomNode : public FDataflowNode
 
 public:
 	
-	FTransferSkinWeightsGroomNode(const UE::Dataflow::FNodeParameters& InParam, FGuid InGuid = FGuid::NewGuid())
-	: FDataflowNode(InParam, InGuid)
-	{
-		RegisterInputConnection(&Collection);
-		RegisterOutputConnection(&Collection, &Collection);
-		RegisterOutputConnection(&BoneIndicesKey);
-		RegisterOutputConnection(&BoneWeightsKey);
-	}
+	FTransferSkinWeightsGroomNode(const UE::Dataflow::FNodeParameters& InParam, FGuid InGuid = FGuid::NewGuid());
 
 	//~ Begin FDataflowNode interface
 	virtual void Evaluate(UE::Dataflow::FContext& Context, const FDataflowOutput* Out) const override;
@@ -40,7 +33,7 @@ public:
 	FManagedArrayCollection Collection;
 
 	/** SkeletalMesh used to transfer the skinning weights. Will be stored onto the groom asset */
-	UPROPERTY(EditAnywhere, Category="Skeletal Mesh")
+	UPROPERTY(EditAnywhere, Category="Skeletal Mesh", meta = (DataflowInput, DataflowOutput, DataflowPassthrough = "SkeletalMesh"))
 	TObjectPtr<USkeletalMesh> SkeletalMesh = nullptr;
 
 	/** LOD used to transfer the weights */
