Exposed SCollectionSpreadSheetWidget functions to required to display collection spreadsheet in MegaPlant editor.

diff --git a/Engine/Plugins/Experimental/Dataflow/Source/DataflowEditor/Public/Dataflow/DataflowCollectionSpreadSheetWidget.h b/Engine/Plugins/Experimental/Dataflow/Source/DataflowEditor/Public/Dataflow/DataflowCollectionSpreadSheetWidget.h
--- a/Engine/Plugins/Experimental/Dataflow/Source/DataflowEditor/Public/Dataflow/DataflowCollectionSpreadSheetWidget.h
+++ b/Engine/Plugins/Experimental/Dataflow/Source/DataflowEditor/Public/Dataflow/DataflowCollectionSpreadSheetWidget.h
@@ -131,10 +131,10 @@ public:
 	* @param InArgs - Arguments passed from Slate
 	*
 	*/
-	void Construct(const FArguments& InArgs);
+	void DATAFLOWEDITOR_API  Construct(const FArguments& InArgs);
 
-	void SetData(const FString& InNodeName);
-	void RefreshWidget();
+	void DATAFLOWEDITOR_API SetData(const FString& InNodeName);
+	void DATAFLOWEDITOR_API RefreshWidget();
 	TSharedPtr<SCollectionSpreadSheet> GetCollectionTable() { return CollectionTable; }
 	void SetStatusText();
 
