Dataflow skin weight paint node : fix uninitialized variable

diff --git a/Engine/Plugins/Experimental/Dataflow/Source/DataflowNodes/Public/Dataflow/DataflowCollectionEditSkinWeightsNode.h b/Engine/Plugins/Experimental/Dataflow/Source/DataflowNodes/Public/Dataflow/DataflowCollectionEditSkinWeightsNode.h
--- a/Engine/Plugins/Experimental/Dataflow/Source/DataflowNodes/Public/Dataflow/DataflowCollectionEditSkinWeightsNode.h
+++ b/Engine/Plugins/Experimental/Dataflow/Source/DataflowNodes/Public/Dataflow/DataflowCollectionEditSkinWeightsNode.h
@@ -78,7 +78,7 @@ public:
 	* This allow invalidating the SkinWeights when the ref skeleton does not match anymore instead of having broken weights
 	*/
 	UPROPERTY()
-	int32 SkinWeightsRefSkeletonHash;
+	int32 SkinWeightsRefSkeletonHash = 0;
 
 	/** Delegate to transfer the bone selection to the tool*/
 	FDataflowBoneSelectionChangedNotifyDelegate OnBoneSelectionChanged;
