Dataflow simulation controls unification

#virtualized

diff --git a/Engine/Plugins/ChaosClothAsset/Source/ChaosClothAssetEngine/Private/ChaosClothAsset/ClothAssetBase.cpp b/Engine/Plugins/ChaosClothAsset/Source/ChaosClothAssetEngine/Private/ChaosClothAsset/ClothAssetBase.cpp
--- a/Engine/Plugins/ChaosClothAsset/Source/ChaosClothAssetEngine/Private/ChaosClothAsset/ClothAssetBase.cpp
+++ b/Engine/Plugins/ChaosClothAsset/Source/ChaosClothAssetEngine/Private/ChaosClothAsset/ClothAssetBase.cpp
@@ -112,6 +112,8 @@ TObjectPtr<UDataflowBaseContent> UChaosClothAssetBase::CreateDataflowContent()
 
 	WriteDataflowContent(SkeletalContent);
 
+	SkeletalContent->OnContentDataChanged.AddUObject(this, &UChaosClothAssetBase::UpdateSimulationActor);
+
 	return SkeletalContent;
 }
 
@@ -140,6 +142,21 @@ void UChaosClothAssetBase::ReadDataflowContent(const TObjectPtr<UDataflowBaseCon
 	}
 }
 
+void UChaosClothAssetBase::UpdateSimulationActor(TObjectPtr<AActor>& SimulationActor) const 
+{
+	TInlineComponentArray<UChaosClothComponent*> ChaosClothComponents(SimulationActor);
+	for(UChaosClothComponent* ClothComponent : ChaosClothComponents)
+	{
+		if(ClothComponent->GetAsset() == this)
+		{
+#if WITH_EDITOR
+			// Update the config properties on the component from the asset.
+			ClothComponent->UpdateConfigProperties();
+#endif
+		}
+	}
+}
+
 const FDataflowInstance& UChaosClothAssetBase::GetDataflowInstance() const
 {
 PRAGMA_DISABLE_DEPRECATION_WARNINGS
diff --git a/Engine/Plugins/ChaosClothAsset/Source/ChaosClothAssetEngine/Public/ChaosClothAsset/ClothAssetBase.h b/Engine/Plugins/ChaosClothAsset/Source/ChaosClothAssetEngine/Public/ChaosClothAsset/ClothAssetBase.h
--- a/Engine/Plugins/ChaosClothAsset/Source/ChaosClothAssetEngine/Public/ChaosClothAsset/ClothAssetBase.h
+++ b/Engine/Plugins/ChaosClothAsset/Source/ChaosClothAssetEngine/Public/ChaosClothAsset/ClothAssetBase.h
@@ -384,6 +384,9 @@ private:
 		WaitUntilAsyncPropertyReleasedInternal((uint64)AsyncProperties, LockType);
 	}
 
+	/** Update simulation actor/components from the asset data */
+	void UpdateSimulationActor(TObjectPtr<AActor>& SimulationActor) const;
+
 #if WITH_EDITORONLY_DATA
 	// The following PreviewScene properties are modeled after PreviewSkeletalMesh in USkeleton
 	//	- they are inside WITH_EDITORONLY_DATA because they are not used at game runtime
diff --git a/Engine/Plugins/Experimental/Dataflow/Source/DataflowEditor/Private/Dataflow/DataflowEditorCommands.cpp b/Engine/Plugins/Experimental/Dataflow/Source/DataflowEditor/Private/Dataflow/DataflowEditorCommands.cpp
--- a/Engine/Plugins/Experimental/Dataflow/Source/DataflowEditor/Private/Dataflow/DataflowEditorCommands.cpp
+++ b/Engine/Plugins/Experimental/Dataflow/Source/DataflowEditor/Private/Dataflow/DataflowEditorCommands.cpp
@@ -36,10 +36,6 @@
 #define LOCTEXT_NAMESPACE "DataflowEditorCommands"
 
 const FString FDataflowEditorCommandsImpl::AddWeightMapNodeIdentifier = TEXT("AddWeightMapNode");
-const FString FDataflowEditorCommandsImpl::RebuildSimulationSceneIdentifier = TEXT("RebuildSimulationScene");
-const FString FDataflowEditorCommandsImpl::PauseSimulationSceneIdentifier = TEXT("PauseSimulationScene");
-const FString FDataflowEditorCommandsImpl::StartSimulationSceneIdentifier = TEXT("StartSimulationScene");
-const FString FDataflowEditorCommandsImpl::StepSimulationSceneIdentifier = TEXT("StepSimulationScene");
 
 FDataflowEditorCommandsImpl::FDataflowEditorCommandsImpl()
 	: TBaseCharacterFXEditorCommands<FDataflowEditorCommandsImpl>("DataflowEditor", 
@@ -55,6 +51,7 @@ void FDataflowEditorCommandsImpl::RegisterCommands()
 	
 	UI_COMMAND(EvaluateNode, "Evaluate", "Trigger an evaluation of the selected node.", EUserInterfaceActionType::Button, FInputChord());
 	UI_COMMAND(EvaluateGraph, "Evaluate Dataflow Graph", "Trigger an evaluation of the graph Terminal node.", EUserInterfaceActionType::Button, FInputChord(EKeys::F5));
+	UI_COMMAND(ToggleSimulation, "Simulate Dataflow Scene", "Toggle the simulation of the registered components in the simulation scene.", EUserInterfaceActionType::RadioButton, FInputChord(EKeys::S, EModifierKey::Alt));
 	UI_COMMAND(EvaluateGraphAutomatic, "Automatic Graph Evaluation", "Set the evaluation mode of the graph to Automatic.", EUserInterfaceActionType::RadioButton, FInputChord());
 	UI_COMMAND(EvaluateGraphManual, "Manual Graph Evaluation", "Set the evaluation mode of the graph to Manual.", EUserInterfaceActionType::RadioButton, FInputChord());
 	UI_COMMAND(ClearGraphCache, "Force Graph Re-evaluation", "Force the entire graph to re-evaluate by deleting the evaluation cache and re-evaluating the Terminal node.", EUserInterfaceActionType::Button, FInputChord());
@@ -72,18 +69,12 @@ void FDataflowEditorCommandsImpl::RegisterCommands()
 	UI_COMMAND(ZoomToFitGraph, "ZoomToFitGraph", "Fit the graph in the graph editor viewport.", EUserInterfaceActionType::None, FInputChord(EKeys::F));
 
 	UI_COMMAND(AddWeightMapNode, "Add Weight Map", "Paint weight maps on the mesh", EUserInterfaceActionType::Button, FInputChord());
-	UI_COMMAND(RebuildSimulationScene, "Rebuild Simulation", "Rebuild the simulation scene", EUserInterfaceActionType::Button, FInputChord());
-	UI_COMMAND(PauseSimulationScene, "Pause Simulation", "Pause the simulation scene", EUserInterfaceActionType::Button, FInputChord());
-	UI_COMMAND(StartSimulationScene, "Start Simulation", "Start the simulation scene", EUserInterfaceActionType::Button, FInputChord());
-	UI_COMMAND(StepSimulationScene, "Step Simulation", "Step the simulation scene", EUserInterfaceActionType::Button, FInputChord());
-
 	UI_COMMAND(AddNewVariable, "Variable", "Adds a new variable to this dataflow graph.", EUserInterfaceActionType::Button, FInputChord());
 	UI_COMMAND(AddNewSubGraph, "SubGraph", "Adds a new subgraph to this dataflow graph.", EUserInterfaceActionType::Button, FInputChord());
 	UI_COMMAND(ConvertToBasicSubGraph, "Convert to Basic Subgraph", "Convert the subgraph to a basic one (no loop).", EUserInterfaceActionType::Button, FInputChord());
 	UI_COMMAND(ConvertToForEachSubGraph, "Convert to For Each Subgraph", "Convert the subgraph to be able to iterate through an array.", EUserInterfaceActionType::Button, FInputChord());
 	UI_COMMAND(CollapseToSubGraph, "Collapse to Sub-graph", "Collapse the selection into a subgraph", EUserInterfaceActionType::Button, FInputChord());
 
-
 	for (const TPair<FName, TUniquePtr<UE::Dataflow::IDataflowConstructionViewMode>>& NameAndMode : UE::Dataflow::FRenderingViewModeFactory::GetInstance().GetViewModes())
 	{
 		TSharedPtr< FUICommandInfo > SetViewModeCommand;
@@ -112,10 +103,10 @@ void FDataflowEditorCommandsImpl::RegisterCommands()
 			this->AsShared(),
 			ToolRegistry.GetToolCommandForNode(NodeName),
 			FName(NodeName.ToString() + "_Tool"),
-			LOCTEXT("DataflowTool", "Dataflow Tool"),					// TODO: Replace placeholder names
-			LOCTEXT("DataflowToolTooltip", "Dataflow Tool Tooltip"),	// TODO: Replace placeholder names
-			FSlateIcon(),
-			EUserInterfaceActionType::Button,
+			ToolRegistry.GetAddNodeButtonText(NodeName),					// TODO: Replace placeholder names
+			FText::Format(LOCTEXT("LaunchDataflowToolTooltip", "Launch a \"{0}\" dataflow tool"), FText::FromString(NodeName.ToString())),
+			ToolRegistry.GetAddNodeButtonIcon(NodeName),
+			EUserInterfaceActionType::RadioButton,
 			FInputChord()
 		);
 
diff --git a/Engine/Plugins/Experimental/Dataflow/Source/DataflowEditor/Private/Dataflow/DataflowEditorModeToolkit.cpp b/Engine/Plugins/Experimental/Dataflow/Source/DataflowEditor/Private/Dataflow/DataflowEditorModeToolkit.cpp
--- a/Engine/Plugins/Experimental/Dataflow/Source/DataflowEditor/Private/Dataflow/DataflowEditorModeToolkit.cpp
+++ b/Engine/Plugins/Experimental/Dataflow/Source/DataflowEditor/Private/Dataflow/DataflowEditorModeToolkit.cpp
@@ -18,6 +18,8 @@
 #include "Tools/UEdMode.h"
 #include "Widgets/Images/SImage.h"
 #include "ToolMenus.h"
+#include "Dataflow/DataflowToolRegistry.h"
+#include "Toolkits/AssetEditorModeUILayer.h"
 
 #define LOCTEXT_NAMESPACE "FDataflowEditorModeToolkit"
 
@@ -31,23 +33,87 @@ FText FDataflowEditorModeToolkit::GetBaseToolkitName() const
 	return NSLOCTEXT("DataflowEditorModeToolkit", "DisplayName", "DataflowEditorMode");
 }
 
+void FDataflowEditorModeToolkit::BindCommands()
+{
+	if (GetScriptableEditorMode().IsValid())
+	{
+		UDataflowEditorMode* EditorMode = Cast<UDataflowEditorMode>(GetScriptableEditorMode());
+		const TArray<FName> ToolCategories = EditorMode->GetToolCategories();
+
+		UE::Dataflow::FDataflowToolRegistry& ToolRegistry = UE::Dataflow::FDataflowToolRegistry::Get();
+		TArray<FName> NodeNames = ToolRegistry.GetNodeNames();
+
+		for (const FName& NodeName : NodeNames)
+		{
+			ToolkitCommands->MapAction(ToolRegistry.GetAddNodeCommandForNode(NodeName),
+				FExecuteAction::CreateUObject(EditorMode, &UDataflowEditorMode::AddNode, NodeName),
+				FCanExecuteAction::CreateUObject(EditorMode, &UDataflowEditorMode::CanAddNode, NodeName));
+		}
+	}
+}
+
 void FDataflowEditorModeToolkit::Init(const TSharedPtr<IToolkitHost>& InitToolkitHost, TWeakObjectPtr<UEdMode> InOwningMode)
 {
-	FModeToolkit::Init(InitToolkitHost, InOwningMode);	
+	bUsesToolkitBuilder = true;
+	
+	FModeToolkit::Init(InitToolkitHost, InOwningMode);
+
+	BindCommands();
+
+	ModeToolPallete.ModeWarningArea = SNew(STextBlock)
+			.AutoWrapText(true)
+			.Font(FCoreStyle::GetDefaultFontStyle("Bold", 9))
+			.ColorAndOpacity(FSlateColor(FLinearColor(0.9f, 0.15f, 0.15f)));
+	ModeToolPallete.ModeWarningArea->SetText(FText::GetEmpty());
+	ModeToolPallete.ModeWarningArea->SetVisibility(EVisibility::Collapsed);
+
+	ModeToolPallete.ModeHeaderArea = SNew(STextBlock)
+		.AutoWrapText(true)
+		.Font(FCoreStyle::GetDefaultFontStyle("Bold", 12));
+	ModeToolPallete.ModeHeaderArea->SetText(LOCTEXT("SelectToolLabel", "Select a Tool from the Toolbar"));
+	ModeToolPallete.ModeHeaderArea->SetJustification(ETextJustify::Center);
+	
+	ToolWarningArea = SNew(STextBlock)
+		.AutoWrapText(true)
+		.Font(FCoreStyle::GetDefaultFontStyle("Bold", 9))
+		.ColorAndOpacity(FSlateColor(FLinearColor(0.9f, 0.15f, 0.15f)));
+	ToolWarningArea->SetText(FText::GetEmpty());
 
 	// Set up tool message areas
+	RegisterPalettes();
 	ClearNotification();
 	ClearWarning();
-	GetScriptableEditorMode()->GetInteractiveToolsContext()->OnToolNotificationMessage.AddSP(this, &FBaseCharacterFXEditorModeToolkit::PostNotification);
-	GetScriptableEditorMode()->GetInteractiveToolsContext()->OnToolWarningMessage.AddSP(this, &FBaseCharacterFXEditorModeToolkit::PostWarning);
 
-	// Hook up the tool detail panel
-	ToolDetailsContainer->SetContent(DetailsView.ToSharedRef());
-	
-	// Set up the overlay -- it will poll different ToolsContexts depending on which one launched the currently active tool
+	// create the toolkit widget
+	{
+		ToolkitSections->ModeWarningArea = ModeToolPallete.ModeWarningArea;
+		ToolkitSections->DetailsView = ModeDetailsView;
+		ToolkitSections->ToolWarningArea = ToolWarningArea;
+
+		SAssignNew(ToolkitWidget, SBorder)
+		.HAlign(HAlign_Fill)
+		.Padding(0)
+		.BorderImage(FAppStyle::GetBrush("ToolPanel.GroupBorder"))
+		[
+			ToolkitBuilder->GenerateWidget()->AsShared()
+		];	
+	}
 
-	SAssignNew(ViewportOverlayWidget, SHorizontalBox)
+	ActiveToolName = FText::GetEmpty();
 
+	if(UEditorInteractiveToolsContext* Context = GetScriptableEditorMode()->GetInteractiveToolsContext())
+	{
+		Context->OnToolNotificationMessage.AddSP(this, &FBaseCharacterFXEditorModeToolkit::PostNotification);
+		Context->OnToolWarningMessage.AddSP(this, &FBaseCharacterFXEditorModeToolkit::PostWarning);
+	}
+
+	// add viewport overlay widget to accept / cancel tool
+	MakeToolAcceptCancelWidget();
+}
+
+void FDataflowEditorModeToolkit::MakeToolAcceptCancelWidget()
+{
+	SAssignNew(ViewportOverlayWidget, SHorizontalBox)
 	+SHorizontalBox::Slot()
 	.HAlign(HAlign_Center)
 	.VAlign(VAlign_Bottom)
@@ -133,6 +199,73 @@ void FDataflowEditorModeToolkit::Init(const TSharedPtr<IToolkitHost>& InitToolki
 	];
 }
 
+void FDataflowEditorModeToolkit::RegisterPalettes()
+{
+	FDataflowEditorCommandsImpl& Commands = FDataflowEditorCommandsImpl::Get();
+	const TSharedPtr<FUICommandList> CommandList = GetToolkitCommands();
+	
+	ToolkitSections = MakeShared<FToolkitSections>();
+
+	FToolkitBuilderArgs ToolkitBuilderArgs(GetScriptableEditorMode()->GetModeInfo().ToolbarCustomizationName);
+	ToolkitBuilderArgs.ToolkitCommandList = GetToolkitCommands();
+	ToolkitBuilderArgs.ToolkitSections = ToolkitSections;
+	ToolkitBuilderArgs.SelectedCategoryTitleVisibility = EVisibility::Collapsed;
+	ToolkitBuilder = MakeShared<FToolkitBuilder>(ToolkitBuilderArgs);
+
+	if (GetScriptableEditorMode().IsValid())
+	{
+		UDataflowEditorMode* EditorMode = Cast<UDataflowEditorMode>(GetScriptableEditorMode());
+		const TArray<FName> ToolCategories = EditorMode->GetToolCategories();
+
+		UE::Dataflow::FDataflowToolRegistry& ToolRegistry = UE::Dataflow::FDataflowToolRegistry::Get();
+		TArray<FName> NodeNames = ToolRegistry.GetNodeNames();
+		
+		for(const FName& CategoryName : ToolCategories)
+		{
+			// Get all the tools that belongs to that category
+			TArray<TSharedPtr<FUICommandInfo>> CategoryCommands;
+			for (const FName& NodeName : NodeNames)
+			{
+				if(ToolRegistry.GetToolCategoryForNode(NodeName) == CategoryName)
+				{
+					const TSharedPtr<FUICommandInfo> CommandInfo = ToolRegistry.GetAddNodeCommandForNode(NodeName);
+					CategoryCommands.Add(CommandInfo);
+				}
+			}
+			// Add palette and the matching tools
+			if(const TSharedPtr<FUICommandInfo>* CommandInfo = Commands.ToolPaletteCommands.Find(CategoryName))
+			{
+				ToolkitBuilder->AddPalette(MakeShareable(new FToolPalette((*CommandInfo).ToSharedRef(), CategoryCommands)));
+			}
+			else
+			{
+				TSharedPtr<FUICommandInfo> NewCommandInfo;
+				const FText Description = FText::Format(LOCTEXT("DataflowToolCategoryDescription", "Tool category {0} used to register dataflow tools."), FText::FromName(CategoryName));
+				
+                FUICommandInfo::MakeCommandInfo(
+                	Commands.AsShared(),
+                	NewCommandInfo,
+                	CategoryName, FText::FromName(CategoryName), Description, FSlateIcon(),
+                	EUserInterfaceActionType::RadioButton,
+                	FInputChord() );
+
+				Commands.ToolPaletteCommands.Add(CategoryName, NewCommandInfo);
+
+				ToolkitBuilder->AddPalette(MakeShareable(new FToolPalette(NewCommandInfo.ToSharedRef(), CategoryCommands)));
+			}
+		}
+	}
+	
+	//ToolkitBuilder->SetActivePaletteOnLoad(Commands.LoadGeneralTools.Get());
+	ToolkitBuilder->UpdateWidget();
+
+	// if selected palette changes, make sure we are showing the palette command buttons, which may be hidden by active Tool
+	ModeToolPallete.ActivePaletteChangedHandle = ToolkitBuilder->OnActivePaletteChanged.AddLambda([this]()
+	{
+		ToolkitBuilder->SetActivePaletteCommandsVisibility(EVisibility::Visible);
+	});
+}
+
 const FSlateBrush* FDataflowEditorModeToolkit::GetActiveToolIcon(const FString& ActiveToolIdentifier) const
 {
 	FName ActiveToolIconName = ISlateStyle::Join(FDataflowEditorCommandsImpl::Get().GetContextName(), TCHAR_TO_ANSI(*ActiveToolIdentifier));
@@ -194,38 +327,27 @@ void FDataflowEditorModeToolkit::OnToolStarted(UInteractiveToolManager* Manager,
 {
 	ensure(Tool == Manager->GetActiveTool(EToolSide::Left));
 
-	if (!Tool)
+	if (Tool)
 	{
-		return;
-	}
+		UpdateActiveToolProperties(Tool);
 
-	if (DetailsView)
-	{
-		DetailsView->SetObjects(Tool->GetToolProperties());
+		Tool->OnPropertySetsModified.AddSP(this, &FDataflowEditorModeToolkit::UpdateActiveToolProperties, Tool);
+		Tool->OnPropertyModifiedDirectlyByTool.AddSP(this, &FDataflowEditorModeToolkit::InvalidateCachedDetailPanelState);
 
-		Tool->OnPropertySetsModified.AddLambda([this, Manager]()
-			{
-				const UInteractiveTool* const Tool = Manager->GetActiveTool(EToolSide::Left);
-		if (Tool != nullptr)
-		{
-			DetailsView->SetObjects(Tool->GetToolProperties(true));
-		}
-			});
-	}
-
-	Tool->OnPropertyModifiedDirectlyByTool.AddSP(this, &FDataflowEditorModeToolkit::InvalidateCachedDetailPanelState);
+		ModeToolPallete.ModeHeaderArea->SetVisibility(EVisibility::Collapsed);
+		ActiveToolName = Tool->GetToolInfo().ToolDisplayName;
 
-	ActiveToolName = Tool->GetToolInfo().ToolDisplayName;
-	FString ActiveToolIdentifier = Manager->GetActiveToolName(EToolSide::Mouse);
-	ActiveToolIdentifier.InsertAt(0, ".");
-	ActiveToolIcon = GetActiveToolIcon(ActiveToolIdentifier);
-
-	if (SBaseCharacterFXEditorViewport* Widget = GetViewportWidgetForManager(Manager))
-	{
-		Widget->AddOverlayWidget(ViewportOverlayWidget.ToSharedRef());
+		FString ActiveToolIdentifier = Manager->GetActiveToolName(EToolSide::Mouse); 
+		ActiveToolIdentifier.InsertAt(0, ".");
+		ActiveToolIcon = GetActiveToolIcon(ActiveToolIdentifier);
+		
+		if (SBaseCharacterFXEditorViewport* Widget = GetViewportWidgetForManager(Manager))
+		{
+			Widget->AddOverlayWidget(ViewportOverlayWidget.ToSharedRef());
 
-		// Manually set keyboard focus to this viewport. Otherwise the user has to click on the viewport before Tool keyboard shortcuts will work.
-		FSlateApplication::Get().SetUserFocus(FSlateApplication::Get().GetUserIndexForKeyboard(), Widget->AsShared(), EFocusCause::SetDirectly);
+			// Manually set keyboard focus to this viewport. Otherwise the user has to click on the viewport before Tool keyboard shortcuts will work.
+			FSlateApplication::Get().SetUserFocus(FSlateApplication::Get().GetUserIndexForKeyboard(), Widget->AsShared(), EFocusCause::SetDirectly);
+		}
 	}
 }
 
@@ -245,8 +367,8 @@ void FDataflowEditorModeToolkit::OnToolEnded(UInteractiveToolManager* Manager, U
 			return nullptr;
 		};
 
-	FModeToolkit::OnToolEnded(Manager, Tool);
-
+	ModeToolPallete.ModeHeaderArea->SetVisibility(EVisibility::Visible);
+	ModeDetailsView->SetObject(nullptr);
 	ActiveToolName = FText::GetEmpty();
 	ClearNotification();
 	ClearWarning();
@@ -281,6 +403,20 @@ void FDataflowEditorModeToolkit::SetSimulationViewportWidget(TWeakPtr<SDataflowS
 	SimulationViewportWidget = InSimulationViewportWidget;
 }
 
+void FDataflowEditorModeToolkit::UpdateActiveToolProperties(UInteractiveTool* Tool)
+{
+	UInteractiveTool* CurTool = GetScriptableEditorMode()->GetToolManager(EToolsContextScope::EdMode)->GetActiveTool(EToolSide::Left);
+	if (CurTool == nullptr)
+	{
+		return;
+	}
+		
+	ModeDetailsView->SetObjects(CurTool->GetToolProperties(true));
+}
 
+void FDataflowEditorModeToolkit::InvalidateCachedDetailPanelState(UObject* ChangedObject)
+{
+	ModeDetailsView->InvalidateCachedState();
+}
 
 #undef LOCTEXT_NAMESPACE
diff --git a/Engine/Plugins/Experimental/Dataflow/Source/DataflowEditor/Private/Dataflow/DataflowEditorStyle.h b/Engine/Plugins/Experimental/Dataflow/Source/DataflowEditor/Private/Dataflow/DataflowEditorStyle.h
--- a/Engine/Plugins/Experimental/Dataflow/Source/DataflowEditor/Private/Dataflow/DataflowEditorStyle.h
+++ b/Engine/Plugins/Experimental/Dataflow/Source/DataflowEditor/Private/Dataflow/DataflowEditorStyle.h
@@ -24,9 +24,9 @@ public:
 
 		const FVector2D ToolbarIconSize = Icon20x20;
 
-		SetCoreContentRoot(FPaths::EngineContentDir() / TEXT("Editor/Slate"));
-
+		SetCoreContentRoot(FPaths::EngineContentDir() / TEXT("Slate"));
 		SetContentRoot(IPluginManager::Get().FindPlugin("Dataflow")->GetBaseDir() / TEXT("Resources"));
+
 		Set("ClassIcon.Dataflow", new IMAGE_BRUSH_SVG("DataflowAsset_16", Icon16x16));
 		Set("ClassThumbnail.Dataflow", new IMAGE_BRUSH_SVG("DataflowAsset_64", Icon64x64));
 
@@ -46,12 +46,12 @@ public:
 		// Dataflow weight map and skin weights icons 
 		Set("Dataflow.PaintWeightMap", new IMAGE_BRUSH_SVG("Slate/Dataflow_WeightMap", Icon20x20));
 		Set("Dataflow.EditSkinWeights", new IMAGE_BRUSH_SVG("Slate/Dataflow_SkinWeight", Icon20x20));
+		Set("Dataflow.EditSkeletonBones", new IMAGE_BRUSH_SVG("Slate/Dataflow_EditSkeleton", Icon20x20));
 
-		// @todo(brice) Remove Example Tools
-		//const FString AttributeEditorPropertyName = "DataflowEditor." + FDataflowEditorCommandsImpl::BeginAttributeEditorToolIdentifier;
-		//Set(*AttributeEditorPropertyName, new FSlateImageBrush(RootToContentDir(TEXT("Slate/Dataflow_SelectObject_20x.png")), Icon20x20));
-		//const FString MeshSelectionPropertyName = "DataflowEditor." + FDataflowEditorCommandsImpl::BeginMeshSelectionToolIdentifier;
-		//Set(*MeshSelectionPropertyName, new FSlateImageBrush(RootToContentDir(TEXT("Slate/Dataflow_SelectVertex_20x.png")), Icon20x20));
+		// Dataflow simulation icons
+		Set("Dataflow.ResetSimulation", new CORE_IMAGE_BRUSH_SVG("Starship/Common/Undo", Icon20x20));
+		Set("Dataflow.LockSimulation", new CORE_IMAGE_BRUSH_SVG("Starship/Common/lock", Icon20x20));
+		Set("Dataflow.UnlockSimulation", new CORE_IMAGE_BRUSH_SVG("Starship/Common/lock-unlocked", Icon20x20));
 
 		DefaultMaterial = Cast<UMaterial>(StaticLoadObject(UMaterial::StaticClass(), NULL, TEXT("/Engine/BasicShapes/BasicShapeMaterial")));
 		VertexMaterial = Cast<UMaterial>(StaticLoadObject(UMaterial::StaticClass(), NULL, TEXT("/Engine/EditorMaterials/Dataflow/DataflowVertexMaterial")));
diff --git a/Engine/Plugins/Experimental/Dataflow/Source/DataflowEditor/Private/Dataflow/DataflowEditorToolkit.cpp b/Engine/Plugins/Experimental/Dataflow/Source/DataflowEditor/Private/Dataflow/DataflowEditorToolkit.cpp
--- a/Engine/Plugins/Experimental/Dataflow/Source/DataflowEditor/Private/Dataflow/DataflowEditorToolkit.cpp
+++ b/Engine/Plugins/Experimental/Dataflow/Source/DataflowEditor/Private/Dataflow/DataflowEditorToolkit.cpp
@@ -42,6 +42,10 @@
 #include "Dataflow/DataflowSubGraphNodes.h"
 #include "Dataflow/DataflowMembersWidget.h"
 #include "Dataflow/DataflowSelectionView.h"
+#include "Dataflow/DataflowOutputLog.h"
+#include "Dataflow/DataflowPath.h"
+#include "Dataflow/DataflowSimulationManager.h"
+#include "Dataflow/DataflowSimulationNodes.h"
 #include "EditorModeManager.h"
 #include "EditorViewportTabContent.h"
 #include "Tools/EdModeInteractiveToolsContext.h"
@@ -52,17 +56,16 @@
 #include "Modules/ModuleManager.h"
 #include "PropertyEditorModule.h"
 #include "Selection.h"
-#include "Widgets/Docking/SDockTab.h"
-#include "Widgets/Layout/SSpacer.h"
 #include "GeometryCache.h"
 #include "ToolMenus.h"
 #include "WorkflowOrientedApp/WorkflowTabManager.h"
 #include "Framework/Notifications/NotificationManager.h"
-#include "Widgets/Notifications/SNotificationList.h"
-#include "Dataflow/DataflowOutputLog.h"
-#include "Dataflow/DataflowPath.h"
-#include "Dataflow/DataflowSimulationNodes.h"
 #include "Toolkits/AssetEditorToolkitMenuContext.h"
+#include "Widgets/Notifications/SNotificationList.h"
+#include "Widgets/Docking/SDockTab.h"
+#include "Widgets/Layout/SSpacer.h"
+#include "Widgets/SDataflowSimulationTimeline.h"
+#include "Widgets/DataflowSimulationBinding.h"
 
 #define LOCTEXT_NAMESPACE "DataflowEditorToolkit"
 
@@ -93,6 +96,7 @@ const FName FDataflowEditorToolkit::CollectionSpreadSheetTabId_3(TEXT("DataflowE
 const FName FDataflowEditorToolkit::CollectionSpreadSheetTabId_4(TEXT("DataflowEditor_CollectionSpreadSheet_4"));
 const FName FDataflowEditorToolkit::SimulationViewportTabId(TEXT("DataflowEditor_SimulationViewport"));
 const FName FDataflowEditorToolkit::SimulationVisualizationTabId(TEXT("DataflowEditor_SimulationVisualizationTab"));
+const FName FDataflowEditorToolkit::SimulationTimelineTabId(TEXT("DataflowEditor_SimulationTimelineTab"));
 const FName FDataflowEditorToolkit::MembersWidgetTabId(TEXT("DataflowEditor_MembersWidgetTab"));
 const FName FDataflowEditorToolkit::OutputLogTabId(TEXT("DataflowEditor_OutputLog"));
 
@@ -237,182 +241,190 @@ FDataflowEditorToolkit::FDataflowEditorToolkit(UAssetEditor* InOwningAssetEditor
 
 	check(DataflowEditor);
 
-	ConstructionDefaultLayout = FTabManager::NewLayout(FName("DataflowConstructionLayout03"))
+	ConstructionDefaultLayout = FTabManager::NewLayout(FName("DataflowConstructionLayout12"))
 		->AddArea
 		(
-			FTabManager::NewPrimaryArea()->SetOrientation(Orient_Horizontal)
+			FTabManager::NewPrimaryArea()->SetOrientation(Orient_Vertical)
 			->Split
 			(
-				FTabManager::NewSplitter()->SetOrientation(Orient_Vertical)
-				->SetSizeCoefficient(0.8f)	// Relative width of (Tools Panel, Construction Viewport, Preview Viewport, Dataflow Graph Editor, Outliner) vs (Asset Details, Preview Scene Details, Dataflow Node Details)
+				FTabManager::NewSplitter()->SetOrientation(Orient_Horizontal)
+				->SetSizeCoefficient(0.50f)	// Relative height of (Tools Panel, Construction Viewport, Preview Viewport) vs (Dataflow Graph Editor, Outliner)
 				->Split
 				(
-					FTabManager::NewSplitter()->SetOrientation(Orient_Horizontal)
-					->SetSizeCoefficient(0.50f)	// Relative height of (Tools Panel, Construction Viewport, Preview Viewport) vs (Dataflow Graph Editor, Outliner)
-					->Split
-					(
-						FTabManager::NewStack()
-						->SetSizeCoefficient(0.1f)		// Relative width of (Tools Panel) vs (Construction Viewport, Preview Viewport)
+				FTabManager::NewStack()
+						->SetSizeCoefficient(0.15f)		// Relative width of (Tools Panel) vs (Construction Viewport, Preview Viewport)
 						->SetExtensionId(UDataflowEditorUISubsystem::EditorSidePanelAreaName)
-						->SetHideTabWell(true)
-					)
-					->Split
+						->SetHideTabWell(false)
+				)
+				->Split
+				(
+				FTabManager::NewSplitter()->SetOrientation(Orient_Vertical)
+						->SetSizeCoefficient(0.7f)
+						->Split
 					(
-						FTabManager::NewStack()
-						->SetSizeCoefficient(0.45f)		// Relative width of (Construction Viewport) vs (Tools Panel, Preview Viewport)
-						->AddTab(ViewportTabID, ETabState::OpenedTab)
-						->SetExtensionId("ViewportArea")
-						->SetHideTabWell(true)
+					FTabManager::NewStack()
+							->SetSizeCoefficient(0.9f)		// Relative width of (Construction Viewport) vs (Tools Panel, Preview Viewport)
+							->AddTab(ViewportTabID, ETabState::OpenedTab)
+							->AddTab(SimulationViewportTabId, ETabState::OpenedTab)
+							->SetExtensionId("ViewportArea")
+							->SetHideTabWell(true)
+							->SetForegroundTab(ViewportTabID)
 					)
 					->Split
 					(
-						FTabManager::NewStack()
-						->SetSizeCoefficient(0.45f)		// Relative width of (Construction Viewport) vs (Tools Panel, Preview Viewport)
-						->AddTab(SimulationViewportTabId, ETabState::OpenedTab)
-						->SetExtensionId("ViewportArea")
-						->SetHideTabWell(true)
+					FTabManager::NewStack()
+							->SetSizeCoefficient(0.1f)		// Relative width of (simulation timeline) vs (Preview Viewport, Construction Viewport)
+							->AddTab(SimulationTimelineTabId, ETabState::OpenedTab)
+							->SetExtensionId("TimelineArea")
+							->SetHideTabWell(true)
 					)
 				)
 				->Split
 				(
-					FTabManager::NewSplitter()->SetOrientation(Orient_Horizontal)
-					->SetSizeCoefficient(0.50f)	// Relative height of (Dataflow Node Details) vs (Asset Details, Preview Scene Details)
-					->Split
+				FTabManager::NewSplitter()->SetOrientation(Orient_Vertical)
+						->SetSizeCoefficient(0.15f)
+						->Split
 					(
 						FTabManager::NewStack()
-						->SetSizeCoefficient(0.2f)
-						->AddTab(MembersWidgetTabId, ETabState::OpenedTab)
-						->SetExtensionId("MembersWidgetArea")
-						->SetHideTabWell(false)
-						->SetForegroundTab(MembersWidgetTabId)
+							->SetSizeCoefficient(0.7f)		// Relative width of (Tools Panel) vs (Construction Viewport, Preview Viewport)
+							->AddTab(OutlinerViewTabId, ETabState::OpenedTab)
+							->SetExtensionId("OutlinerArea")
+							->SetHideTabWell(false)
 					)
 					->Split
 					(
-						FTabManager::NewSplitter()->SetOrientation(Orient_Vertical)
-						->SetSizeCoefficient(0.8f)	// Relative width of (Tools Panel, Construction Viewport, Preview Viewport, Dataflow Graph Editor, Outliner) vs (Asset Details, Preview Scene Details, Dataflow Node Details)
-						->Split
-						(
-							FTabManager::NewStack()
-							->SetSizeCoefficient(0.8f)	// Relative height of (Dataflow Graph Editor) vs (Dataflow Output log)
-							->AddTab(GraphCanvasTabId, ETabState::OpenedTab)
-							->AddTab(SubGraphCanvasTabId, ETabState::ClosedTab)
-							->SetExtensionId("GraphEditorArea")
-							->SetHideTabWell(false)
-							->SetForegroundTab(GraphCanvasTabId)
-						)
-						->Split
-						(
-							FTabManager::NewStack()
-							->SetSizeCoefficient(0.2f)	// Relative height of (Dataflow Graph Editor) vs (Dataflow Output log)
-							->AddTab(OutputLogTabId, ETabState::OpenedTab)
-							->AddTab(CollectionSpreadSheetTabId_1, ETabState::OpenedTab)
-							->SetExtensionId("OutputLogArea")
+						FTabManager::NewStack()
+							->SetSizeCoefficient(0.3f)
+							->AddTab(DetailsTabID, ETabState::OpenedTab)
+							->AddTab(PreviewSceneTabId, ETabState::OpenedTab)
+							->SetExtensionId("DetailsArea")
 							->SetHideTabWell(false)
-							->SetForegroundTab(OutputLogTabId)
-						)
+							->SetForegroundTab(DetailsTabID)
 					)
 				)
 			)
 			->Split
 			(
-				FTabManager::NewSplitter()->SetOrientation(Orient_Vertical)
-				->SetSizeCoefficient(0.2f)	// Relative width of (Asset Details, Preview Scene Details, Dataflow Node Details) vs (Tools Panel, Construction Viewport, Preview Viewport, Dataflow Graph Editor, Outliner)
+				FTabManager::NewSplitter()->SetOrientation(Orient_Horizontal)
+				->SetSizeCoefficient(0.50f)	// Relative height of (Dataflow Node Details) vs (Asset Details, Preview Scene Details)
 				->Split
 				(
 					FTabManager::NewStack()
-					->SetSizeCoefficient(0.65f)	// Relative height of (Asset Details, Preview Scene Details) vs (Dataflow Node Details)
-					->AddTab(DetailsTabID, ETabState::OpenedTab)
-					->AddTab(PreviewSceneTabId, ETabState::OpenedTab)
-					->AddTab(SimulationVisualizationTabId, ETabState::OpenedTab)
-					->SetExtensionId("DetailsArea")
-					->SetHideTabWell(true)
-					->SetForegroundTab(DetailsTabID)
+					->SetSizeCoefficient(0.15f)
+					->AddTab(MembersWidgetTabId, ETabState::OpenedTab)
+					->SetExtensionId("MembersArea")
+					->SetHideTabWell(false)
+					->SetForegroundTab(MembersWidgetTabId)
 				)
 				->Split
 				(
 					FTabManager::NewStack()
-					->SetSizeCoefficient(0.2f)	// Relative height of (Dataflow Node Details) vs (Asset Details, Preview Scene Details)
+						->SetSizeCoefficient(0.7f)	// Relative height of (Tools Panel, Construction Viewport, Preview Viewport, Dataflow Graph Editor, Outliner) vs (Asset Details, Preview Scene Details, Dataflow Node Details)
+						->AddTab(GraphCanvasTabId, ETabState::OpenedTab)
+						->AddTab(SubGraphCanvasTabId, ETabState::ClosedTab)
+						->AddTab(CollectionSpreadSheetTabId_1, ETabState::OpenedTab)
+						->AddTab(OutputLogTabId, ETabState::OpenedTab)
+						->SetExtensionId("GraphArea")
+						->SetHideTabWell(false)
+						->SetForegroundTab(GraphCanvasTabId)
+				)
+				->Split
+				(
+					FTabManager::NewStack()
+					->SetSizeCoefficient(0.15f)	// Relative height of (Dataflow Node Details) vs (Asset Details, Preview Scene Details)
 					->AddTab(NodeDetailsTabId, ETabState::OpenedTab)
-					->SetExtensionId("NodeDetailsArea")
+					->SetExtensionId("NodeArea")
 					->SetHideTabWell(false)
 				)
 			)
 		);
 
-	SimulationDefaultLayout = FTabManager::NewLayout(FName("DataflowSimulationLayout02"))
+	SimulationDefaultLayout = FTabManager::NewLayout(FName("DataflowSimulationLayout4"))
 		->AddArea
 		(
-			FTabManager::NewPrimaryArea()->SetOrientation(Orient_Horizontal)
+			FTabManager::NewPrimaryArea()->SetOrientation(Orient_Vertical)
 			->Split
 			(
-				FTabManager::NewSplitter()->SetOrientation(Orient_Vertical)
-				->SetSizeCoefficient(0.8f)	// Relative width of (Tools Panel, Construction Viewport, Preview Viewport, Dataflow Graph Editor, Outliner) vs (Asset Details, Preview Scene Details, Dataflow Node Details)
+				FTabManager::NewSplitter()->SetOrientation(Orient_Horizontal)
+				->SetSizeCoefficient(0.50f)	// Relative height of (Tools Panel, Construction Viewport, Preview Viewport) vs (Dataflow Graph Editor, Outliner)
 				->Split
 				(
-					FTabManager::NewSplitter()->SetOrientation(Orient_Horizontal)
-					->SetSizeCoefficient(0.60f)	// Relative height of (Tools Panel, Construction Viewport, Preview Viewport) vs (Dataflow Graph Editor, Outliner)
-					->Split
+				FTabManager::NewSplitter()->SetOrientation(Orient_Vertical)
+						->SetSizeCoefficient(0.85f)
+						->Split
 					(
-						FTabManager::NewStack()
-						->SetSizeCoefficient(0.1f)		// Relative width of (Tools Panel) vs (Construction Viewport, Preview Viewport)
-						->SetExtensionId(UDataflowEditorUISubsystem::EditorSidePanelAreaName)
-						->SetHideTabWell(true)
+					FTabManager::NewStack()
+							->SetSizeCoefficient(0.9f)		// Relative width of (Construction Viewport) vs (Tools Panel, Preview Viewport)
+							->AddTab(SimulationViewportTabId, ETabState::OpenedTab)
+							->SetExtensionId("ViewportArea")
+							->SetHideTabWell(true)
+							->SetForegroundTab(ViewportTabID)
 					)
 					->Split
 					(
-						FTabManager::NewStack()
-						->SetSizeCoefficient(0.9f)		// Relative width of (Construction Viewport) vs (Tools Panel, Preview Viewport)
-						->AddTab(ViewportTabID, ETabState::ClosedTab)
-						->AddTab(SimulationViewportTabId, ETabState::OpenedTab)
-						->SetExtensionId("ViewportArea")
-						->SetHideTabWell(false)
+					FTabManager::NewStack()
+							->SetSizeCoefficient(0.1f)		// Relative width of (simulation timeline) vs (Preview Viewport, Construction Viewport)
+							->AddTab(SimulationTimelineTabId, ETabState::OpenedTab)
+							->SetExtensionId("TimelineArea")
+							->SetHideTabWell(true)
 					)
 				)
 				->Split
 				(
-					FTabManager::NewSplitter()->SetOrientation(Orient_Horizontal)
-					->SetSizeCoefficient(0.40f)	// Relative height of (Dataflow Node Details) vs (Asset Details, Preview Scene Details)
-					->Split
+				FTabManager::NewSplitter()->SetOrientation(Orient_Vertical)
+						->SetSizeCoefficient(0.15f)
+						->Split
 					(
 						FTabManager::NewStack()
-						->SetSizeCoefficient(0.2f)
-						->AddTab(CollectionSpreadSheetTabId_1, ETabState::ClosedTab)
-						->SetExtensionId("CollectionSpreadSheetArea")
-						->SetHideTabWell(false)
+							->SetSizeCoefficient(0.7f)		// Relative width of (Tools Panel) vs (Construction Viewport, Preview Viewport)
+							->AddTab(OutlinerViewTabId, ETabState::OpenedTab)
+							->SetExtensionId("OutlinerArea")
+							->SetHideTabWell(false)
 					)
 					->Split
 					(
 						FTabManager::NewStack()
-						->SetSizeCoefficient(0.8f)	// Relative height of (Dataflow Graph Editor, Outliner) vs (Tools Panel, Construction Viewport, Preview Viewport)
-						->AddTab(GraphCanvasTabId, ETabState::OpenedTab)
-						->AddTab(SubGraphCanvasTabId, ETabState::OpenedTab)
-						->SetExtensionId("GraphEditorArea")
-						->SetHideTabWell(false)
-						->SetForegroundTab(GraphCanvasTabId)
+							->SetSizeCoefficient(0.3f)
+							->AddTab(DetailsTabID, ETabState::OpenedTab)
+							->AddTab(PreviewSceneTabId, ETabState::OpenedTab)
+							->SetExtensionId("DetailsArea")
+							->SetHideTabWell(false)
+							->SetForegroundTab(DetailsTabID)
 					)
 				)
 			)
 			->Split
 			(
-				FTabManager::NewSplitter()->SetOrientation(Orient_Vertical)
-				->SetSizeCoefficient(0.2f)	// Relative width of (Asset Details, Preview Scene Details, Dataflow Node Details) vs (Tools Panel, Construction Viewport, Preview Viewport, Dataflow Graph Editor, Outliner)
+				FTabManager::NewSplitter()->SetOrientation(Orient_Horizontal)
+				->SetSizeCoefficient(0.50f)	// Relative height of (Dataflow Node Details) vs (Asset Details, Preview Scene Details)
 				->Split
 				(
 					FTabManager::NewStack()
-					->SetSizeCoefficient(0.65f)	// Relative height of (Asset Details, Preview Scene Details) vs (Dataflow Node Details)
-					->AddTab(DetailsTabID, ETabState::OpenedTab)
-					->AddTab(PreviewSceneTabId, ETabState::OpenedTab)
-					->SetExtensionId("DetailsArea")
-					->SetHideTabWell(true)
-					->SetForegroundTab(DetailsTabID)
+					->SetSizeCoefficient(0.15f)
+					->AddTab(MembersWidgetTabId, ETabState::OpenedTab)
+					->SetExtensionId("MembersArea")
+					->SetHideTabWell(false)
+					->SetForegroundTab(MembersWidgetTabId)
+				)
+				->Split
+				(
+					FTabManager::NewStack()
+						->SetSizeCoefficient(0.7f)	// Relative height of (Tools Panel, Construction Viewport, Preview Viewport, Dataflow Graph Editor, Outliner) vs (Asset Details, Preview Scene Details, Dataflow Node Details)
+						->AddTab(GraphCanvasTabId, ETabState::OpenedTab)
+						->AddTab(SubGraphCanvasTabId, ETabState::ClosedTab)
+						->AddTab(CollectionSpreadSheetTabId_1, ETabState::OpenedTab)
+						->AddTab(OutputLogTabId, ETabState::OpenedTab)
+						->SetExtensionId("GraphArea")
+						->SetHideTabWell(false)
+						->SetForegroundTab(GraphCanvasTabId)
 				)
 				->Split
 				(
 					FTabManager::NewStack()
-					->SetSizeCoefficient(0.2f)	// Relative height of (Dataflow Node Details) vs (Asset Details, Preview Scene Details)
+					->SetSizeCoefficient(0.15f)	// Relative height of (Dataflow Node Details) vs (Asset Details, Preview Scene Details)
 					->AddTab(NodeDetailsTabId, ETabState::OpenedTab)
-					->SetExtensionId("NodeDetailsArea")
-					->SetHideTabWell(true)
+					->SetExtensionId("NodeArea")
+					->SetHideTabWell(false)
 				)
 			)
 		);
@@ -746,16 +758,7 @@ void FDataflowEditorToolkit::PostInitAssetEditor()
 	FToolMenuSection& Section = AssetToolbar->FindOrAddSection("DataflowTools");
 
 	AddEvaluationWidget(Section);
-
-	for (const TPair<FName, TSharedPtr<const FUICommandInfo>>& NodeAndAddCommand : DataflowMode->NodeTypeToAddNodeCommandMap)
-	{
-		ToolkitCommands->MapAction(NodeAndAddCommand.Value,
-			FExecuteAction::CreateUObject(DataflowMode, &UDataflowEditorMode::AddNode, NodeAndAddCommand.Key),
-			FCanExecuteAction::CreateUObject(DataflowMode, &UDataflowEditorMode::CanAddNode, NodeAndAddCommand.Key));
-
-		Section.AddEntry(FToolMenuEntry::InitToolBarButton(NodeAndAddCommand.Value));
-	}
-
+	
 	// Force scenes to update loaded asset viewer settings
 	UAssetViewerSettings::Get()->OnAssetViewerSettingsChanged().Broadcast(NAME_None);
 }
@@ -922,6 +925,13 @@ void FDataflowEditorToolkit::AddEvaluationWidget(FToolMenuSection& InSection)
 		FCanExecuteAction(),
 		FIsActionChecked::CreateSP(this, &FDataflowEditorToolkit::IsAsyncEvaluationEnabled)
 	);
+	ToolkitCommands->MapAction(FDataflowEditorCommands::Get().ToggleSimulation,
+		FExecuteAction::CreateSP(this, &FDataflowEditorToolkit::ToggleDataflowSimulation),
+		FCanExecuteAction::CreateSP(this, &FDataflowEditorToolkit::HasSimulationManager),
+		FIsActionChecked::CreateSP(this, &FDataflowEditorToolkit::IsSimulationEnabled)
+	);
+	
+	
 	
 	InSection.AddDynamicEntry("DataflowEvaluateGraphDynamic", FNewToolMenuSectionDelegate::CreateLambda([](FToolMenuSection& DynamicSection)
 		{
@@ -956,6 +966,22 @@ void FDataflowEditorToolkit::AddEvaluationWidget(FToolMenuSection& InSection)
 			}
 		}));
 
+	InSection.AddDynamicEntry("DataflowToggleSimulation", FNewToolMenuSectionDelegate::CreateLambda([](FToolMenuSection& DynamicSection)
+		{
+			static const FSlateIcon EnableSimulationIconIcon(FAppStyle::GetAppStyleSetName(), "PhysicsAssetEditor.SimulationAll");
+			
+			using namespace UE::Dataflow::Private;
+			if (const TSharedPtr<FDataflowEditorToolkit> Toolkit = GetAssetEditorToolkitFromMenuContext(DynamicSection.FindContext<UAssetEditorToolkitMenuContext>()))
+			{
+				DynamicSection.AddEntry(FToolMenuEntry::InitToolBarButton(
+						FDataflowEditorCommands::Get().ToggleSimulation,
+						TAttribute<FText>(),
+						TAttribute<FText>(),
+						TAttribute<FSlateIcon>(EnableSimulationIconIcon)
+					));
+			}
+		}));
+
 	// load options
 	if (UDataflowEditorOptions* const Options = UDataflowEditorOptions::StaticClass()->GetDefaultObject<UDataflowEditorOptions>())
 	{
@@ -1512,6 +1538,32 @@ bool FDataflowEditorToolkit::IsGraphDirty() const
 	return false;
 }
 
+void FDataflowEditorToolkit::ToggleDataflowSimulation()
+{
+	if(HasSimulationManager())
+	{
+		UDataflowSimulationManager* SimulationManager = SimulationScene->GetWorld()->GetSubsystem<UDataflowSimulationManager>();
+		
+		const bool SimulationFlag = SimulationManager->GetSimulationEnabled();
+		SimulationManager->SetSimulationEnabled(!SimulationFlag);
+	}
+}
+
+bool FDataflowEditorToolkit::IsSimulationEnabled() const
+{
+	if(HasSimulationManager())
+	{
+		UDataflowSimulationManager* SimulationManager = SimulationScene->GetWorld()->GetSubsystem<UDataflowSimulationManager>();
+		return SimulationManager->GetSimulationEnabled();
+	}
+	return false;
+}
+
+bool FDataflowEditorToolkit::HasSimulationManager() const
+{
+	return SimulationScene && SimulationScene->GetWorld() && SimulationScene->GetWorld()->GetSubsystem<UDataflowSimulationManager>();
+}
+
 void FDataflowEditorToolkit::EvaluateGraph()
 {
 	if (const TObjectPtr<UDataflowBaseContent>& EditorContent = GetEditorContent())
@@ -2345,6 +2397,18 @@ TSharedPtr<SWidget> FDataflowEditorToolkit::CreateSimulationVisualizationWidget(
 	return MenuBuilder.MakeWidget();
 }
 
+TSharedRef<SDockTab> FDataflowEditorToolkit::SpawnTab_SimulationTimeline(const FSpawnTabArgs& Args)
+{
+	TSharedRef<SDockTab> SimulationTimelineTab = SNew(SDockTab)
+		.Label(LOCTEXT("DataflowEditor_SimulationTimeline_TabTitle", "Simulation Timeline"));
+
+	SimulationBinding = MakeShared<FDataflowSimulationBinding>(SimulationScene);
+	SimulationTimelineWidget = SNew(SDataflowSimulationTimeline, SimulationBinding.ToSharedRef());
+	SimulationTimelineTab->SetContent(SimulationTimelineWidget.ToSharedRef());
+	
+	return SimulationTimelineTab;
+}
+
 TSharedRef<SDockTab> FDataflowEditorToolkit::SpawnTab_SimulationVisualization(const FSpawnTabArgs& Args)
 {
 	TSharedRef<SDockTab> SimulationVisualizationTab = SNew(SDockTab)
@@ -2498,9 +2562,9 @@ void FDataflowEditorToolkit::RegisterTabSpawners(const TSharedRef<FTabManager>&
 		.SetDisplayName(LOCTEXT("CollectionSpreadSheetTab4", "Collection SpreadSheet 4"))
 		.SetGroup(CollectionSpreadSheetWorkspaceMenuCategoryRef)
 		.SetIcon(FSlateIcon(FAppStyle::GetAppStyleSetName(), "LevelEditor.Tabs.Outliner"));
-
-	InTabManager->RegisterTabSpawner(SimulationVisualizationTabId, FOnSpawnTab::CreateSP(this, &FDataflowEditorToolkit::SpawnTab_SimulationVisualization))
-		.SetDisplayName(LOCTEXT("SimulationVisualizationTab", "Simulation Visualization"))
+		
+	InTabManager->RegisterTabSpawner(SimulationTimelineTabId, FOnSpawnTab::CreateSP(this, &FDataflowEditorToolkit::SpawnTab_SimulationTimeline))
+		.SetDisplayName(LOCTEXT("SimulationTimelineTab", "Simulation Timeline"))
 		.SetGroup(AssetEditorTabsCategory.ToSharedRef())
 		.SetIcon(FSlateIcon(FAppStyle::GetAppStyleSetName(), "LevelEditor.Tabs.Outliner"));
 	
@@ -2548,7 +2612,7 @@ void FDataflowEditorToolkit::UnregisterTabSpawners(const TSharedRef<class FTabMa
 	InTabManager->UnregisterTabSpawner(CollectionSpreadSheetTabId_2);
 	InTabManager->UnregisterTabSpawner(CollectionSpreadSheetTabId_3);
 	InTabManager->UnregisterTabSpawner(CollectionSpreadSheetTabId_4);
-	InTabManager->UnregisterTabSpawner(SimulationVisualizationTabId);
+	InTabManager->UnregisterTabSpawner(SimulationTimelineTabId);
 	InTabManager->UnregisterTabSpawner(MembersWidgetTabId);
 	InTabManager->UnregisterTabSpawner(OutputLogTabId);
 }
diff --git a/Engine/Plugins/Experimental/Dataflow/Source/DataflowEditor/Private/Dataflow/DataflowSimulationScene.cpp b/Engine/Plugins/Experimental/Dataflow/Source/DataflowEditor/Private/Dataflow/DataflowSimulationScene.cpp
--- a/Engine/Plugins/Experimental/Dataflow/Source/DataflowEditor/Private/Dataflow/DataflowSimulationScene.cpp
+++ b/Engine/Plugins/Experimental/Dataflow/Source/DataflowEditor/Private/Dataflow/DataflowSimulationScene.cpp
@@ -57,7 +57,8 @@ FDataflowSimulationScene::FDataflowSimulationScene(FPreviewScene::ConstructionVa
 		if(const UDataflow* DataflowAsset = GetEditorContent()->GetDataflowAsset())
 		{
 			SceneDescription->CacheParams = DataflowAsset->PreviewCacheParams;
-			SceneDescription->CacheAsset = Cast<UChaosCacheCollection>(DataflowAsset->PreviewCacheAsset.LoadSynchronous());
+			SceneDescription->CacheAsset = DataflowAsset->PreviewCacheAsset ? Cast<UChaosCacheCollection>(DataflowAsset->PreviewCacheAsset.LoadSynchronous()) :
+					LoadObject<UChaosCacheCollection>(nullptr, TEXT("/Dataflow/CC_DefaultChaosCache.CC_DefaultChaosCache"));
 			SceneDescription->BlueprintClass = DataflowAsset->PreviewBlueprintClass;
 			SceneDescription->BlueprintTransform = DataflowAsset->PreviewBlueprintTransform; 
 			SceneDescription->GeometryCacheAsset = Cast<UGeometryCache>(DataflowAsset->PreviewGeometryCacheAsset.LoadSynchronous());
@@ -146,9 +147,14 @@ void FDataflowSimulationScene::ResetSimulationScene()
 	UnbindSceneSelection();
 }
 
+void  FDataflowSimulationScene::SetSimulationLocked(const bool bSimulationLocked, const bool bIsPlaying) 
+{ 
+	bIsSimulationLocked = bSimulationLocked;
+}
+
 void FDataflowSimulationScene::PauseSimulationScene() const
 {
-	if(SceneDescription && (SceneDescription->CacheAsset == nullptr))
+	if(!bIsSimulationLocked)
 	{
 		GetWorld()->GetSubsystem<UDataflowSimulationManager>()->SetSimulationEnabled(false);
 		UE::Dataflow::PauseSkeletonAnimation(PreviewActor);
@@ -157,7 +163,7 @@ void FDataflowSimulationScene::PauseSimulationScene() const
 
 void FDataflowSimulationScene::StartSimulationScene() const
 {
-	if(SceneDescription && (SceneDescription->CacheAsset == nullptr))
+	if(!bIsSimulationLocked)
 	{
 		GetWorld()->GetSubsystem<UDataflowSimulationManager>()->SetSimulationEnabled(true);
 		UE::Dataflow::StartSkeletonAnimation(PreviewActor);
@@ -166,7 +172,7 @@ void FDataflowSimulationScene::StartSimulationScene() const
 
 void FDataflowSimulationScene::StepSimulationScene() const
 {
-	if(SceneDescription && (SceneDescription->CacheAsset == nullptr))
+	if(!bIsSimulationLocked)
 	{
 		GetWorld()->GetSubsystem<UDataflowSimulationManager>()->SetSimulationEnabled(true);
 		GetWorld()->GetSubsystem<UDataflowSimulationManager>()->SetSimulationStepping(true);
@@ -174,27 +180,47 @@ void FDataflowSimulationScene::StepSimulationScene() const
 	}
 }
 
+bool FDataflowSimulationScene::HasCachingEnabled() const
+{
+	return (SceneDescription && (SceneDescription->CacheAsset != nullptr));
+}
+
 bool FDataflowSimulationScene::IsSimulationEnabled() const
 {
-	if(SceneDescription && (SceneDescription->CacheAsset == nullptr))
+	if(!bIsSimulationLocked)
 	{
 		return GetWorld()->GetSubsystem<UDataflowSimulationManager>()->GetSimulationEnabled();
 	}
 	return false;
 }
 
-void FDataflowSimulationScene::RebuildSimulationScene(const bool bIsSimulationEnabled)
+void FDataflowSimulationScene::RebuildSimulationScene()
 {
-	if(SceneDescription && (SceneDescription->CacheAsset == nullptr))
+	// Unregister components, cache manager, selection...
+	ResetSimulationScene();
+
+	// Register components, cache manager, selection...
+	CreateSimulationScene();
+
+	// Override the simulation enabled flag
+	GetWorld()->GetSubsystem<UDataflowSimulationManager>()->SetSimulationEnabled(true);
+
+	// Reset the simulation dirty flag
+	if(const TObjectPtr<UDataflowBaseContent>& EditorContent = GetEditorContent())
 	{
-		// Unregister components, cache manager, selection...
-		ResetSimulationScene();
+		EditorContent->SetSimulationDirty(false);
+	}
 
-		// Register components, cache manager, selection...
-		CreateSimulationScene();
+	// Set the locked simulation flag to false
+	bIsSimulationLocked = false;
 
-		// Override the simulation enabled flag
-		GetWorld()->GetSubsystem<UDataflowSimulationManager>()->SetSimulationEnabled(bIsSimulationEnabled);
+	// Update simulation visualizations
+	for (const TPair<FName, TUniquePtr<UE::Dataflow::IDataflowSimulationVisualization>>& Visualization : UE::Dataflow::FDataflowSimulationVisualizationRegistry::GetInstance().GetVisualizations())
+	{
+		if (Visualization.Value.IsValid())
+		{
+			Visualization.Value->SimulationSceneUpdated(this);
+		}
 	}
 }
 
@@ -217,7 +243,7 @@ void FDataflowSimulationScene::CreateSimulationScene()
 {
 	const int32 PreviewLOD = GetPreviewLOD();
 
-	if(SimulationGenerator && SceneDescription && SceneDescription->BlueprintClass && GetWorld())
+	if(SimulationGenerator && SceneDescription && GetWorld())
 	{
 		SimulationGenerator->SetCacheParams(SceneDescription->CacheParams);
 		SimulationGenerator->SetCacheAsset(SceneDescription->CacheAsset);
@@ -245,7 +271,11 @@ void FDataflowSimulationScene::CreateSimulationScene()
 		}
 		AddSceneObject(RootSceneActor, false);
 		
-		GetWorld()->GetSubsystem<UDataflowSimulationManager>()->SetSimulationEnabled(false);
+		// Enable the simulation by default
+		GetWorld()->GetSubsystem<UDataflowSimulationManager>()->SetSimulationEnabled(!IsSimulationLocked());
+
+		// Pause the animation 
+		UE::Dataflow::PauseSkeletonAnimation(PreviewActor);
 	}
 
 	bPreviewSceneDirty = true;
@@ -258,20 +288,23 @@ void FDataflowSimulationScene::CreateSimulationScene()
 
 void FDataflowSimulationScene::RecordSimulationCache()
 {
-	if(SceneDescription->CacheParams.bAsyncCaching)
+	if(AChaosCacheManager* CacheManager = Cast<AChaosCacheManager>(RootSceneActor))
 	{
-		if(SimulationGenerator.IsValid())
+		if(SceneDescription->CacheParams.bAsyncCaching)
 		{
-			SimulationGenerator->RequestGeneratorAction(UE::Dataflow::EDataflowGeneratorActions::StartGenerate);
+			if(SimulationGenerator.IsValid())
+			{
+				SimulationGenerator->RequestGeneratorAction(UE::Dataflow::EDataflowGeneratorActions::StartGenerate);
+			}
 		}
-	}
-	else
-	{
-		bIsRecordingCache = true;
-        SimulationTime = TimeRange[0] - DeltaTime;
-		if(AChaosCacheManager* CacheManager = Cast<AChaosCacheManager>(RootSceneActor))
+		else
 		{
+			bIsRecordingCache = true;
+			bIsSimulationLocked = false;
+		
+			SimulationTime = TimeRange[0] - DeltaTime;
 			GetWorld()->GetSubsystem<UDataflowSimulationManager>()->SetSimulationEnabled(true);
+			
 			CacheManager->CacheMode = ECacheMode::Record;
 			CacheManager->SetObservedComponentProperties(CacheManager->CacheMode);
 			CacheManager->BeginEvaluate();
@@ -279,30 +312,31 @@ void FDataflowSimulationScene::RecordSimulationCache()
 	}
 }
 
-void FDataflowSimulationScene::UpdateSimulationCache(float& ElapsedTime)
+void FDataflowSimulationScene::UpdateSimulationCache()
 {
-	if(RootSceneActor)
+	if(AChaosCacheManager* CacheManager = Cast<AChaosCacheManager>(RootSceneActor))
 	{
-		AChaosCacheManager* CacheManager = Cast<AChaosCacheManager>(RootSceneActor);
 		if(!bIsRecordingCache)
 		{
-			// Update the cached simulation at some point in time
-			const float ScrollBarTime = SimulationTime - TimeRange[0];
-			if (ScrollBarTime != CacheManager->StartTime)
+			if(bIsSimulationLocked)
 			{
-				CacheManager->SetStartTime(ScrollBarTime);
+				// Update the cached simulation at some point in time
+				const float ScrollBarTime = SimulationTime - TimeRange[0];
+				if (ScrollBarTime != CacheManager->StartTime)
+				{
+					CacheManager->SetStartTime(ScrollBarTime);
+				}
 			}
 		}
 		else
 		{
-			ElapsedTime = DeltaTime;
 			SimulationTime += DeltaTime;
-
 			if(SimulationTime >= TimeRange[1]+UE_SMALL_NUMBER)
 			{
 				bIsRecordingCache = false;
+				bIsSimulationLocked = true;
+				
 				CacheManager->EndEvaluate();
-
 				GetWorld()->GetSubsystem<UDataflowSimulationManager>()->SetSimulationEnabled(false);
 						
 				CacheManager->CacheMode = ECacheMode::None;
@@ -384,44 +418,31 @@ FBox FDataflowSimulationScene::GetBoundingBox() const
 
 void FDataflowSimulationScene::TickDataflowScene(const float DeltaSeconds)
 {
-	float ElapsedTime = DeltaSeconds;
 	if(const TObjectPtr<UDataflowBaseContent>& EditorContent = GetEditorContent())
 	{
 		if (const TObjectPtr<UDataflow> DataflowGraph = EditorContent->GetDataflowAsset())
 		{
-			if (UE::Dataflow::ShouldResetWorld(DataflowGraph, GetWorld(), LastTimeStamp) || EditorContent->IsSimulationDirty())
+			if(!bIsSimulationLocked && (UE::Dataflow::ShouldResetWorld(DataflowGraph, GetWorld(), LastTimeStamp) || EditorContent->IsSimulationDirty()))
 			{
-				// Unregister components, cache manager, selection...
-				ResetSimulationScene();
-
-				// Register components, cache manager, selection...
-				CreateSimulationScene();
-
-				// Reset the dirty flag
-				EditorContent->SetSimulationDirty(false);
-
-				// Update simulation visualizations
-				for (const TPair<FName, TUniquePtr<UE::Dataflow::IDataflowSimulationVisualization>>& Visualization : UE::Dataflow::FDataflowSimulationVisualizationRegistry::GetInstance().GetVisualizations())
-				{
-					if (Visualization.Value.IsValid())
-					{
-						Visualization.Value->SimulationSceneUpdated(this);
-					}
-				}
+				EditorContent->OnContentDataChanged.Broadcast(PreviewActor);
 			}
 		}
 
-		// Load the cache at some point in time
-		if(SceneDescription->CacheAsset)
+		// Update the simulation cache at the simulation time
+		UpdateSimulationCache();
+		
+		// Update all the skelmesh animations at the simulation time
+		UE::Dataflow::UpdateSkeletonAnimation(PreviewActor, SimulationTime);
+
+		// Compute the modified elapsed time after updating the simulation cache (loading/recording) if necessary
+		const float ElapsedTime = bIsRecordingCache ? DeltaTime  : DeltaSeconds;
+
+		// Advance the world ticking if necessary
+		if (ElapsedTime > 0.0f)
 		{
-			// Update the simulation cache (loading/recording)
-			UpdateSimulationCache(ElapsedTime);
-			
-			// Update all the skelmesh animations at the simulation time
-			UE::Dataflow::UpdateSkeletonAnimation(PreviewActor, SimulationTime);
+			GetWorld()->Tick(ELevelTick::LEVELTICK_All, ElapsedTime);
 		}
 	}
-	GetWorld()->Tick(ELevelTick::LEVELTICK_All, ElapsedTime);
 }
 
 void FDataflowSimulationScene::AddReferencedObjects(FReferenceCollector& Collector)
diff --git a/Engine/Plugins/Experimental/Dataflow/Source/DataflowEditor/Private/Dataflow/DataflowSimulationViewport.cpp b/Engine/Plugins/Experimental/Dataflow/Source/DataflowEditor/Private/Dataflow/DataflowSimulationViewport.cpp
--- a/Engine/Plugins/Experimental/Dataflow/Source/DataflowEditor/Private/Dataflow/DataflowSimulationViewport.cpp
+++ b/Engine/Plugins/Experimental/Dataflow/Source/DataflowEditor/Private/Dataflow/DataflowSimulationViewport.cpp
@@ -19,6 +19,7 @@
 #include "LODSyncInterface.h"
 #include "ViewportToolbar/UnrealEdViewportToolbar.h"
 #include "Widgets/Text/SRichTextBlock.h"
+#include "ToolMenuSection.h"
 
 #define LOCTEXT_NAMESPACE "SDataflowSimulationViewport"
 
@@ -43,36 +44,6 @@ void SDataflowSimulationViewport::Construct(const FArguments& InArgs, const FAss
 	{
 		TWeakPtr<FDataflowSimulationScene> SimulationScene = GetSimulationScene();
 
-		auto HasCacheAsset = [SimulationScene]()
-		{
-			if(SimulationScene.Pin() && SimulationScene.Pin()->GetPreviewSceneDescription())
-			{
-				return (SimulationScene.Pin()->GetPreviewSceneDescription()->CacheAsset != nullptr) ? EVisibility::Visible : EVisibility::Collapsed;
-			}
-			return EVisibility::Collapsed;
-		};
-            
-		ViewportOverlay->AddSlot()
-		[
-			SNew(SHorizontalBox)
-			+ SHorizontalBox::Slot()
-			.HAlign(HAlign_Fill)
-			.VAlign(VAlign_Bottom)
-			.FillWidth(1)
-			.Padding(10.0f, 0.0f)
-			[
-				SNew(SBorder)
-				.BorderImage(FAppStyle::Get().GetBrush("EditorViewport.OverlayBrush"))
-				.Visibility(TAttribute<EVisibility>::Create(HasCacheAsset))
-				.Padding(10.0f, 2.0f)
-				[
-					SNew(SDataflowSimulationPanel, SimulationScene)
-					.ViewInputMin(this, &SDataflowSimulationViewport::GetViewMinInput)
-					.ViewInputMax(this, &SDataflowSimulationViewport::GetViewMaxInput)
-				]
-			]
-		];
-
 		ViewportOverlay->AddSlot()
 		[
 			SNew(SHorizontalBox)
@@ -136,6 +107,13 @@ TSharedPtr<SWidget> SDataflowSimulationViewport::BuildViewportToolbar()
 			RightSection.AddEntry(UE::UnrealEd::CreateCameraSubmenu(UE::UnrealEd::FViewportCameraMenuOptions().ShowCameraMovement()));
 			RightSection.AddEntry(UE::UnrealEd::CreateViewModesSubmenu());
 			RightSection.AddEntry(UE::UnrealEd::CreateDefaultShowSubmenu());
+			RightSection.AddEntry(UE::UnrealEd::CreateAssetViewerProfileSubmenu());
+			{
+				RightSection.AddDynamicEntry("ViewAndVisualization", FNewToolMenuSectionDelegate::CreateLambda([this](FToolMenuSection& Section)
+				{
+					BuildVisualizationMenu(Section);	
+				}));
+			}
 			
 			{
 				// LOD
@@ -148,81 +126,6 @@ TSharedPtr<SWidget> SDataflowSimulationViewport::BuildViewportToolbar()
 					}
 				}));
 			}
-			
-			{
-				RightSection.AddEntry(UE::UnrealEd::CreateAssetViewerProfileSubmenu());
-				UToolMenu* AssetViewerMenu = UToolMenus::Get()->ExtendMenu(UToolMenus::JoinMenuPaths(ToolbarName, "AssetViewerProfile"));
-				FToolMenuSection& SimulationSection = AssetViewerMenu->FindOrAddSection("SimulationControls", LOCTEXT("SimulationControlsSection", "Simulation Scene Controls"));
-				
-				SimulationSection.AddDynamicEntry("DynamicSceneControls", FNewToolMenuSectionDelegate::CreateLambda([](FToolMenuSection& Section)
-				{
-					UUnrealEdViewportToolbarContext* Context = Section.FindContext<UUnrealEdViewportToolbarContext>();
-					if (!Context)
-					{
-						return;
-					}
-					
-					TWeakPtr<SDataflowSimulationViewport> WeakViewport = StaticCastWeakPtr<SDataflowSimulationViewport>(Context->Viewport);
-					if (!WeakViewport.IsValid())
-					{
-						return;
-					}
-				
-					// Simulation Controls
-					FToolMenuEntryToolBarData ToolBarData;
-					ToolBarData.BlockGroupName = "SimulationControlGroup";
-					ToolBarData.LabelOverride = FText::GetEmpty();
-					ToolBarData.ResizeParams.ClippingPriority = 2000;
-					
-					FToolMenuEntry& Rebuild = Section.AddMenuEntry(FDataflowEditorCommands::Get().RebuildSimulationScene);
-					Rebuild.ToolBarData = ToolBarData;
-					Rebuild.Icon = FSlateIcon(FAppStyle::Get().GetStyleSetName(), "Animation.Backward_End");
-					Rebuild.SetShowInToolbarTopLevel(true);
-					
-					FToolMenuEntry& Pause = Section.AddMenuEntry(FDataflowEditorCommands::Get().PauseSimulationScene);
-					Pause.ToolBarData = ToolBarData;
-					Pause.Icon = FSlateIcon(FAppStyle::Get().GetStyleSetName(), "Animation.Pause");
-					Pause.SetShowInToolbarTopLevel(true);
-					Pause.Visibility = [WeakViewport]
-					{
-						if (TSharedPtr<SDataflowSimulationViewport> Viewport = WeakViewport.Pin())
-						{
-							if (TSharedPtr<FEditorViewportClient> Client = Viewport->GetViewportClient())
-							{
-								if (FDataflowSimulationScene* SimulationScene = StaticCast<FDataflowSimulationScene*>(Client->GetPreviewScene()))
-								{
-									return SimulationScene->IsSimulationEnabled();
-								}
-							}
-						}
-						return true;
-					};
-					
-					FToolMenuEntry& Start = Section.AddMenuEntry(FDataflowEditorCommands::Get().StartSimulationScene);
-					Start.ToolBarData = ToolBarData;
-					Start.Icon = FSlateIcon(FAppStyle::Get().GetStyleSetName(), "Icons.Play");
-					Start.SetShowInToolbarTopLevel(true);
-					Start.Visibility = [WeakViewport]
-                    {
-                    	if (TSharedPtr<SDataflowSimulationViewport> Viewport = WeakViewport.Pin())
-                    	{
-                    		if (TSharedPtr<FEditorViewportClient> Client = Viewport->GetViewportClient())
-                    		{
-                    			if (FDataflowSimulationScene* SimulationScene = StaticCast<FDataflowSimulationScene*>(Client->GetPreviewScene()))
-                    			{
-                    				return !SimulationScene->IsSimulationEnabled();
-                    			}
-                    		}
-                    	}
-                    	return true;
-                    };
-					
-					FToolMenuEntry& Step = Section.AddMenuEntry(FDataflowEditorCommands::Get().StepSimulationScene);
-					Step.ToolBarData = ToolBarData;
-					Step.Icon = FSlateIcon(FAppStyle::Get().GetStyleSetName(), "Animation.Forward_Step");
-					Step.SetShowInToolbarTopLevel(true);
-				}));	
-			}
 		}
 	}
 	
@@ -238,6 +141,34 @@ TSharedPtr<SWidget> SDataflowSimulationViewport::BuildViewportToolbar()
 	return UToolMenus::Get()->GenerateWidget(ToolbarName, Context);
 }
 
+void SDataflowSimulationViewport::BuildVisualizationMenu(FToolMenuSection& MenuSection)
+{
+	if (UUnrealEdViewportToolbarContext* Context = MenuSection.FindContext<UUnrealEdViewportToolbarContext>())
+	{
+		MenuSection.AddSubMenu(
+			"Visualization",
+			LOCTEXT("VisualizationSubmenuLabel", "Visualization"),
+			LOCTEXT("VisualizationSubmenuTooltip", "Controls visualization options"),
+			FNewMenuDelegate::CreateLambda([WeakViewport = Context->Viewport](FMenuBuilder& Menu)
+				{
+					if (const TSharedPtr<SDataflowSimulationViewport> Viewport = StaticCastSharedPtr<SDataflowSimulationViewport>(WeakViewport.Pin()))
+					{
+						if (const TSharedPtr<FEditorViewportClient> ViewportClient = Viewport->GetViewportClient())
+						{
+							for (const TPair<FName, TUniquePtr<UE::Dataflow::IDataflowSimulationVisualization>>& Visualization : UE::Dataflow::FDataflowSimulationVisualizationRegistry::GetInstance().GetVisualizations())
+							{
+								Visualization.Value->ExtendSimulationVisualizationMenu(StaticCastSharedPtr<FDataflowSimulationViewportClient>(ViewportClient), Menu);
+							}
+						}
+					}
+				}),
+			false,
+			FSlateIcon(FAppStyle::Get().GetStyleSetName(), "DetailsView.ViewOptions")
+		);
+	}
+}
+
+
 TSharedPtr<IPreviewProfileController> SDataflowSimulationViewport::CreatePreviewProfileController()
 {
 	TSharedPtr<FDataflowSimulationSceneProfileIndexStorage> ProfileIndexStorage;
@@ -285,56 +216,6 @@ UDataflowEditorMode* SDataflowSimulationViewport::GetEdMode() const
 void SDataflowSimulationViewport::BindCommands()
 {
 	SAssetEditorViewport::BindCommands();
-	
-	const FDataflowEditorCommandsImpl& CommandInfos = FDataflowEditorCommands::Get();
-
-	CommandList->MapAction(
-		CommandInfos.RebuildSimulationScene,
-		FExecuteAction::CreateLambda([this]()
-		{
-			if(FDataflowSimulationScene* SimulationScene = StaticCast<FDataflowSimulationScene*>(Client->GetPreviewScene()))
-			{
-				SimulationScene->RebuildSimulationScene(false);
-			}
-		}),
-		FCanExecuteAction::CreateLambda([this]() { return true; }),
-		FIsActionChecked::CreateLambda([this]() { return false; }));
-
-	CommandList->MapAction(
-		CommandInfos.PauseSimulationScene,
-		FExecuteAction::CreateLambda([this]()
-		{
-			if(FDataflowSimulationScene* SimulationScene = StaticCast<FDataflowSimulationScene*>(Client->GetPreviewScene()))
-			{
-				SimulationScene->PauseSimulationScene();
-			}
-		}),
-		FCanExecuteAction::CreateLambda([this]() { return true; }),
-		FIsActionChecked::CreateLambda([this]() { return false; }));
-
-	CommandList->MapAction(
-		CommandInfos.StartSimulationScene,
-		FExecuteAction::CreateLambda([this]()
-		{
-			if(FDataflowSimulationScene* SimulationScene = StaticCast<FDataflowSimulationScene*>(Client->GetPreviewScene()))
-			{
-				SimulationScene->StartSimulationScene();
-			}
-		}),
-		FCanExecuteAction::CreateLambda([this]() { return true; }),
-		FIsActionChecked::CreateLambda([this]() { return false; }));
-
-	CommandList->MapAction(
-		CommandInfos.StepSimulationScene,
-		FExecuteAction::CreateLambda([this]()
-		{
-			if(FDataflowSimulationScene* SimulationScene = StaticCast<FDataflowSimulationScene*>(Client->GetPreviewScene()))
-			{
-				SimulationScene->StepSimulationScene();
-			}
-		}),
-		FCanExecuteAction::CreateLambda([this]() { return true; }),
-		FIsActionChecked::CreateLambda([this]() { return false; }));
 }
 
 bool SDataflowSimulationViewport::IsVisible() const
diff --git a/Engine/Plugins/Experimental/Dataflow/Source/DataflowEditor/Private/Widgets/DataflowSimulationBinding.cpp b/Engine/Plugins/Experimental/Dataflow/Source/DataflowEditor/Private/Widgets/DataflowSimulationBinding.cpp
new file mode 100644
--- /dev/null
+++ b/Engine/Plugins/Experimental/Dataflow/Source/DataflowEditor/Private/Widgets/DataflowSimulationBinding.cpp
@@ -0,0 +1,180 @@
+// Copyright Epic Games, Inc. All Rights Reserved.
+
+#include "DataflowSimulationBinding.h"
+#include "CommonFrameRates.h"
+#include "Chaos/CacheCollection.h"
+#include "Chaos/ChaosCache.h"
+#include "Dataflow/DataflowSimulationScene.h"
+#include "Fonts/FontMeasure.h"
+
+FDataflowSimulationBinding::FDataflowSimulationBinding(const TWeakPtr<FDataflowSimulationScene>&  InSimulationScene)
+	: SimulationScene(InSimulationScene)
+{
+	if (const TSharedPtr<const FDataflowSimulationScene> PreviewScene = SimulationScene.Pin())
+	{
+		PlaybackRange = FAnimatedRange(PreviewScene->GetTimeRange()[0], PreviewScene->GetTimeRange()[1]);
+		SetViewRange(PlaybackRange);
+	}
+}
+
+FFrameRate FDataflowSimulationBinding::GetFrameRate() const
+{
+	if (const TSharedPtr<const FDataflowSimulationScene> PreviewScene = SimulationScene.Pin())
+	{
+		return FFrameRate(PreviewScene->GetFrameRate(), 1);
+	}
+	return FFrameRate();
+}
+
+int32 FDataflowSimulationBinding::GetTickResolution() const
+{
+	return int32(1000);
+}
+
+FAnimatedRange FDataflowSimulationBinding::GetViewRange() const
+{
+	return ViewRange;
+}
+
+void FDataflowSimulationBinding::SetViewRange(TRange<double> InRange)
+{
+	ViewRange = InRange;
+
+	if (WorkingRange.HasLowerBound() && WorkingRange.HasUpperBound())
+	{
+		WorkingRange = TRange<double>::Hull(WorkingRange, ViewRange);
+	}
+	else
+	{
+		WorkingRange = ViewRange;
+	}
+}
+
+FAnimatedRange FDataflowSimulationBinding::GetWorkingRange() const
+{
+	return WorkingRange;
+}
+
+TRange<FFrameNumber> FDataflowSimulationBinding::GetPlaybackRange() const
+{
+	const float TickResolution = static_cast<float>(GetTickResolution());
+	return TRange<FFrameNumber>(FFrameNumber(FMath::RoundToInt32(PlaybackRange.GetLowerBoundValue() * TickResolution)),
+								FFrameNumber(FMath::RoundToInt32(PlaybackRange.GetUpperBoundValue() * TickResolution)));
+}
+
+FFrameNumber FDataflowSimulationBinding::GetScrubPosition() const
+{
+	if (const TSharedPtr<const FDataflowSimulationScene> PreviewScene = SimulationScene.Pin())
+	{
+		return FFrameNumber(FMath::RoundToInt32(GetScrubTime() * (double)GetTickResolution()));
+	}
+
+	return FFrameNumber(0);
+}
+
+float FDataflowSimulationBinding::GetScrubTime() const
+{
+	if (const TSharedPtr<FDataflowSimulationScene> PreviewScene = SimulationScene.Pin())
+	{
+		return PreviewScene->SimulationTime-PreviewScene->GetTimeRange()[0];
+	}
+
+	return 0.0f;
+}
+
+void FDataflowSimulationBinding::SetScrubPosition(FFrameTime ScrubPosition) const
+{
+	if (const TSharedPtr<FDataflowSimulationScene> PreviewScene = SimulationScene.Pin())
+	{
+		const float TickResolution = static_cast<float>(GetTickResolution());
+		const float ScrubTime = ScrubPosition.AsDecimal() / TickResolution;
+		PreviewScene->SimulationTime = ScrubTime + PreviewScene->GetTimeRange()[0];
+	}
+}
+
+void FDataflowSimulationBinding::SetScrubTime(const float ScrubTime)
+{
+	if (const TSharedPtr<FDataflowSimulationScene> PreviewScene = SimulationScene.Pin())
+	{
+		PreviewScene->SimulationTime = ScrubTime + PreviewScene->GetTimeRange()[0];
+	}
+}
+
+float FDataflowSimulationBinding::GetDeltaTime() const
+{
+	if (const TSharedPtr<FDataflowSimulationScene> PreviewScene = SimulationScene.Pin())
+	{
+		return PreviewScene->GetDeltaTime();
+	}
+	return 0.0f;
+}
+
+float FDataflowSimulationBinding::GetSequenceLength() const
+{
+	if (const TSharedPtr<FDataflowSimulationScene> PreviewScene = SimulationScene.Pin())
+	{
+		return PreviewScene->GetTimeRange()[1] - PreviewScene->GetTimeRange()[0];
+	}
+	return 0.0f;
+}
+
+void FDataflowSimulationBinding::RecordSimulationCache()
+{
+	if (const TSharedPtr<FDataflowSimulationScene> PreviewScene = SimulationScene.Pin())
+	{
+		return PreviewScene->RecordSimulationCache();
+	}
+}
+
+void FDataflowSimulationBinding::ResetSimulationScene()
+{
+	if (const TSharedPtr<FDataflowSimulationScene> PreviewScene = SimulationScene.Pin())
+	{
+		return PreviewScene->RebuildSimulationScene();
+	}
+}
+
+bool FDataflowSimulationBinding::IsSimulationLocked() const
+{
+	if (const TSharedPtr<FDataflowSimulationScene> PreviewScene = SimulationScene.Pin())
+	{
+		return PreviewScene->IsSimulationLocked();
+	}
+	return false;
+}
+
+void FDataflowSimulationBinding::SetSimulationLocked(const bool bIsSimulaitonLocked)
+{
+	if (const TSharedPtr<FDataflowSimulationScene> PreviewScene = SimulationScene.Pin())
+	{
+		PreviewScene->SetSimulationLocked(bIsSimulaitonLocked, false);
+	}
+}
+
+void FDataflowSimulationBinding::FillCacheNames(TArray<FString>& CacheNames) const
+{
+	CacheNames.Reset();
+	if (const TSharedPtr<FDataflowSimulationScene> PreviewScene = SimulationScene.Pin())
+	{
+		if(UDataflowSimulationSceneDescription* SceneDescription = PreviewScene->GetPreviewSceneDescription())
+		{
+			if(SceneDescription->CacheAsset)
+			{
+				for(UChaosCache* Cache : SceneDescription->CacheAsset->Caches)
+				{
+					CacheNames.Add(Cache->GetName());
+				}
+			}
+		}
+	}
+}
+
+void FDataflowSimulationBinding::HandleViewRangeChanged(TRange<double> InRange, EViewRangeInterpolation InInterpolation)
+{
+	SetViewRange(InRange);
+}
+
+void FDataflowSimulationBinding::HandleWorkingRangeChanged(TRange<double> InRange)
+{
+	WorkingRange = InRange;
+}
\ No newline at end of file
diff --git a/Engine/Plugins/Experimental/Dataflow/Source/DataflowEditor/Private/Widgets/DataflowSimulationBinding.h b/Engine/Plugins/Experimental/Dataflow/Source/DataflowEditor/Private/Widgets/DataflowSimulationBinding.h
new file mode 100644
--- /dev/null
+++ b/Engine/Plugins/Experimental/Dataflow/Source/DataflowEditor/Private/Widgets/DataflowSimulationBinding.h
@@ -0,0 +1,91 @@
+﻿// Copyright Epic Games, Inc. All Rights Reserved.
+
+#pragma once
+
+#include "CoreMinimal.h"
+#include "Widgets/SWidget.h"
+#include "ITimeSlider.h"
+#include "TimeSliderArgs.h"
+#include "Dataflow/DataflowSimulationScene.h"
+#include "Editor/SequencerWidgets/Public/AnimatedRange.h"
+
+/** Dataflow binding to centralize the bridge in between the timeline widget and the dataflow simulation scene */
+class FDataflowSimulationBinding : public TSharedFromThis<FDataflowSimulationBinding>
+{
+public:
+	FDataflowSimulationBinding(const TWeakPtr<FDataflowSimulationScene>&  InSimulationScene);
+
+	/** Get the simulation scene from the binding */
+	TWeakPtr<FDataflowSimulationScene> GetSimulationScene() const { return SimulationScene; }
+
+	/** Get the framerate specified by the anim sequence */
+	FFrameRate GetFrameRate() const;
+
+	/** Get the tick resolution we are displaying at */
+	int32 GetTickResolution() const;
+
+	/** Get the current view range */
+	FAnimatedRange GetViewRange() const;
+
+	/** Set the current view range */
+	void SetViewRange(TRange<double> InRange);
+
+	/** Get the working range of the model's data */
+	FAnimatedRange GetWorkingRange() const;
+
+	/** Get the playback range of the model's data */
+	TRange<FFrameNumber> GetPlaybackRange() const;
+
+	/** Get the current scrub position */
+	FFrameNumber GetScrubPosition() const;
+
+	/** Get the current scrub time */
+	float GetScrubTime() const;
+
+	/** Set the current scrub time */
+	void SetScrubTime(const float NewScrubTime);
+	
+	/** Get the simulation delta time*/
+	float GetDeltaTime() const;
+
+	/** Get the simulation sequence length */
+	float GetSequenceLength() const;
+
+	/** Set the current scrub position */
+	void SetScrubPosition(FFrameTime NewScrubPostion) const;
+
+	/** Handle the view range being changed */
+	void HandleViewRangeChanged(TRange<double> InRange, EViewRangeInterpolation InInterpolation);
+
+	/** Handle the working range being changed */
+	void HandleWorkingRangeChanged(TRange<double> InRange);
+
+	/** Handle the recording of the simulation cache */
+	void RecordSimulationCache();
+
+	/** Handle the reset of the simulation scene */
+	void ResetSimulationScene();
+
+	/** Get the locked simulation flag */
+	bool IsSimulationLocked() const;
+
+	/** Set the locked simulation flag */
+	void SetSimulationLocked(const bool bIsSimulaitonLocked);
+
+	/** Retrieve all the cache names */
+	void FillCacheNames(TArray<FString>& CacheNames) const;
+
+private:
+	/** Simulation scene to be used for the widget */
+	TWeakPtr<FDataflowSimulationScene> SimulationScene;
+
+	/** View simulation range */
+	FAnimatedRange ViewRange;
+
+	/** Working simulation range */
+	FAnimatedRange WorkingRange;
+
+	/** Playback simulation range */
+	FAnimatedRange PlaybackRange;
+
+};
diff --git a/Engine/Plugins/Experimental/Dataflow/Source/DataflowEditor/Private/Widgets/DataflowTimeSliderController.cpp b/Engine/Plugins/Experimental/Dataflow/Source/DataflowEditor/Private/Widgets/DataflowTimeSliderController.cpp
new file mode 100644
--- /dev/null
+++ b/Engine/Plugins/Experimental/Dataflow/Source/DataflowEditor/Private/Widgets/DataflowTimeSliderController.cpp
@@ -0,0 +1,1003 @@
+// Copyright Epic Games, Inc. All Rights Reserved.
+
+#include "DataflowTimeSliderController.h"
+
+#include "CommonFrameRates.h"
+#include "Fonts/FontMeasure.h"
+#include "Framework/Application/SlateApplication.h"
+#include "MovieSceneTimeHelpers.h"
+#include "Styling/AppStyle.h"
+#include "Widgets/SDataflowSimulationTimeline.h"
+
+namespace ScrubConstants
+{
+	/**The fraction of the current view range to scroll per unit delta  */
+	constexpr float ScrollPanFraction = 0.1f;
+
+}
+
+struct FDataflowTimeSliderController::FScrubRangeToScreen
+{
+	double ViewStart;
+	float  PixelsPerInput;
+
+	FScrubRangeToScreen(const TRange<double>& InViewInput, const FVector2D& InWidgetSize)
+	{
+		const double ViewInputRange = InViewInput.Size<double>();
+
+		ViewStart = InViewInput.GetLowerBoundValue();
+		PixelsPerInput = ViewInputRange > 0 ? static_cast<float>(InWidgetSize.X / ViewInputRange) : 0;
+	}
+
+	/** Local Widget Space -> Curve Input domain. */
+	double LocalXToInput(float ScreenX) const
+	{
+		return PixelsPerInput > 0 ? (ScreenX / PixelsPerInput) + ViewStart : ViewStart;
+	}
+
+	/** Curve Input domain -> local Widget Space */
+	float InputToLocalX(double Input) const
+	{
+		return static_cast<float>((Input - ViewStart) * PixelsPerInput);
+	}
+};
+
+struct FDataflowTimeSliderController::FDrawTickArgs
+{
+	/** Geometry of the area */
+	FGeometry AllottedGeometry;
+	
+	/** Culling rect of the area */
+	FSlateRect CullingRect;
+	
+	/** Color of each tick */
+	FLinearColor TickColor;
+	
+	/** Offset in Y where to start the tick */
+	float TickOffset;
+	
+	/** Height in of major ticks */
+	float MajorTickHeight;
+	
+	/** Start layer for elements */
+	int32 StartLayer;
+	
+	/** Draw effects to apply */
+	ESlateDrawEffect DrawEffects;
+	
+	/** Whether or not to only draw major ticks */
+	bool bOnlyDrawMajorTicks;
+	
+	/** Whether or not to mirror labels */
+	bool bMirrorLabels;
+
+};
+
+FDataflowTimeSliderController::FDataflowTimeSliderController(const FTimeSliderArgs& InArgs, TWeakPtr<SDataflowSimulationTimeline> InSimulationTimeline)
+	: TimeSliderArgs(InArgs)
+	, SimulationTimeline(InSimulationTimeline)
+{
+	ScrubFillBrush = FAppStyle::GetBrush(TEXT("Sequencer.Timeline.ScrubFill"));
+	ScrubHandleUpBrush = FAppStyle::GetBrush(TEXT("Sequencer.Timeline.VanillaScrubHandleUp"));
+	ScrubHandleDownBrush = FAppStyle::GetBrush(TEXT("Sequencer.Timeline.VanillaScrubHandleDown"));
+	EditableTimeBrush = FAppStyle::GetBrush(TEXT("AnimTimeline.SectionMarker"));
+}
+
+FDataflowTimeSliderController::~FDataflowTimeSliderController()
+{
+}
+
+int32 FDataflowTimeSliderController::OnPaintTimeSlider(bool bMirrorLabels, const FGeometry& AllottedGeometry, const FSlateRect& MyCullingRect, FSlateWindowElementList& OutDrawElements, int32 LayerId, const FWidgetStyle& InWidgetStyle, bool bParentEnabled) const
+{
+	const bool bEnabled = bParentEnabled;
+	const ESlateDrawEffect DrawEffects = bEnabled ? ESlateDrawEffect::None : ESlateDrawEffect::DisabledEffect;
+
+	const TRange<double> LocalViewRange = TimeSliderArgs.ViewRange.Get();
+	const float    LocalViewRangeMin = static_cast<float>(LocalViewRange.GetLowerBoundValue());
+	const float    LocalViewRangeMax = static_cast<float>(LocalViewRange.GetUpperBoundValue());
+	const float    LocalSequenceLength = LocalViewRangeMax - LocalViewRangeMin;
+	const TRange<FFrameNumber> LocalPlaybackRange = TimeSliderArgs.PlaybackRange.Get();
+
+	if (LocalSequenceLength > 0)
+	{
+		FScrubRangeToScreen RangeToScreen(LocalViewRange, AllottedGeometry.Size);
+
+		// draw tick marks
+		constexpr float MajorTickHeight = 9.0f;
+
+		FDrawTickArgs Args;
+		{
+			Args.AllottedGeometry = AllottedGeometry;
+			Args.bMirrorLabels = bMirrorLabels;
+			Args.bOnlyDrawMajorTicks = false;
+			Args.TickColor = FLinearColor::White;
+			Args.CullingRect = MyCullingRect;
+			Args.DrawEffects = DrawEffects;
+			Args.StartLayer = LayerId;
+			Args.TickOffset = bMirrorLabels ? 0.0f : FMath::Abs(static_cast<float>(AllottedGeometry.Size.Y) - MajorTickHeight);
+			Args.MajorTickHeight = MajorTickHeight;
+		}
+
+		DrawTicks(OutDrawElements, LocalViewRange, RangeToScreen, Args);
+
+		// draw playback & selection range
+		FPaintPlaybackRangeArgs PlaybackRangeArgs(
+			bMirrorLabels ? FAppStyle::GetBrush("Sequencer.Timeline.PlayRange_Bottom_L") : FAppStyle::GetBrush("Sequencer.Timeline.PlayRange_Top_L"),
+			bMirrorLabels ? FAppStyle::GetBrush("Sequencer.Timeline.PlayRange_Bottom_R") : FAppStyle::GetBrush("Sequencer.Timeline.PlayRange_Top_R"),
+			6.f
+		);
+
+		LayerId = DrawPlaybackRange(AllottedGeometry, MyCullingRect, OutDrawElements, LayerId, RangeToScreen, PlaybackRangeArgs);
+
+		PlaybackRangeArgs.SolidFillOpacity = 0.05f;
+		LayerId = DrawSelectionRange(AllottedGeometry, MyCullingRect, OutDrawElements, LayerId, RangeToScreen, PlaybackRangeArgs);
+
+		// Draw the scrub handle
+		const float      HandleStart = RangeToScreen.InputToLocalX(TimeSliderArgs.ScrubPosition.Get().AsDecimal() / GetTickResolution().AsDecimal()) - 7.0f;
+		const float      HandleEnd = HandleStart + 13.0f;
+
+		const int32 ArrowLayer = LayerId + 2;
+		FPaintGeometry MyGeometry = AllottedGeometry.ToPaintGeometry(FVector2f(HandleEnd - HandleStart, AllottedGeometry.Size.Y), FSlateLayoutTransform(FVector2f(HandleStart, 0.f)));
+		FLinearColor ScrubColor = InWidgetStyle.GetColorAndOpacityTint();
+		{
+			ScrubColor.A = ScrubColor.A * 0.75f;
+			ScrubColor.B *= 0.1f;
+			ScrubColor.G *= 0.2f;
+		}
+
+		const FSlateBrush* Brush = (bMirrorLabels ? ScrubHandleUpBrush : ScrubHandleDownBrush);
+
+		FSlateDrawElement::MakeBox(
+			OutDrawElements,
+			ArrowLayer,
+			MyGeometry,
+			Brush,
+			DrawEffects,
+			ScrubColor
+		);
+
+		{
+			// Draw the current time next to the scrub handle
+			FString FrameString = TimeSliderArgs.NumericTypeInterface->ToString(TimeSliderArgs.ScrubPosition.Get().GetFrame().Value);
+			FSlateFontInfo SmallLayoutFont = FCoreStyle::GetDefaultFontStyle("Regular", 10);
+
+			const TSharedRef< FSlateFontMeasure > FontMeasureService = FSlateApplication::Get().GetRenderer()->GetFontMeasureService();
+			const FVector2D TextSize = FontMeasureService->Measure(FrameString, SmallLayoutFont);
+
+			// Flip the text position if getting near the end of the view range
+			constexpr float TextOffsetPx = 2.f;
+			const bool  bDrawLeft = (static_cast<float>(AllottedGeometry.Size.X) - HandleEnd) < (TextSize.X + 14.f) - TextOffsetPx;
+			const float TextPosition = bDrawLeft ? HandleStart - static_cast<float>(TextSize.X) - TextOffsetPx : HandleEnd + TextOffsetPx;
+
+			FVector2D TextOffset(TextPosition, Args.bMirrorLabels ? TextSize.Y - 6.f : Args.AllottedGeometry.Size.Y - (Args.MajorTickHeight + TextSize.Y));
+
+			FSlateDrawElement::MakeText(
+				OutDrawElements,
+				Args.StartLayer + 1,
+				Args.AllottedGeometry.ToPaintGeometry(TextSize, FSlateLayoutTransform(TextOffset)),
+				FrameString,
+				SmallLayoutFont,
+				Args.DrawEffects,
+				Args.TickColor
+			);
+		}
+
+		return ArrowLayer;
+	}
+
+	return LayerId;
+}
+
+int32 FDataflowTimeSliderController::OnPaintViewArea(const FGeometry& AllottedGeometry, const FSlateRect& MyCullingRect, FSlateWindowElementList& OutDrawElements, int32 LayerId, bool bEnabled, const FPaintViewAreaArgs& Args) const
+{
+	const ESlateDrawEffect DrawEffects = bEnabled ? ESlateDrawEffect::None : ESlateDrawEffect::DisabledEffect;
+
+	TRange<double> LocalViewRange = TimeSliderArgs.ViewRange.Get();
+	FScrubRangeToScreen RangeToScreen(LocalViewRange, AllottedGeometry.Size);
+
+	if (Args.PlaybackRangeArgs.IsSet())
+	{
+		FPaintPlaybackRangeArgs PaintArgs = Args.PlaybackRangeArgs.GetValue();
+		LayerId = DrawPlaybackRange(AllottedGeometry, MyCullingRect, OutDrawElements, LayerId, RangeToScreen, PaintArgs);
+		PaintArgs.SolidFillOpacity = 0.2f;
+		LayerId = DrawSelectionRange(AllottedGeometry, MyCullingRect, OutDrawElements, LayerId, RangeToScreen, PaintArgs);
+	}
+
+	if (Args.bDisplayTickLines)
+	{
+		static FLinearColor TickColor(0.1f, 0.1f, 0.1f, 0.3f);
+
+		// Draw major tick lines in the section area
+		FDrawTickArgs DrawTickArgs;
+		{
+			DrawTickArgs.AllottedGeometry = AllottedGeometry;
+			DrawTickArgs.bMirrorLabels = false;
+			DrawTickArgs.bOnlyDrawMajorTicks = true;
+			DrawTickArgs.TickColor = TickColor;
+			DrawTickArgs.CullingRect = MyCullingRect;
+			DrawTickArgs.DrawEffects = DrawEffects;
+			// Draw major ticks under sections
+			DrawTickArgs.StartLayer = LayerId - 1;
+			// Draw the tick the entire height of the section area
+			DrawTickArgs.TickOffset = 0.0f;
+			DrawTickArgs.MajorTickHeight = static_cast<float>(AllottedGeometry.Size.Y);
+		}
+
+		DrawTicks(OutDrawElements, LocalViewRange, RangeToScreen, DrawTickArgs);
+	}
+
+	if (Args.bDisplayScrubPosition)
+	{
+		// Draw a line for the scrub position
+		const float LinePos = RangeToScreen.InputToLocalX(TimeSliderArgs.ScrubPosition.Get().AsDecimal() / GetTickResolution().AsDecimal());
+
+		TArray<FVector2D> LinePoints;
+		{
+			LinePoints.AddUninitialized(2);
+			LinePoints[0] = FVector2D(0.0f, 0.0f);
+			LinePoints[1] = FVector2D(0.0f, FMath::FloorToFloat(AllottedGeometry.Size.Y));
+		}
+
+		FSlateDrawElement::MakeLines(
+			OutDrawElements,
+			LayerId + 1,
+			AllottedGeometry.ToPaintGeometry(FVector2f(1.0f, 1.0f), FSlateLayoutTransform(FVector2f(LinePos, 0.0f))),
+			LinePoints,
+			DrawEffects,
+			FLinearColor(1.f, 1.f, 1.f, .5f),
+			false
+		);
+	}
+
+	return LayerId;
+}
+
+FReply FDataflowTimeSliderController::OnMouseButtonDown(SWidget& WidgetOwner, const FGeometry& MyGeometry, const FPointerEvent& MouseEvent)
+{
+	DistanceDragged = 0;
+	MouseDownPosition[0] = MouseDownPosition[1] = MouseEvent.GetScreenSpacePosition();
+	MouseDownGeometry = MyGeometry;	
+	return FReply::Unhandled();
+}
+
+FReply FDataflowTimeSliderController::OnMouseButtonUp(SWidget& WidgetOwner, const FGeometry& MyGeometry, const FPointerEvent& MouseEvent)
+{
+	bool bHandleLeftMouseButton = MouseEvent.GetEffectingButton() == EKeys::LeftMouseButton && WidgetOwner.HasMouseCapture();
+	bool bHandleRightMouseButton = MouseEvent.GetEffectingButton() == EKeys::RightMouseButton && WidgetOwner.HasMouseCapture() && TimeSliderArgs.AllowZoom;
+
+	FScrubRangeToScreen RangeToScreen = FScrubRangeToScreen(TimeSliderArgs.ViewRange.Get(), MyGeometry.Size);
+	FFrameTime          MouseTime = ComputeFrameTimeFromMouse(MyGeometry, MouseEvent.GetScreenSpacePosition(), RangeToScreen);
+
+	if (bHandleRightMouseButton)
+	{
+		if (!bPanning && DistanceDragged == 0.0f)
+		{
+			return SimulationTimeline.Pin()->OnMouseButtonUp(MyGeometry, MouseEvent).ReleaseMouseCapture();
+		}
+
+		bPanning = false;
+		DistanceDragged = 0.f;
+
+		return FReply::Handled().ReleaseMouseCapture();
+	}
+	else if (bHandleLeftMouseButton)
+	{
+		if (MouseDragType == DRAG_PLAYBACK_START)
+		{
+			TimeSliderArgs.OnPlaybackRangeEndDrag.ExecuteIfBound();
+		}
+		else if (MouseDragType == DRAG_PLAYBACK_END)
+		{
+			TimeSliderArgs.OnPlaybackRangeEndDrag.ExecuteIfBound();
+		}
+		else if (MouseDragType == DRAG_SELECTION_START)
+		{
+			TimeSliderArgs.OnSelectionRangeEndDrag.ExecuteIfBound();
+		}
+		else if (MouseDragType == DRAG_SELECTION_END)
+		{
+			TimeSliderArgs.OnSelectionRangeEndDrag.ExecuteIfBound();
+		}
+		else if (MouseDragType == DRAG_SETTING_RANGE)
+		{
+			// Zooming
+			FFrameTime MouseDownStart = ComputeFrameTimeFromMouse(MyGeometry, MouseDownPosition[0], RangeToScreen);
+
+			const bool bCanZoomIn = MouseTime > MouseDownStart;
+			const bool bCanZoomOut = ViewRangeStack.Num() > 0;
+			if (bCanZoomIn || bCanZoomOut)
+			{
+				TRange<double> ViewRange = TimeSliderArgs.ViewRange.Get();
+				if (!bCanZoomIn)
+				{
+					ViewRange = ViewRangeStack.Pop();
+				}
+
+				if (bCanZoomIn)
+				{
+					// push the current value onto the stack
+					ViewRangeStack.Add(ViewRange);
+
+					ViewRange = TRange<double>(MouseDownStart.FrameNumber / GetTickResolution(), MouseTime.FrameNumber / GetTickResolution());
+				}
+
+				TimeSliderArgs.OnViewRangeChanged.ExecuteIfBound(ViewRange, EViewRangeInterpolation::Immediate);
+				if (!TimeSliderArgs.ViewRange.IsBound())
+				{
+					// The output is not bound to a delegate so we'll manage the value ourselves
+					TimeSliderArgs.ViewRange.Set(ViewRange);
+				}
+			}
+		}
+		else
+		{
+			TimeSliderArgs.OnEndScrubberMovement.ExecuteIfBound();
+
+			CommitScrubPosition(MouseTime, /*bIsScrubbing=*/false);
+		}
+
+		MouseDragType = DRAG_NONE;
+		DistanceDragged = 0.f;
+
+		return FReply::Handled().ReleaseMouseCapture();
+	}
+
+	return FReply::Unhandled();
+}
+
+FReply FDataflowTimeSliderController::OnMouseButtonDoubleClick(SWidget& OwnerWidget, const FGeometry& MyGeometry, const FPointerEvent& MouseEvent)
+{
+	return FReply::Unhandled();
+}
+
+FReply FDataflowTimeSliderController::OnMouseMove(SWidget& WidgetOwner, const FGeometry& MyGeometry, const FPointerEvent& MouseEvent)
+{
+	const bool bHandleLeftMouseButton = MouseEvent.IsMouseButtonDown(EKeys::LeftMouseButton);
+	const bool bHandleRightMouseButton = MouseEvent.IsMouseButtonDown(EKeys::RightMouseButton) && TimeSliderArgs.AllowZoom;
+
+	if (bHandleRightMouseButton)
+	{
+		if (!bPanning)
+		{
+			DistanceDragged += static_cast<float>(FMath::Abs(MouseEvent.GetCursorDelta().X));
+			if (DistanceDragged > FSlateApplication::Get().GetDragTriggerDistance())
+			{
+				bPanning = true;
+			}
+		}
+		else
+		{
+			const TRange<double> LocalViewRange = TimeSliderArgs.ViewRange.Get();
+			const double LocalViewRangeMin = LocalViewRange.GetLowerBoundValue();
+			const double LocalViewRangeMax = LocalViewRange.GetUpperBoundValue();
+
+			const FScrubRangeToScreen ScaleInfo(LocalViewRange, MyGeometry.Size);
+			const FVector2D ScreenDelta = MouseEvent.GetCursorDelta();
+			FVector2D InputDelta;
+			InputDelta.X = ScreenDelta.X / ScaleInfo.PixelsPerInput;
+
+			double NewViewOutputMin = LocalViewRangeMin - InputDelta.X;
+			double NewViewOutputMax = LocalViewRangeMax - InputDelta.X;
+
+			ClampViewRange(NewViewOutputMin, NewViewOutputMax);
+			SetViewRange(NewViewOutputMin, NewViewOutputMax, EViewRangeInterpolation::Immediate);
+		}
+	}
+	else if (bHandleLeftMouseButton)
+	{
+		const TRange<double> LocalViewRange = TimeSliderArgs.ViewRange.Get();
+		const FScrubRangeToScreen RangeToScreen(LocalViewRange, MyGeometry.Size);
+		DistanceDragged += static_cast<float>(FMath::Abs(MouseEvent.GetCursorDelta().X));
+
+		if (MouseDragType == DRAG_NONE)
+		{
+			if (DistanceDragged > FSlateApplication::Get().GetDragTriggerDistance())
+			{
+				const FFrameTime MouseDownFree = ComputeFrameTimeFromMouse(MyGeometry, MouseDownPosition[0], RangeToScreen, false);
+
+				const FFrameRate FrameResolution = GetTickResolution();
+				const bool       bLockedPlayRange = TimeSliderArgs.IsPlaybackRangeLocked.Get();
+				const float      MouseDownPixel = RangeToScreen.InputToLocalX(MouseDownFree / FrameResolution);
+				const bool       bHitScrubber = GetHitTestScrubberPixelRange(TimeSliderArgs.ScrubPosition.Get(), RangeToScreen).HandleRange.Contains(MouseDownPixel);
+				
+				const TRange<double>   SelectionRange = TimeSliderArgs.SelectionRange.Get() / FrameResolution;
+				const TRange<double>   PlaybackRange = TimeSliderArgs.PlaybackRange.Get() / FrameResolution;
+
+				// Disable selection range test if it's empty so that the playback range scrubbing gets priority
+				if (!SelectionRange.IsEmpty() && !bHitScrubber && HitTestRangeEnd(RangeToScreen, SelectionRange, MouseDownPixel))
+				{
+					// selection range end scrubber
+					MouseDragType = DRAG_SELECTION_END;
+					TimeSliderArgs.OnSelectionRangeBeginDrag.ExecuteIfBound();
+				}
+				else if (!SelectionRange.IsEmpty() && !bHitScrubber && HitTestRangeStart(RangeToScreen, SelectionRange, MouseDownPixel))
+				{
+					// selection range start scrubber
+					MouseDragType = DRAG_SELECTION_START;
+					TimeSliderArgs.OnSelectionRangeBeginDrag.ExecuteIfBound();
+				}
+				else if (!bLockedPlayRange && !bHitScrubber && HitTestRangeEnd(RangeToScreen, PlaybackRange, MouseDownPixel))
+				{
+					// playback range end scrubber
+					MouseDragType = DRAG_PLAYBACK_END;
+					TimeSliderArgs.OnPlaybackRangeBeginDrag.ExecuteIfBound();
+				}
+				else if (!bLockedPlayRange && !bHitScrubber && HitTestRangeStart(RangeToScreen, PlaybackRange, MouseDownPixel))
+				{
+					// playback range start scrubber
+					MouseDragType = DRAG_PLAYBACK_START;
+					TimeSliderArgs.OnPlaybackRangeBeginDrag.ExecuteIfBound();
+				}
+				else if (FSlateApplication::Get().GetModifierKeys().AreModifersDown(EModifierKey::Control))
+				{
+					MouseDragType = DRAG_SETTING_RANGE;
+				}
+				else
+				{
+					MouseDragType = DRAG_SCRUBBING_TIME;
+					TimeSliderArgs.OnBeginScrubberMovement.ExecuteIfBound();
+				}
+			}
+		}
+		else
+		{
+			FFrameTime MouseTime = ComputeFrameTimeFromMouse(MyGeometry, MouseEvent.GetScreenSpacePosition(), RangeToScreen);
+
+			// Set the start range time?
+			if (MouseDragType == DRAG_PLAYBACK_START)
+			{
+				SetPlaybackRangeStart(MouseTime.FrameNumber);
+			}
+			// Set the end range time?
+			else if (MouseDragType == DRAG_PLAYBACK_END)
+			{
+				SetPlaybackRangeEnd(MouseTime.FrameNumber - 1);
+			}
+			else if (MouseDragType == DRAG_SELECTION_START)
+			{
+				SetSelectionRangeStart(MouseTime.FrameNumber);
+			}
+			// Set the end range time?
+			else if (MouseDragType == DRAG_SELECTION_END)
+			{
+				SetSelectionRangeEnd(MouseTime.FrameNumber);
+			}
+			else if (MouseDragType == DRAG_SCRUBBING_TIME)
+			{
+				// Delegate responsibility for clamping to the current viewrange to the client
+				CommitScrubPosition(MouseTime, /*bIsScrubbing=*/true);
+			}
+			else if (MouseDragType == DRAG_SETTING_RANGE)
+			{
+				MouseDownPosition[1] = MouseEvent.GetScreenSpacePosition();
+			}
+			else if (MouseDragType == DRAG_TIME)
+			{
+				return FReply::Unhandled();
+			}
+		}
+	}
+
+	if (DistanceDragged != 0.f && (bHandleLeftMouseButton || bHandleRightMouseButton))
+	{
+		return FReply::Handled().CaptureMouse(WidgetOwner.AsShared());
+	}
+
+	return FReply::Handled();
+
+}
+
+FReply FDataflowTimeSliderController::OnMouseWheel(SWidget& WidgetOwner, const FGeometry& MyGeometry, const FPointerEvent& MouseEvent)
+{
+	TOptional<TRange<float>> NewTargetRange;
+
+	if (TimeSliderArgs.AllowZoom && MouseEvent.IsControlDown())
+	{
+		const float MouseFractionX = static_cast<float>(MyGeometry.AbsoluteToLocal(MouseEvent.GetScreenSpacePosition()).X / MyGeometry.GetLocalSize().X);
+
+		const float ZoomDelta = -0.2f * MouseEvent.GetWheelDelta();
+		if (ZoomByDelta(ZoomDelta, MouseFractionX))
+		{
+			return FReply::Handled();
+		}
+	}
+	else if (MouseEvent.IsShiftDown())
+	{
+		PanByDelta(-MouseEvent.GetWheelDelta());
+		return FReply::Handled();
+	}
+
+	return FReply::Unhandled();
+}
+
+FCursorReply FDataflowTimeSliderController::OnCursorQuery(TSharedRef<const SWidget> WidgetOwner, const FGeometry& MyGeometry, const FPointerEvent& CursorEvent) const
+{
+	const FScrubRangeToScreen RangeToScreen(TimeSliderArgs.ViewRange.Get(), MyGeometry.Size);
+
+	const FFrameRate FrameResolution = GetTickResolution();
+	const bool       bLockedPlayRange = TimeSliderArgs.IsPlaybackRangeLocked.Get();
+	const float      HitTestPixel = static_cast<float>(MyGeometry.AbsoluteToLocal(CursorEvent.GetScreenSpacePosition()).X);
+	const bool       bHitScrubber = GetHitTestScrubberPixelRange(TimeSliderArgs.ScrubPosition.Get(), RangeToScreen).HandleRange.Contains(HitTestPixel);
+	
+	const TRange<double>   SelectionRange = TimeSliderArgs.SelectionRange.Get() / FrameResolution;
+	const TRange<double>   PlaybackRange = TimeSliderArgs.PlaybackRange.Get() / FrameResolution;
+
+	if (MouseDragType == DRAG_SCRUBBING_TIME)
+	{
+		return FCursorReply::Unhandled();
+	}
+
+	// Use L/R resize cursor if we're dragging or hovering a playback range bound
+	if ((MouseDragType == DRAG_PLAYBACK_END) ||
+		(MouseDragType == DRAG_PLAYBACK_START) ||
+		(MouseDragType == DRAG_SELECTION_START) ||
+		(MouseDragType == DRAG_SELECTION_END) ||
+		(!bLockedPlayRange && !bHitScrubber && HitTestRangeStart(RangeToScreen, PlaybackRange, HitTestPixel)) ||
+		(!bLockedPlayRange && !bHitScrubber && HitTestRangeEnd(RangeToScreen, PlaybackRange, HitTestPixel)) ||
+		(!SelectionRange.IsEmpty() && !bHitScrubber && HitTestRangeStart(RangeToScreen, SelectionRange, HitTestPixel)) ||
+		(!SelectionRange.IsEmpty() && !bHitScrubber && HitTestRangeEnd(RangeToScreen, SelectionRange, HitTestPixel)))
+	{
+		return FCursorReply::Cursor(EMouseCursor::ResizeLeftRight);
+	}
+
+	return FCursorReply::Unhandled();
+}
+
+void FDataflowTimeSliderController::SetViewRange(double NewRangeMin, double NewRangeMax, EViewRangeInterpolation Interpolation)
+{
+	// Clamp to a minimum size to avoid zero-sized or negative visible ranges
+	double MinVisibleTimeRange = FFrameNumber(1) / GetTickResolution();
+	TRange<double> ExistingViewRange = TimeSliderArgs.ViewRange.Get();
+
+	if (NewRangeMax == ExistingViewRange.GetUpperBoundValue())
+	{
+		if (NewRangeMin > NewRangeMax - MinVisibleTimeRange)
+		{
+			NewRangeMin = NewRangeMax - MinVisibleTimeRange;
+		}
+	}
+	else if (NewRangeMax < NewRangeMin + MinVisibleTimeRange)
+	{
+		NewRangeMax = NewRangeMin + MinVisibleTimeRange;
+	}
+
+	// Clamp to the clamp range
+	const TRange<double> NewRange = TRange<double>(NewRangeMin, NewRangeMax);
+	TimeSliderArgs.OnViewRangeChanged.ExecuteIfBound(NewRange, Interpolation);
+
+	if (!TimeSliderArgs.ViewRange.IsBound())
+	{
+		// The  output is not bound to a delegate so we'll manage the value ourselves (no animation)
+		TimeSliderArgs.ViewRange.Set(NewRange);
+	}
+}
+
+void FDataflowTimeSliderController::SetClampRange(double NewRangeMin, double NewRangeMax)
+{
+	const TRange<double> NewRange(NewRangeMin, NewRangeMax);
+
+	TimeSliderArgs.OnClampRangeChanged.ExecuteIfBound(NewRange);
+
+	if (!TimeSliderArgs.ClampRange.IsBound())
+	{
+		// The  output is not bound to a delegate so we'll manage the value ourselves (no animation)
+		TimeSliderArgs.ClampRange.Set(NewRange);
+	}
+}
+
+void FDataflowTimeSliderController::SetPlayRange(FFrameNumber RangeStart, int32 RangeDuration)
+{
+	check(RangeDuration >= 0);
+
+	const TRange<FFrameNumber> NewRange(RangeStart, RangeStart + RangeDuration);
+
+	TimeSliderArgs.OnPlaybackRangeChanged.ExecuteIfBound(NewRange);
+
+	if (!TimeSliderArgs.PlaybackRange.IsBound())
+	{
+		// The  output is not bound to a delegate so we'll manage the value ourselves (no animation)
+		TimeSliderArgs.PlaybackRange.Set(NewRange);
+	}
+}
+
+void FDataflowTimeSliderController::ClampViewRange(double& NewRangeMin, double& NewRangeMax)
+{
+	bool bNeedsClampSet = false;
+	double NewClampRangeMin = TimeSliderArgs.ClampRange.Get().GetLowerBoundValue();
+	if (NewRangeMin < TimeSliderArgs.ClampRange.Get().GetLowerBoundValue())
+	{
+		NewClampRangeMin = NewRangeMin;
+		bNeedsClampSet = true;
+	}
+
+	double NewClampRangeMax = TimeSliderArgs.ClampRange.Get().GetUpperBoundValue();
+	if (NewRangeMax > TimeSliderArgs.ClampRange.Get().GetUpperBoundValue())
+	{
+		NewClampRangeMax = NewRangeMax;
+		bNeedsClampSet = true;
+	}
+
+	if (bNeedsClampSet)
+	{
+		SetClampRange(NewClampRangeMin, NewClampRangeMax);
+	}
+}
+
+bool FDataflowTimeSliderController::ZoomByDelta(float InDelta, float MousePositionFraction)
+{
+	TRange<double> LocalViewRange = TimeSliderArgs.ViewRange.Get().GetAnimationTarget();
+	double LocalViewRangeMax = LocalViewRange.GetUpperBoundValue();
+	double LocalViewRangeMin = LocalViewRange.GetLowerBoundValue();
+	const double OutputViewSize = LocalViewRangeMax - LocalViewRangeMin;
+	const double OutputChange = OutputViewSize * InDelta;
+
+	double NewViewOutputMin = LocalViewRangeMin - (OutputChange * MousePositionFraction);
+	double NewViewOutputMax = LocalViewRangeMax + (OutputChange * (1.f - MousePositionFraction));
+
+	if (NewViewOutputMin < NewViewOutputMax)
+	{
+		ClampViewRange(NewViewOutputMin, NewViewOutputMax);
+		SetViewRange(NewViewOutputMin, NewViewOutputMax, EViewRangeInterpolation::Animated);
+		return true;
+	}
+
+	return false;
+}
+
+void FDataflowTimeSliderController::PanByDelta(float InDelta)
+{
+	const TRange<double> LocalViewRange = TimeSliderArgs.ViewRange.Get().GetAnimationTarget();
+
+	const double CurrentMin = LocalViewRange.GetLowerBoundValue();
+	const double CurrentMax = LocalViewRange.GetUpperBoundValue();
+
+	// Adjust the delta to be a percentage of the current range
+	InDelta *= ScrubConstants::ScrollPanFraction * static_cast<float>(CurrentMax - CurrentMin);
+
+	double NewViewOutputMin = CurrentMin + InDelta;
+	double NewViewOutputMax = CurrentMax + InDelta;
+
+	ClampViewRange(NewViewOutputMin, NewViewOutputMax);
+	SetViewRange(NewViewOutputMin, NewViewOutputMax, EViewRangeInterpolation::Animated);
+}
+
+FFrameTime FDataflowTimeSliderController::GetFrameTimeFromMouse(const FGeometry& Geometry, FVector2D ScreenSpacePosition) const
+{
+	FScrubRangeToScreen ScrubRangeToScreen(TimeSliderArgs.ViewRange.Get(), Geometry.Size);
+	return ComputeFrameTimeFromMouse(Geometry, ScreenSpacePosition, ScrubRangeToScreen);
+}
+
+void FDataflowTimeSliderController::CommitScrubPosition(FFrameTime NewValue, bool bIsScrubbing)
+{
+	// Manage the scrub position ourselves if its not bound to a delegate
+	if (!TimeSliderArgs.ScrubPosition.IsBound())
+	{
+		TimeSliderArgs.ScrubPosition.Set(NewValue);
+	}
+
+	TimeSliderArgs.OnScrubPositionChanged.ExecuteIfBound(NewValue, bIsScrubbing, /*bEvaluate*/ true); 
+}
+
+void FDataflowTimeSliderController::DrawTicks(FSlateWindowElementList& OutDrawElements, const TRange<double>& ViewRange, const FScrubRangeToScreen& RangeToScreen, FDrawTickArgs& InArgs) const
+{
+	TSharedPtr<SDataflowSimulationTimeline> Timeline = SimulationTimeline.Pin();
+	if (!Timeline.IsValid())
+	{
+		return;
+	}
+
+	if (!FMath::IsFinite(ViewRange.GetLowerBoundValue()) || !FMath::IsFinite(ViewRange.GetUpperBoundValue()))
+	{
+		return;
+	}
+
+	FFrameRate     FrameResolution = GetTickResolution();
+	FPaintGeometry PaintGeometry = InArgs.AllottedGeometry.ToPaintGeometry();
+	FSlateFontInfo SmallLayoutFont = FCoreStyle::GetDefaultFontStyle("Regular", 8);
+
+	double MajorGridStep = 0.0;
+	int32  MinorDivisions = 0;
+	if (!Timeline->GetGridMetrics(static_cast<float>(InArgs.AllottedGeometry.Size.X), MajorGridStep, MinorDivisions))
+	{
+		return;
+	}
+
+	if (InArgs.bOnlyDrawMajorTicks)
+	{
+		MinorDivisions = 0;
+	}
+
+	TArray<FVector2D> LinePoints;
+	LinePoints.SetNumUninitialized(2);
+
+	const bool bAntiAliasLines = false;
+
+	const double FirstMajorLine = FMath::FloorToDouble(ViewRange.GetLowerBoundValue() / MajorGridStep) * MajorGridStep;
+	const double LastMajorLine = FMath::CeilToDouble(ViewRange.GetUpperBoundValue() / MajorGridStep) * MajorGridStep;
+
+	for (double CurrentMajorLine = FirstMajorLine; CurrentMajorLine < LastMajorLine; CurrentMajorLine += MajorGridStep)
+	{
+		float MajorLinePx = RangeToScreen.InputToLocalX(CurrentMajorLine);
+
+		LinePoints[0] = FVector2D(MajorLinePx, InArgs.TickOffset);
+		LinePoints[1] = FVector2D(MajorLinePx, InArgs.TickOffset + InArgs.MajorTickHeight);
+
+		// Draw each tick mark
+		FSlateDrawElement::MakeLines(
+			OutDrawElements,
+			InArgs.StartLayer,
+			PaintGeometry,
+			LinePoints,
+			InArgs.DrawEffects,
+			InArgs.TickColor,
+			bAntiAliasLines
+		);
+
+		if (!InArgs.bOnlyDrawMajorTicks)
+		{
+			FString FrameString = TimeSliderArgs.NumericTypeInterface->ToString((CurrentMajorLine * FrameResolution).RoundToFrame().Value);
+
+			// Space the text between the tick mark but slightly above
+			FVector2D TextOffset(MajorLinePx + 5.f, InArgs.bMirrorLabels ? 3.f : FMath::Abs(InArgs.AllottedGeometry.Size.Y - (InArgs.MajorTickHeight + 3.f)));
+			FSlateDrawElement::MakeText(
+				OutDrawElements,
+				InArgs.StartLayer + 1,
+				InArgs.AllottedGeometry.ToPaintGeometry(InArgs.AllottedGeometry.Size, FSlateLayoutTransform(TextOffset)),
+				FrameString,
+				SmallLayoutFont,
+				InArgs.DrawEffects,
+				InArgs.TickColor * 0.65f
+			);
+		}
+
+		for (int32 Step = 1; Step < MinorDivisions; ++Step)
+		{
+			// Compute the size of each tick mark.  If we are half way between to visible values display a slightly larger tick mark
+			const float MinorTickHeight = ((MinorDivisions % 2 == 0) && (Step % (MinorDivisions / 2)) == 0) ? 6.0f : 2.0f;
+			const float MinorLinePx = RangeToScreen.InputToLocalX(CurrentMajorLine + Step * MajorGridStep / MinorDivisions);
+
+			LinePoints[0] = FVector2D(MinorLinePx, InArgs.bMirrorLabels ? 0.0f : FMath::Abs(InArgs.AllottedGeometry.Size.Y - MinorTickHeight));
+			LinePoints[1] = FVector2D(MinorLinePx, LinePoints[0].Y + MinorTickHeight);
+
+			// Draw each sub mark
+			FSlateDrawElement::MakeLines(
+				OutDrawElements,
+				InArgs.StartLayer,
+				PaintGeometry,
+				LinePoints,
+				InArgs.DrawEffects,
+				InArgs.TickColor,
+				bAntiAliasLines
+			);
+		}
+	}
+}
+
+int32 FDataflowTimeSliderController::DrawSelectionRange(const FGeometry& AllottedGeometry, const FSlateRect& MyCullingRect, FSlateWindowElementList& OutDrawElements, int32 LayerId, const FScrubRangeToScreen& RangeToScreen, const FPaintPlaybackRangeArgs& Args) const
+{
+	TRange<double> SelectionRange = TimeSliderArgs.SelectionRange.Get() / GetTickResolution();
+
+	if (!SelectionRange.IsEmpty() && SelectionRange.HasLowerBound() && SelectionRange.HasUpperBound())
+	{
+		const float SelectionRangeL = RangeToScreen.InputToLocalX(SelectionRange.GetLowerBoundValue()) - 1;
+		const float SelectionRangeR = RangeToScreen.InputToLocalX(SelectionRange.GetUpperBoundValue()) + 1;
+		const FLinearColor DrawColor = FAppStyle::GetSlateColor("SelectionColor").GetColor(FWidgetStyle());
+
+		if (Args.SolidFillOpacity > 0.f)
+		{
+			FSlateDrawElement::MakeBox(
+				OutDrawElements,
+				LayerId + 1,
+				AllottedGeometry.ToPaintGeometry(FVector2f(SelectionRangeR - SelectionRangeL, AllottedGeometry.Size.Y), FSlateLayoutTransform(FVector2D(SelectionRangeL, 0.f))),
+				FAppStyle::GetBrush("WhiteBrush"),
+				ESlateDrawEffect::None,
+				DrawColor.CopyWithNewOpacity(Args.SolidFillOpacity)
+			);
+		}
+
+		FSlateDrawElement::MakeBox(
+			OutDrawElements,
+			LayerId + 1,
+			AllottedGeometry.ToPaintGeometry(FVector2f(Args.BrushWidth, AllottedGeometry.Size.Y), FSlateLayoutTransform(FVector2D(SelectionRangeL, 0.f))),
+			Args.StartBrush,
+			ESlateDrawEffect::None,
+			DrawColor
+		);
+
+		FSlateDrawElement::MakeBox(
+			OutDrawElements,
+			LayerId + 1,
+			AllottedGeometry.ToPaintGeometry(FVector2f(Args.BrushWidth, AllottedGeometry.Size.Y), FSlateLayoutTransform(FVector2D(SelectionRangeR - Args.BrushWidth, 0.f))),
+			Args.EndBrush,
+			ESlateDrawEffect::None,
+			DrawColor
+		);
+	}
+
+	return LayerId + 1;
+}
+
+int32 FDataflowTimeSliderController::DrawPlaybackRange(const FGeometry& AllottedGeometry, const FSlateRect& MyCullingRect, FSlateWindowElementList& OutDrawElements, int32 LayerId, const FScrubRangeToScreen& RangeToScreen, const FPaintPlaybackRangeArgs& Args) const
+{
+	if (!TimeSliderArgs.PlaybackRange.IsSet())
+	{
+		return LayerId;
+	}
+
+	const uint8 OpacityBlend = TimeSliderArgs.SubSequenceRange.Get().IsSet() ? 128 : 255;
+
+	TRange<FFrameNumber> PlaybackRange = TimeSliderArgs.PlaybackRange.Get();
+	FFrameRate TickResolution = GetTickResolution();
+	const float PlaybackRangeL = RangeToScreen.InputToLocalX(PlaybackRange.GetLowerBoundValue() / TickResolution);
+	const float PlaybackRangeR = RangeToScreen.InputToLocalX(PlaybackRange.GetUpperBoundValue() / TickResolution) - 1;
+
+	FSlateDrawElement::MakeBox(
+		OutDrawElements,
+		LayerId + 1,
+		AllottedGeometry.ToPaintGeometry(FVector2f(Args.BrushWidth, AllottedGeometry.Size.Y), FSlateLayoutTransform(FVector2D(PlaybackRangeL, 0.f))),
+		Args.StartBrush,
+		ESlateDrawEffect::None,
+		FColor(32, 128, 32, OpacityBlend)	// 120, 75, 50 (HSV)
+	);
+
+	FSlateDrawElement::MakeBox(
+		OutDrawElements,
+		LayerId + 1,
+		AllottedGeometry.ToPaintGeometry(FVector2f(Args.BrushWidth, AllottedGeometry.Size.Y), FSlateLayoutTransform(FVector2D(PlaybackRangeR - Args.BrushWidth, 0.f))),
+		Args.EndBrush,
+		ESlateDrawEffect::None,
+		FColor(128, 32, 32, OpacityBlend)	// 0, 75, 50 (HSV)
+	);
+
+	// Black tint for excluded regions
+	FSlateDrawElement::MakeBox(
+		OutDrawElements,
+		LayerId + 1,
+		AllottedGeometry.ToPaintGeometry(FVector2f(PlaybackRangeL, AllottedGeometry.Size.Y), FSlateLayoutTransform(FVector2D(0.f, 0.f))),
+		FAppStyle::GetBrush("WhiteBrush"),
+		ESlateDrawEffect::None,
+		FLinearColor::Black.CopyWithNewOpacity(0.3f * OpacityBlend / 255.f)
+	);
+
+	FSlateDrawElement::MakeBox(
+		OutDrawElements,
+		LayerId + 1,
+		AllottedGeometry.ToPaintGeometry(FVector2f(AllottedGeometry.Size.X - PlaybackRangeR, AllottedGeometry.Size.Y), FSlateLayoutTransform(FVector2D(PlaybackRangeR, 0.f))),
+		FAppStyle::GetBrush("WhiteBrush"),
+		ESlateDrawEffect::None,
+		FLinearColor::Black.CopyWithNewOpacity(0.3f * OpacityBlend / 255.f)
+	);
+
+	return LayerId + 1;
+}
+
+bool FDataflowTimeSliderController::HitTestRangeStart(const FScrubRangeToScreen& RangeToScreen, const TRange<double>& Range, float HitPixel) const
+{
+	if (Range.HasLowerBound())
+	{
+		static float BrushSizeInStateUnits = 6.f, DragToleranceSlateUnits = 2.f, MouseTolerance = 2.f;
+		const float  RangeStartPixel = RangeToScreen.InputToLocalX(Range.GetLowerBoundValue());
+
+		// Hit test against the brush region to the right of the playback start position, +/- DragToleranceSlateUnits
+		return HitPixel >= RangeStartPixel - MouseTolerance - DragToleranceSlateUnits &&
+			HitPixel <= RangeStartPixel + MouseTolerance + BrushSizeInStateUnits + DragToleranceSlateUnits;
+	}
+
+	return false;
+}
+
+bool FDataflowTimeSliderController::HitTestRangeEnd(const FScrubRangeToScreen& RangeToScreen, const TRange<double>& Range, float HitPixel) const
+{
+	if (Range.HasUpperBound())
+	{
+		static float BrushSizeInStateUnits = 6.f, DragToleranceSlateUnits = 2.f, MouseTolerance = 2.f;
+		const float  RangeEndPixel = RangeToScreen.InputToLocalX(Range.GetUpperBoundValue());
+
+		// Hit test against the brush region to the left of the playback end position, +/- DragToleranceSlateUnits
+		return HitPixel >= RangeEndPixel - MouseTolerance - BrushSizeInStateUnits - DragToleranceSlateUnits &&
+			HitPixel <= RangeEndPixel + MouseTolerance + DragToleranceSlateUnits;
+	}
+
+	return false;
+}
+
+void FDataflowTimeSliderController::SetPlaybackRangeStart(FFrameNumber NewStart)
+{
+	TRange<FFrameNumber> PlaybackRange = TimeSliderArgs.PlaybackRange.Get();
+
+	if (NewStart <= UE::MovieScene::DiscreteExclusiveUpper(PlaybackRange))
+	{
+		TimeSliderArgs.OnPlaybackRangeChanged.ExecuteIfBound(TRange<FFrameNumber>(NewStart, PlaybackRange.GetUpperBound()));
+	}
+}
+
+void FDataflowTimeSliderController::SetPlaybackRangeEnd(FFrameNumber NewEnd)
+{
+	TRange<FFrameNumber> PlaybackRange = TimeSliderArgs.PlaybackRange.Get();
+
+	if (NewEnd >= UE::MovieScene::DiscreteInclusiveLower(PlaybackRange))
+	{
+		TimeSliderArgs.OnPlaybackRangeChanged.ExecuteIfBound(TRange<FFrameNumber>(PlaybackRange.GetLowerBound(), NewEnd));
+	}
+}
+
+void FDataflowTimeSliderController::SetSelectionRangeStart(FFrameNumber NewStart)
+{
+	TRange<FFrameNumber> SelectionRange = TimeSliderArgs.SelectionRange.Get();
+
+	if (SelectionRange.IsEmpty())
+	{
+		TimeSliderArgs.OnSelectionRangeChanged.ExecuteIfBound(TRange<FFrameNumber>(NewStart, NewStart + 1));
+	}
+	else if (NewStart <= UE::MovieScene::DiscreteExclusiveUpper(SelectionRange))
+	{
+		TimeSliderArgs.OnSelectionRangeChanged.ExecuteIfBound(TRange<FFrameNumber>(NewStart, SelectionRange.GetUpperBound()));
+	}
+}
+
+void FDataflowTimeSliderController::SetSelectionRangeEnd(FFrameNumber NewEnd)
+{
+	TRange<FFrameNumber> SelectionRange = TimeSliderArgs.SelectionRange.Get();
+
+	if (SelectionRange.IsEmpty())
+	{
+		TimeSliderArgs.OnSelectionRangeChanged.ExecuteIfBound(TRange<FFrameNumber>(NewEnd - 1, NewEnd));
+	}
+	else if (NewEnd >= UE::MovieScene::DiscreteInclusiveLower(SelectionRange))
+	{
+		TimeSliderArgs.OnSelectionRangeChanged.ExecuteIfBound(TRange<FFrameNumber>(SelectionRange.GetLowerBound(), NewEnd));
+	}
+}
+
+FFrameTime FDataflowTimeSliderController::ComputeFrameTimeFromMouse(const FGeometry& Geometry, FVector2D ScreenSpacePosition, FScrubRangeToScreen RangeToScreen, bool CheckSnapping) const
+{
+	const FVector2D CursorPos = Geometry.AbsoluteToLocal(ScreenSpacePosition);
+	const double    MouseValue = RangeToScreen.LocalXToInput(static_cast<float>(CursorPos.X));
+
+	return MouseValue * GetTickResolution();
+}
+
+FDataflowTimeSliderController::FScrubPixelRange FDataflowTimeSliderController::GetHitTestScrubberPixelRange(FFrameTime ScrubTime, const FScrubRangeToScreen& RangeToScreen) const
+{
+	constexpr float DragToleranceSlateUnits = 2.f, MouseTolerance = 2.f;
+	return GetScrubberPixelRange(ScrubTime, GetTickResolution(), GetDisplayRate(), RangeToScreen, DragToleranceSlateUnits + MouseTolerance);
+}
+
+FDataflowTimeSliderController::FScrubPixelRange FDataflowTimeSliderController::GetScrubberPixelRange(FFrameTime ScrubTime, const FScrubRangeToScreen& RangeToScreen) const
+{
+	return GetScrubberPixelRange(ScrubTime, GetTickResolution(), GetDisplayRate(), RangeToScreen);
+}
+
+FDataflowTimeSliderController::FScrubPixelRange FDataflowTimeSliderController::GetScrubberPixelRange(FFrameTime ScrubTime, FFrameRate Resolution, FFrameRate PlayRate, const FScrubRangeToScreen& RangeToScreen, float DilationPixels) const
+{
+	const FFrameNumber Frame = ScrubTime.FloorToFrame();// ConvertFrameTime(ScrubTime, Resolution, PlayRate).FloorToFrame();
+
+	float StartPixel = RangeToScreen.InputToLocalX(Frame / Resolution);
+	float EndPixel = RangeToScreen.InputToLocalX((Frame + 1) / Resolution);
+
+	{
+		const float RoundedStartPixel = static_cast<float>(FMath::RoundToInt(StartPixel));
+		EndPixel -= (StartPixel - RoundedStartPixel);
+
+		StartPixel = RoundedStartPixel;
+		EndPixel = FMath::Max(EndPixel, StartPixel + 1);
+	}
+
+	FScrubPixelRange Range;
+
+	float MinScrubSize = 14.f;
+	Range.bClamped = EndPixel - StartPixel < MinScrubSize;
+	Range.Range = TRange<float>(StartPixel, EndPixel);
+	if (Range.bClamped)
+	{
+		Range.HandleRange = TRange<float>(
+			(StartPixel + EndPixel - MinScrubSize) * .5f,
+			(StartPixel + EndPixel + MinScrubSize) * .5f);
+	}
+	else
+	{
+		Range.HandleRange = Range.Range;
+	}
+
+	return Range;
+}
\ No newline at end of file
diff --git a/Engine/Plugins/Experimental/Dataflow/Source/DataflowEditor/Private/Widgets/DataflowTimeSliderController.h b/Engine/Plugins/Experimental/Dataflow/Source/DataflowEditor/Private/Widgets/DataflowTimeSliderController.h
new file mode 100644
--- /dev/null
+++ b/Engine/Plugins/Experimental/Dataflow/Source/DataflowEditor/Private/Widgets/DataflowTimeSliderController.h
@@ -0,0 +1,188 @@
+// Copyright Epic Games, Inc. All Rights Reserved.
+
+#pragma once
+
+#include "CoreMinimal.h"
+#include "Widgets/SWidget.h"
+#include "ITimeSlider.h"
+#include "TimeSliderArgs.h"
+#include "Input/CursorReply.h"
+#include "Input/Reply.h"
+
+struct FSlateBrush;
+class SDataflowSimulationTimeline;
+
+class FDataflowTimeSliderController : public ITimeSliderController, public TSharedFromThis<FDataflowTimeSliderController>
+{
+	struct FDrawTickArgs;
+	struct FScrubRangeToScreen;
+
+public:
+	FDataflowTimeSliderController(const FTimeSliderArgs& InArgs, TWeakPtr<SDataflowSimulationTimeline> DataflowSimulationTimeline);
+	~FDataflowTimeSliderController();
+
+	//~ Begin ITimeSliderController Interface 
+	virtual int32 OnPaintTimeSlider(bool bMirrorLabels, const FGeometry& AllottedGeometry, const FSlateRect& MyCullingRect, FSlateWindowElementList& OutDrawElements, int32 LayerId, const FWidgetStyle& InWidgetStyle, bool bParentEnabled) const override;
+	virtual int32 OnPaintViewArea(const FGeometry& AllottedGeometry, const FSlateRect& MyCullingRect, FSlateWindowElementList& OutDrawElements, int32 LayerId, bool bEnabled, const FPaintViewAreaArgs& Args) const override;
+	virtual FReply OnMouseButtonDown(SWidget& WidgetOwner, const FGeometry& MyGeometry, const FPointerEvent& MouseEvent) override;
+	virtual FReply OnMouseButtonUp(SWidget& WidgetOwner, const FGeometry& MyGeometry, const FPointerEvent& MouseEvent) override;
+	virtual FReply OnMouseButtonDoubleClick(SWidget& OwnerWidget, const FGeometry& MyGeometry, const FPointerEvent& MouseEvent) override;
+	virtual FReply OnMouseMove(SWidget& WidgetOwner, const FGeometry& MyGeometry, const FPointerEvent& MouseEvent) override;
+	virtual FReply OnMouseWheel(SWidget& WidgetOwner, const FGeometry& MyGeometry, const FPointerEvent& MouseEvent) override;
+	virtual FCursorReply OnCursorQuery(TSharedRef<const SWidget> WidgetOwner, const FGeometry& MyGeometry, const FPointerEvent& CursorEvent) const override;
+	virtual void SetViewRange(double NewRangeMin, double NewRangeMax, EViewRangeInterpolation Interpolation) override;
+	virtual void SetClampRange(double NewRangeMin, double NewRangeMax) override;
+	virtual void SetPlayRange(FFrameNumber RangeStart, int32 RangeDuration) override;
+	virtual FFrameRate GetDisplayRate() const override { return TimeSliderArgs.DisplayRate.Get(); }
+	virtual FFrameRate GetTickResolution() const override { return TimeSliderArgs.TickResolution.Get(); }
+	virtual FAnimatedRange GetViewRange() const override { return TimeSliderArgs.ViewRange.Get(); }
+	virtual FAnimatedRange GetClampRange() const override { return TimeSliderArgs.ClampRange.Get(); }
+	virtual TRange<FFrameNumber> GetPlayRange() const override { return TimeSliderArgs.PlaybackRange.Get(TRange<FFrameNumber>()); }
+	virtual FFrameTime GetScrubPosition() const { return TimeSliderArgs.ScrubPosition.Get(); }
+	virtual void SetScrubPosition(FFrameTime InTime, bool bEvaluate) { TimeSliderArgs.OnScrubPositionChanged.ExecuteIfBound(InTime, false, bEvaluate); }
+	//~ End ITimeSliderController Interface 
+
+	/**
+	 * Clamp the given range to the clamp range
+	 *
+	 * @param NewRangeMin		The new lower bound of the range
+	 * @param NewRangeMax		The new upper bound of the range
+	 */
+	void ClampViewRange(double& NewRangeMin, double& NewRangeMax);
+
+	/**
+	 * Zoom the range by a given delta.
+	 *
+	 * @param InDelta					The total amount to zoom by (+ve = zoom out, -ve = zoom in)
+	 * @param MousePositionFraction		Bias to apply to lower/upper extents of the range. (0 = lower, 0.5 = equal, 1 = upper)
+	 */
+	bool ZoomByDelta(float InDelta, float MousePositionFraction = 0.5f);
+
+	/**
+	 * Pan the range by a given delta
+	 *
+	 * @param InDelta		The total amount to pan by (+ve = pan forwards in time, -ve = pan backwards in time)
+	 */
+	void PanByDelta(float InDelta);
+
+	/** Determine frame time from a mouse position */
+	FFrameTime GetFrameTimeFromMouse(const FGeometry& Geometry, FVector2D ScreenSpacePosition) const;
+
+private:
+	/**
+	 * Call this method when the user's interaction has changed the scrub position
+	 *
+	 * @param NewValue				Value resulting from the user's interaction
+	 * @param bIsScrubbing			True if done via scrubbing, false if just releasing scrubbing
+	 */
+	void CommitScrubPosition(FFrameTime NewValue, bool bIsScrubbing);
+
+	/**
+	 * Draw time tick marks
+	 *
+	 * @param OutDrawElements	List to add draw elements to
+	 * @param ViewRange			The currently visible time range in seconds
+	 * @param RangeToScreen		Time range to screen space converter
+	 * @param InArgs			Parameters for drawing the tick lines
+	 */
+	void DrawTicks(FSlateWindowElementList& OutDrawElements, const TRange<double>& ViewRange, const FScrubRangeToScreen& RangeToScreen, FDrawTickArgs& InArgs) const;
+
+	/**
+	 * Draw the selection range.
+	 *
+	 * @return The new layer ID.
+	 */
+	int32 DrawSelectionRange(const FGeometry& AllottedGeometry, const FSlateRect& MyCullingRect, FSlateWindowElementList& OutDrawElements, int32 LayerId, const FScrubRangeToScreen& RangeToScreen, const FPaintPlaybackRangeArgs& Args) const;
+
+	/**
+	 * Draw the playback range.
+	 *
+	 * @return the new layer ID
+	 */
+	int32 DrawPlaybackRange(const FGeometry& AllottedGeometry, const FSlateRect& MyCullingRect, FSlateWindowElementList& OutDrawElements, int32 LayerId, const FScrubRangeToScreen& RangeToScreen, const FPaintPlaybackRangeArgs& Args) const;
+
+	/** Hit test the lower bound of a range */
+	bool HitTestRangeStart(const FScrubRangeToScreen& RangeToScreen, const TRange<double>& Range, float HitPixel) const;
+
+	/** Hit test the upper bound of a range */
+	bool HitTestRangeEnd(const FScrubRangeToScreen& RangeToScreen, const TRange<double>& Range, float HitPixel) const;
+
+	/** Set the playback range start */
+	void SetPlaybackRangeStart(FFrameNumber NewStart);
+
+	/** Set the playback range end */
+	void SetPlaybackRangeEnd(FFrameNumber NewEnd);
+
+	/** Set the selection range start */
+	void SetSelectionRangeStart(FFrameNumber NewStart);
+
+	/** Set the selection range end */
+	void SetSelectionRangeEnd(FFrameNumber NewEnd);
+
+	FFrameTime ComputeFrameTimeFromMouse(const FGeometry& Geometry, FVector2D ScreenSpacePosition, FScrubRangeToScreen RangeToScreen, bool CheckSnapping = true) const;
+
+private:
+	struct FScrubPixelRange
+	{
+		TRange<float> Range;
+		TRange<float> HandleRange;
+		bool bClamped;
+	};
+
+	FScrubPixelRange GetHitTestScrubberPixelRange(FFrameTime ScrubTime, const FScrubRangeToScreen& RangeToScreen) const;
+	FScrubPixelRange GetScrubberPixelRange(FFrameTime ScrubTime, const FScrubRangeToScreen& RangeToScreen) const;
+	FScrubPixelRange GetScrubberPixelRange(FFrameTime ScrubTime, FFrameRate Resolution, FFrameRate PlayRate, const FScrubRangeToScreen& RangeToScreen, float DilationPixels = 0.f) const;
+
+	/** Time Slider args */
+	FTimeSliderArgs TimeSliderArgs;
+
+	/** Pointer back to the timeline */
+	TWeakPtr<SDataflowSimulationTimeline> SimulationTimeline;
+
+	/** Brush for drawing the fill area on the scrubber */
+	const FSlateBrush* ScrubFillBrush;
+
+	/** Brush for drawing an upwards facing scrub handles */
+	const FSlateBrush* ScrubHandleUpBrush;
+
+	/** Brush for drawing a downwards facing scrub handle */
+	const FSlateBrush* ScrubHandleDownBrush;
+
+	/** Brush for drawing an editable time */
+	const FSlateBrush* EditableTimeBrush;
+
+	// Mouse Controls
+	/** Total mouse delta during dragging **/
+	float DistanceDragged;
+
+	/** If we are dragging a scrubber or dragging to set the time range */
+	enum DragType
+	{
+		DRAG_SCRUBBING_TIME,
+		DRAG_SETTING_RANGE,
+		DRAG_PLAYBACK_START,
+		DRAG_PLAYBACK_END,
+		DRAG_SELECTION_START,
+		DRAG_SELECTION_END,
+		DRAG_TIME,
+		DRAG_NONE
+	};
+
+	DragType MouseDragType;
+
+	/** If we are currently panning the panel */
+	bool bPanning;
+
+	/** Index of the current dragged time */
+	int32 DraggedTimeIndex;
+
+	/** Mouse down position range */
+	FVector2D MouseDownPosition[2];
+
+	/** Geometry on mouse down */
+	FGeometry MouseDownGeometry;
+
+	/** Range stack */
+	TArray<TRange<double>> ViewRangeStack;
+
+};
\ No newline at end of file
diff --git a/Engine/Plugins/Experimental/Dataflow/Source/DataflowEditor/Private/Widgets/SDataflowSimulationTimeline.cpp b/Engine/Plugins/Experimental/Dataflow/Source/DataflowEditor/Private/Widgets/SDataflowSimulationTimeline.cpp
new file mode 100644
--- /dev/null
+++ b/Engine/Plugins/Experimental/Dataflow/Source/DataflowEditor/Private/Widgets/SDataflowSimulationTimeline.cpp
@@ -0,0 +1,467 @@
+// Copyright Epic Games, Inc. All Rights Reserved.
+
+#include "SDataflowSimulationTimeline.h"
+
+#include "Fonts/FontMeasure.h"
+#include "FrameNumberNumericInterface.h"
+#include "FrameNumberTimeEvaluator.h"
+#include "Framework/Application/SlateApplication.h"
+#include "ISequencerWidgetsModule.h"
+#include "Modules/ModuleManager.h"
+#include "Widgets/Layout/SGridPanel.h"
+#include "Widgets/Layout/SScrollBorder.h"
+#include "Widgets/Layout/SSpacer.h"
+#include "Widgets/SBoxPanel.h"
+#include "SDataflowTimelineOverlay.h"
+#include "SDataflowSplitterOverlay.h"
+#include "SDataflowTransportControl.h"
+#include "Widgets/Text/STextBlock.h"
+#include "Widgets/Views/SListView.h"
+#include "Widgets/Input/SSpinBox.h"
+#include "Widgets/Input/SSearchBox.h"
+
+#include "DataflowSimulationBinding.h"
+#include "DataflowTimeSliderController.h"
+
+#define LOCTEXT_NAMESPACE "SDataflowSimulationTimeline"
+
+void SDataflowSimulationTimeline::Construct(const FArguments& InArgs, const TSharedRef<FDataflowSimulationBinding>& InSimulationBinding)
+{
+	TWeakPtr<FDataflowSimulationBinding>  WeakSimulationBinding = InSimulationBinding;
+	SimulationBinding = InSimulationBinding;
+	
+	const EFrameNumberDisplayFormats DisplayFormat = EFrameNumberDisplayFormats::Frames;
+	const TAttribute<FFrameRate> TickResolution = MakeAttributeLambda([WeakSimulationBinding]()
+		{
+			return FFrameRate(WeakSimulationBinding.Pin()->GetTickResolution(), 1);
+		});
+	
+	ViewRange = MakeAttributeLambda([WeakSimulationBinding]() { return WeakSimulationBinding.IsValid() ? WeakSimulationBinding.Pin()->GetViewRange() : FAnimatedRange(0.0, 0.0); });
+	
+	const TAttribute<FFrameRate> DisplayRate = MakeAttributeLambda([WeakSimulationBinding]()
+		{
+			return WeakSimulationBinding.Pin()->GetFrameRate();
+		});
+	
+	ColumnFillCoefficients[0] = 0.3f;
+	ColumnFillCoefficients[1] = 0.7f;
+	
+	TAttribute<float> FillCoefficient_0, FillCoefficient_1;
+	{
+		FillCoefficient_0.Bind(TAttribute<float>::FGetter::CreateSP(this, &SDataflowSimulationTimeline::GetColumnFillCoefficient, 0));
+		FillCoefficient_1.Bind(TAttribute<float>::FGetter::CreateSP(this, &SDataflowSimulationTimeline::GetColumnFillCoefficient, 1));
+	};
+	
+	NumericTypeInterface = MakeShareable(new FFrameNumberInterface(DisplayFormat, 0, TickResolution, DisplayRate));
+	
+	FTimeSliderArgs TimeSliderArgs;
+	{
+		TimeSliderArgs.ScrubPosition = MakeAttributeLambda([WeakSimulationBinding]() { return WeakSimulationBinding.IsValid() ? WeakSimulationBinding.Pin()->GetScrubPosition() : FFrameTime(0); });
+		TimeSliderArgs.ViewRange = ViewRange;
+		TimeSliderArgs.PlaybackRange = MakeAttributeLambda([WeakSimulationBinding]() { return WeakSimulationBinding.IsValid() ? WeakSimulationBinding.Pin()->GetPlaybackRange() : TRange<FFrameNumber>(0, 0); });
+		TimeSliderArgs.ClampRange = MakeAttributeLambda([WeakSimulationBinding]() { return WeakSimulationBinding.IsValid() ? WeakSimulationBinding.Pin()->GetWorkingRange() : FAnimatedRange(0.0, 0.0); });
+		TimeSliderArgs.DisplayRate = DisplayRate;
+		TimeSliderArgs.TickResolution = TickResolution;
+		TimeSliderArgs.OnViewRangeChanged = FOnViewRangeChanged::CreateSP(&InSimulationBinding.Get(), &FDataflowSimulationBinding::HandleViewRangeChanged);
+		TimeSliderArgs.OnClampRangeChanged = FOnTimeRangeChanged::CreateSP(&InSimulationBinding.Get(), &FDataflowSimulationBinding::HandleWorkingRangeChanged);
+		TimeSliderArgs.IsPlaybackRangeLocked = true;
+		TimeSliderArgs.PlaybackStatus = EMovieScenePlayerStatus::Stopped;
+		TimeSliderArgs.NumericTypeInterface = NumericTypeInterface;
+		TimeSliderArgs.OnScrubPositionChanged = FOnScrubPositionChanged::CreateSP(this, &SDataflowSimulationTimeline::HandleScrubPositionChanged);
+	}
+	
+	TimeSliderController = MakeShareable(new FDataflowTimeSliderController(TimeSliderArgs, SharedThis(this)));
+	
+	TSharedRef<FDataflowTimeSliderController> TimeSliderControllerRef = TimeSliderController.ToSharedRef();
+	
+	// Create the top slider
+	constexpr bool bMirrorLabels = false;
+	ISequencerWidgetsModule& SequencerWidgets = FModuleManager::Get().LoadModuleChecked<ISequencerWidgetsModule>("SequencerWidgets");
+	TimeSlider = SequencerWidgets.CreateTimeSlider(TimeSliderControllerRef, bMirrorLabels);
+	
+	// Create bottom time range slider
+	TSharedRef<ITimeSlider> BottomTimeRange = SequencerWidgets.CreateTimeRange(
+		FTimeRangeArgs(
+			EShowRange::ViewRange | EShowRange::WorkingRange | EShowRange::PlaybackRange,
+			EShowRange::ViewRange | EShowRange::WorkingRange,
+			TimeSliderControllerRef,
+			EVisibility::Visible,
+			NumericTypeInterface.ToSharedRef()
+		),
+		SequencerWidgets.CreateTimeRangeSlider(TimeSliderControllerRef)
+	);
+	
+	TSharedRef<SScrollBar> ScrollBar = SNew(SScrollBar)
+		.Thickness(FVector2D(5.0f, 5.0f));
+	
+	InitTrackNames();
+	
+	TracksListView = SNew(SListView< TSharedPtr<FString> >)
+		.ExternalScrollbar(ScrollBar)
+		.ListItemsSource(&TrackNames)
+		.SelectionMode(ESelectionMode::Single)
+		.OnGenerateRow(this, &SDataflowSimulationTimeline::HandleTimelineListViewGenerateRow);
+	
+	// Grid Panel Constants
+	static constexpr int32 Column0 = 0, Column1 = 1;
+	static constexpr int32 Row0 = 0, Row1 = 1, Row2 = 2;
+	
+	const FMargin ResizeBarPadding(4.f, 0, 0, 0);
+	
+	ChildSlot
+		[
+			SNew(SOverlay)
+				+ SOverlay::Slot()
+				[
+					SNew(SGridPanel)
+						.FillRow(1, 1.0f)
+						.FillColumn(0, FillCoefficient_0)
+						.FillColumn(1, FillCoefficient_1)
+						+ SGridPanel::Slot(Column1, Row0, SGridPanel::Layer(10))
+						.Padding(ResizeBarPadding)
+						[
+							SNew(SBorder)
+								.BorderImage(FAppStyle::GetBrush("ToolPanel.GroupBorder"))
+								.BorderBackgroundColor(FLinearColor(.50f, .50f, .50f, 1.0f))
+								.Padding(0.f)
+								.Clipping(EWidgetClipping::ClipToBounds)
+								[
+									TimeSlider.ToSharedRef()
+								]
+						]
+						+ SGridPanel::Slot(Column1, Row1, SGridPanel::Layer(10))
+						.Padding(ResizeBarPadding)
+						[
+							SNew(SDataflowTimelineOverlay, TimeSliderControllerRef)
+								.Visibility(EVisibility::HitTestInvisible)
+								.DisplayScrubPosition(false)
+								.DisplayTickLines(true)
+								.Clipping(EWidgetClipping::ClipToBounds)
+								.PaintPlaybackRangeArgs(FPaintPlaybackRangeArgs(FAppStyle::GetBrush("Sequencer.Timeline.PlayRange_L"), FAppStyle::GetBrush("Sequencer.Timeline.PlayRange_R"), 6.f))
+						]
+	
+						// Overlay that draws the scrub position
+						+ SGridPanel::Slot(Column1, Row1, SGridPanel::Layer(20))
+						.Padding(ResizeBarPadding)
+						[
+							SNew(SDataflowTimelineOverlay, TimeSliderControllerRef)
+								.Visibility(EVisibility::HitTestInvisible)
+								.DisplayScrubPosition(true)
+								.DisplayTickLines(false)
+								.Clipping(EWidgetClipping::ClipToBounds)
+						]
+						+ SGridPanel::Slot(Column1, Row2, SGridPanel::Layer(10))
+						.Padding(ResizeBarPadding)
+						[
+							SNew(SBorder)
+								.BorderImage(FAppStyle::GetBrush("ToolPanel.GroupBorder"))
+								.BorderBackgroundColor(FLinearColor(0.5f, 0.5f, 0.5f, 1.0f))
+								.Clipping(EWidgetClipping::ClipToBounds)
+								.Padding(0)
+								[
+									BottomTimeRange
+								]
+						]
+						+SGridPanel::Slot(Column0, Row0, SGridPanel::Layer(10))
+						[
+							SNew(SHorizontalBox)
+							+SHorizontalBox::Slot()
+							.FillWidth(1.0f)
+							.VAlign(VAlign_Center)
+							[
+								SNew(SSearchBox)
+								.HintText(LOCTEXT("FilterTracksHint", "Filter"))
+								.OnTextChanged_Lambda([this](const FText& FilterText)
+								{
+									const FString& FilterString = FilterText.ToString();
+									FilteredTracks.Empty();
+
+									for (TSharedPtr<FString> TrackName : TrackNames)
+									{
+										if (FilterString.IsEmpty() || TrackName->Contains(FilterString))
+										{
+											FilteredTracks.Add(TrackName);
+										}
+									}
+									TracksListView->RequestListRefresh();
+								})
+							]
+							+SHorizontalBox::Slot()
+							.VAlign(VAlign_Center)
+							.HAlign(HAlign_Center)
+							.AutoWidth()
+							.Padding(2.0f, 0.0f, 2.0f, 0.0f)
+							[
+								SNew(SBox)
+								.MinDesiredWidth(30.0f)
+								.VAlign(VAlign_Center)
+								.HAlign(HAlign_Center)
+								[
+									// Current play time.
+									SNew(SSpinBox<double>)
+									.Style(&FAppStyle::GetWidgetStyle<FSpinBoxStyle>("Sequencer.PlayTimeSpinBox"))
+									.Value_Lambda([this]() -> double
+									{
+										if (SimulationBinding.IsValid())
+										{
+											return SimulationBinding.Pin()->GetScrubPosition().Value;
+										}
+										return 0.0;
+									})
+									.OnValueChanged_Lambda( [this](double FrameTime)
+									{
+										if (SimulationBinding.IsValid())
+										{
+											SimulationBinding.Pin()->SetScrubPosition(FFrameTime::FromDecimal(FrameTime));
+										}
+									})
+									.OnValueCommitted_Lambda([this](double FrameTime, ETextCommit::Type)
+									{
+										if (SimulationBinding.IsValid())
+										{
+											SimulationBinding.Pin()->SetScrubPosition(FFrameTime::FromDecimal(FrameTime));
+										}
+									})
+									.MinValue(TOptional<double>())
+									.MaxValue(TOptional<double>())
+									.TypeInterface(NumericTypeInterface)
+									.Delta_Lambda([this]() -> double
+										{
+											if (SimulationBinding.IsValid())
+											{
+												return FFrameRate(SimulationBinding.Pin()->GetTickResolution(), 1).AsDecimal() * SimulationBinding.Pin()->GetFrameRate().AsInterval();
+											}
+											return 0.0;
+										})
+									.LinearDeltaSensitivity(25)
+								]
+							]
+						]
+						+ SGridPanel::Slot(Column0, Row2, SGridPanel::Layer(10))
+						.VAlign(VAlign_Center)
+						.HAlign(HAlign_Center)
+						[
+							SNew(SDataflowTransportControl, InSimulationBinding)
+						]
+						+ SGridPanel::Slot(Column0, Row1, SGridPanel::Layer(5))
+						.ColumnSpan(2)
+						[
+							SNew(SHorizontalBox)
+								+ SHorizontalBox::Slot()
+								[
+									SNew(SOverlay)
+										+ SOverlay::Slot()
+										[
+											SNew(SVerticalBox)
+												+ SVerticalBox::Slot()
+												.FillHeight(1.f)
+												[
+													SNew(SScrollBorder, TracksListView.ToSharedRef())
+													[
+														SNew(SHorizontalBox)
+														+SHorizontalBox::Slot()
+														.FillWidth(FillCoefficient_0)
+														[
+															SNew(SBox)
+																[
+																	TracksListView.ToSharedRef()
+																]
+														]
+														+ SHorizontalBox::Slot()
+														.FillWidth(FillCoefficient_1)
+														[
+															SNew(SBox)
+															.Padding(ResizeBarPadding)
+															.Clipping(EWidgetClipping::ClipToBounds)	
+														]
+													]
+												]
+										]
+										+ SOverlay::Slot()
+										.HAlign(HAlign_Right)
+										[
+											ScrollBar
+										]
+								]
+						]
+				]
+				+ SOverlay::Slot()
+				[
+					// track area virtual splitter overlay
+					SNew(SDataflowSplitterOverlay)
+						.Style(FAppStyle::Get(), "AnimTimeline.Outliner.Splitter")
+						.Visibility(EVisibility::SelfHitTestInvisible)
+	
+						+ SSplitter::Slot()
+						.Value(FillCoefficient_0)
+						.OnSlotResized(SSplitter::FOnSlotResized::CreateSP(this, &SDataflowSimulationTimeline::OnColumnFillCoefficientChanged, 0))
+						[
+							SNew(SSpacer)
+						]
+	
+						+ SSplitter::Slot()
+						.Value(FillCoefficient_1)
+						.OnSlotResized(SSplitter::FOnSlotResized::CreateSP(this, &SDataflowSimulationTimeline::OnColumnFillCoefficientChanged, 1))
+						[
+							SNew(SSpacer)
+						]
+				]
+		];
+
+}
+
+void SDataflowSimulationTimeline::HandleScrubPositionChanged(FFrameTime NewScrubPosition, bool bIsScrubbing, bool bEvaluate) const
+{
+	if (const TSharedPtr<FDataflowSimulationBinding> SharedBinding = SimulationBinding.Pin())
+	{
+		SharedBinding->SetScrubPosition(NewScrubPosition);
+	}
+}
+
+void SDataflowSimulationTimeline::InitTrackNames()
+{
+	TrackNames.Empty();
+	if (const TSharedPtr<FDataflowSimulationBinding> SharedBinding = SimulationBinding.Pin())
+	{
+		TArray<FString> CacheNames;
+		SharedBinding->FillCacheNames(CacheNames);
+
+		for(const FString& CacheName : CacheNames)
+		{
+			TrackNames.Add(MakeShareable(new FString(CacheName)));
+		}
+	}
+}
+
+void SDataflowSimulationTimeline::OnColumnFillCoefficientChanged(float FillCoefficient, int32 ColumnIndex)
+{
+	ColumnFillCoefficients[ColumnIndex] = FillCoefficient;
+}
+
+TSharedRef<ITableRow> SDataflowSimulationTimeline::HandleTimelineListViewGenerateRow(TSharedPtr<FString> Text, const TSharedRef<STableViewBase>& OwnerTable)
+{
+	return SNew(STableRow<TSharedPtr<FString> >, OwnerTable)
+		[
+			SNew(SBox)
+				[
+					SNew(STextBlock)
+						.Margin(FMargin{10,5})
+						.Text(FText::FromString(*Text))
+				]
+		];
+}
+
+// FFrameRate::ComputeGridSpacing doesnt deal well with prime numbers, so we have a custom impl here
+static bool ComputeGridSpacing(const FFrameRate& InFrameRate, float PixelsPerSecond, double& OutMajorInterval, int32& OutMinorDivisions, float MinTickPx, float DesiredMajorTickPx)
+{
+	// First try built-in spacing
+	const bool bResult = InFrameRate.ComputeGridSpacing(PixelsPerSecond, OutMajorInterval, OutMinorDivisions, MinTickPx, DesiredMajorTickPx);
+	if (!bResult || OutMajorInterval == 1.0)
+	{
+		if (PixelsPerSecond <= 0.f)
+		{
+			return false;
+		}
+
+		const int32 RoundedFPS = static_cast<int32>(FMath::RoundToInt(InFrameRate.AsDecimal()));
+
+		if (RoundedFPS > 0)
+		{
+			// Showing frames
+			TArray<int32, TInlineAllocator<10>> CommonBases;
+
+			// Divide the rounded frame rate by 2s, 3s or 5s recursively
+			{
+				const int32 Denominators[] = { 2, 3, 5 };
+
+				int32 LowestBase = RoundedFPS;
+				for (;;)
+				{
+					CommonBases.Add(LowestBase);
+
+					if (LowestBase % 2 == 0) { LowestBase = LowestBase / 2; }
+					else if (LowestBase % 3 == 0) { LowestBase = LowestBase / 3; }
+					else if (LowestBase % 5 == 0) { LowestBase = LowestBase / 5; }
+					else
+					{
+						int32 LowestResult = LowestBase;
+						for (int32 Denominator : Denominators)
+						{
+							int32 Result = LowestBase / Denominator;
+							if (Result > 0 && Result < LowestResult)
+							{
+								LowestResult = Result;
+							}
+						}
+
+						if (LowestResult < LowestBase)
+						{
+							LowestBase = LowestResult;
+						}
+						else
+						{
+							break;
+						}
+					}
+				}
+			}
+
+			Algo::Reverse(CommonBases);
+
+			const int32 Scale = static_cast<int32>(FMath::CeilToInt(DesiredMajorTickPx / PixelsPerSecond * InFrameRate.AsDecimal()));
+			const int32 BaseIndex = FMath::Min(Algo::LowerBound(CommonBases, Scale), CommonBases.Num() - 1);
+			const int32 Base = CommonBases[BaseIndex];
+
+			const int32 MajorIntervalFrames = FMath::CeilToInt((float)Scale / static_cast<float>(Base)) * Base;
+			OutMajorInterval = MajorIntervalFrames * InFrameRate.AsInterval();
+
+			// Find the lowest number of divisions we can show that's larger than the minimum tick size
+			OutMinorDivisions = 0;
+			for (int32 DivIndex = 0; DivIndex < BaseIndex; ++DivIndex)
+			{
+				if (Base % CommonBases[DivIndex] == 0)
+				{
+					const int32 MinorDivisions = MajorIntervalFrames / CommonBases[DivIndex];
+					if (OutMajorInterval / MinorDivisions * PixelsPerSecond >= MinTickPx)
+					{
+						OutMinorDivisions = MinorDivisions;
+						break;
+					}
+				}
+			}
+		}
+	}
+
+	return OutMajorInterval != 0;
+}
+
+bool SDataflowSimulationTimeline::GetGridMetrics(float PhysicalWidth, double& OutMajorInterval, int32& OutMinorDivisions) const
+{
+	const FSlateFontInfo SmallLayoutFont = FCoreStyle::GetDefaultFontStyle("Regular", 8);
+	const TSharedRef<FSlateFontMeasure> FontMeasureService = FSlateApplication::Get().GetRenderer()->GetFontMeasureService();
+
+	const FFrameRate DisplayRate = SimulationBinding.Pin()->GetFrameRate();
+	const double BiggestTime = ViewRange.Get().GetUpperBoundValue();
+	const FString TickString = NumericTypeInterface->ToString((BiggestTime * DisplayRate).FrameNumber.Value);
+	const FVector2D MaxTextSize = FontMeasureService->Measure(TickString, SmallLayoutFont);
+
+	constexpr float MajorTickMultiplier = 2.f;
+
+	const float MinTickPx = static_cast<float>(MaxTextSize.X) + 5.f;
+	const float DesiredMajorTickPx = static_cast<float>(MaxTextSize.X) * MajorTickMultiplier;
+
+	if (PhysicalWidth > 0 && DisplayRate.AsDecimal() > 0)
+	{
+		return ComputeGridSpacing(
+			DisplayRate,
+			static_cast<float>(PhysicalWidth / ViewRange.Get().Size<double>()),
+			OutMajorInterval,
+			OutMinorDivisions,
+			MinTickPx,
+			DesiredMajorTickPx);
+	}
+
+	return false;
+}
+
+#undef LOCTEXT_NAMESPACE
\ No newline at end of file
diff --git a/Engine/Plugins/Experimental/Dataflow/Source/DataflowEditor/Private/Widgets/SDataflowSimulationTimeline.h b/Engine/Plugins/Experimental/Dataflow/Source/DataflowEditor/Private/Widgets/SDataflowSimulationTimeline.h
new file mode 100644
--- /dev/null
+++ b/Engine/Plugins/Experimental/Dataflow/Source/DataflowEditor/Private/Widgets/SDataflowSimulationTimeline.h
@@ -0,0 +1,75 @@
+﻿// Copyright Epic Games, Inc. All Rights Reserved.
+
+#pragma once
+
+#include "CoreMinimal.h"
+#include "ITimeSlider.h"
+#include "AnimatedRange.h"
+#include "Widgets/SCompoundWidget.h"
+#include "Widgets/Input/NumericTypeInterface.h"
+#include "Widgets/Views/STableViewBase.h"
+#include "Widgets/Views/STableRow.h"
+
+class FDataflowTimeSliderController;
+class FDataflowSimulationBinding;
+
+/** Main dataflow simulation timeline widget to be used in the dataflow toolkit */
+class SDataflowSimulationTimeline : public SCompoundWidget
+{
+public:
+	SLATE_BEGIN_ARGS(SDataflowSimulationTimeline) { }
+	SLATE_END_ARGS()
+
+	void Construct(const FArguments& InArgs, const TSharedRef<FDataflowSimulationBinding>& InSimulationBinbding);
+
+	/** Get numeric Type interface for converting between frame numbers and display formats. */
+	TSharedRef<INumericTypeInterface<double>> GetNumericTypeInterface() const { return NumericTypeInterface.ToSharedRef(); }
+
+	/** Compute a major grid interval and number of minor divisions to display */
+	bool GetGridMetrics(float PhysicalWidth, double& OutMajorInterval, int32& OutMinorDivisions) const;
+
+private:
+	/** Handle any changes to the scrub position */
+	void HandleScrubPositionChanged(FFrameTime NewScrubPosition, bool bIsScrubbing, bool bEvaluate) const;
+
+	/** Init track names from the simulation binding */
+	void InitTrackNames();
+
+	/** Get the column coefficient */
+	float GetColumnFillCoefficient(int32 ColumnIndex) const
+	{
+		return ColumnFillCoefficients[ColumnIndex];
+	}
+	/** callback when the column coefficient is changing */
+	void OnColumnFillCoefficientChanged(float FillCoefficient, int32 ColumnIndex);
+
+	/** Add a track row in the list viiew */
+	TSharedRef<ITableRow> HandleTimelineListViewGenerateRow(TSharedPtr<FString> Text, const TSharedRef<STableViewBase>& OwnerTable);
+
+	/** The fill coefficients of each column in the grid. */
+	float ColumnFillCoefficients[2];
+
+	/** Simulation binding to get the timing information from dataflow */
+	TWeakPtr<FDataflowSimulationBinding> SimulationBinding;
+
+	/** Interface to display the frame number */
+	TSharedPtr<INumericTypeInterface<double>> NumericTypeInterface;
+
+	/** Time slider controller used in the timeline */
+	TSharedPtr<FDataflowTimeSliderController> TimeSliderController;
+
+	/** Time Slider used in the time line */
+	TSharedPtr<ITimeSlider> TimeSlider;
+
+	/** List of track names holding the caches */
+	TArray<TSharedPtr<FString>> TrackNames;
+
+	/** List of filtered tracks holding the caches */
+	TArray<TSharedPtr<FString>> FilteredTracks;
+
+	/** Track list view in betwwen the timeline and the transport controls */
+	TSharedPtr<SListView<TSharedPtr<FString>>> TracksListView;
+
+	/** Animation view range*/
+	TAttribute<FAnimatedRange> ViewRange;
+};
\ No newline at end of file
diff --git a/Engine/Plugins/Experimental/Dataflow/Source/DataflowEditor/Private/Widgets/SDataflowSplitterOverlay.cpp b/Engine/Plugins/Experimental/Dataflow/Source/DataflowEditor/Private/Widgets/SDataflowSplitterOverlay.cpp
new file mode 100644
--- /dev/null
+++ b/Engine/Plugins/Experimental/Dataflow/Source/DataflowEditor/Private/Widgets/SDataflowSplitterOverlay.cpp
@@ -0,0 +1,131 @@
+﻿// Copyright Epic Games, Inc. All Rights Reserved.
+
+#include "SDataflowSplitterOverlay.h"
+
+#include "Delegates/Delegate.h"
+#include "HAL/PlatformCrt.h"
+#include "Layout/ArrangedChildren.h"
+#include "Layout/ArrangedWidget.h"
+#include "Layout/ChildrenBase.h"
+#include "Layout/Geometry.h"
+#include "Layout/Visibility.h"
+#include "Math/Vector2D.h"
+#include "Misc/Attribute.h"
+#include "SlotBase.h"
+#include "Types/SlateEnums.h"
+#include "Widgets/DeclarativeSyntaxSupport.h"
+#include "Widgets/Layout/SBox.h"
+
+struct FCaptureLostEvent;
+struct FPointerEvent;
+
+void SDataflowSplitterOverlay::Construct( const FArguments& InArgs )
+{
+	SetVisibility(EVisibility::SelfHitTestInvisible);
+
+	SplitterWidget = SArgumentNew(InArgs, SSplitter);
+	SplitterWidget->SetVisibility(EVisibility::HitTestInvisible);
+	AddSlot()
+	[
+		SplitterWidget.ToSharedRef()
+	];
+
+	for (int32 Index = 0; Index < SplitterWidget->GetChildren()->Num() - 1; ++Index)
+	{
+		AddSlot()
+		.Padding(TAttribute<FMargin>::Create(TAttribute<FMargin>::FGetter::CreateSP(this, &SDataflowSplitterOverlay::GetSplitterHandlePadding, Index)))
+		[
+			SNew(SBox)
+			.Visibility(EVisibility::Visible)
+		];
+	}
+}
+
+void SDataflowSplitterOverlay::OnArrangeChildren( const FGeometry& AllottedGeometry, FArrangedChildren& ArrangedChildren ) const
+{
+	FArrangedChildren SplitterChildren(ArrangedChildren.GetFilter());
+	SplitterWidget->ArrangeChildren(AllottedGeometry, SplitterChildren);
+
+	SlotPadding.Reset();
+
+	for (int32 Index = 0; Index < SplitterChildren.Num() - 1; ++Index)
+	{
+		const FGeometry& ThisGeometry = SplitterChildren[Index].Geometry;
+		const FGeometry& NextGeometry = SplitterChildren[Index + 1].Geometry;
+
+		if (SplitterWidget->GetOrientation() == EOrientation::Orient_Horizontal)
+		{
+			SlotPadding.Add(FMargin(
+				ThisGeometry.Position.X + static_cast<float>(ThisGeometry.GetLocalSize().X),
+				0,
+				static_cast<float>(AllottedGeometry.Size.X) - NextGeometry.Position.X,
+				0)
+			);
+		}
+		else
+		{
+			SlotPadding.Add(FMargin(
+				0,
+				ThisGeometry.Position.Y + static_cast<float>(ThisGeometry.GetLocalSize().Y),
+				0,
+				static_cast<float>(AllottedGeometry.Size.Y) - NextGeometry.Position.Y)
+			);
+		}
+	}
+
+	SOverlay::OnArrangeChildren(AllottedGeometry, ArrangedChildren);
+}
+
+FMargin SDataflowSplitterOverlay::GetSplitterHandlePadding(int32 Index) const
+{
+	if (SlotPadding.IsValidIndex(Index))
+	{
+		return SlotPadding[Index];
+	}
+
+	return 0.f;
+}
+	
+FCursorReply SDataflowSplitterOverlay::OnCursorQuery( const FGeometry& MyGeometry, const FPointerEvent& CursorEvent ) const
+{
+	return SplitterWidget->OnCursorQuery(MyGeometry, CursorEvent);
+}
+
+FReply SDataflowSplitterOverlay::OnMouseButtonDown( const FGeometry& MyGeometry, const FPointerEvent& MouseEvent )
+{
+	FReply Reply = SplitterWidget->OnMouseButtonDown(MyGeometry, MouseEvent);
+	if (Reply.GetMouseCaptor().IsValid())
+	{
+		// Set us to be the mouse captor so we can forward events properly
+		Reply.CaptureMouse( SharedThis(this) );
+		SetVisibility(EVisibility::Visible);
+	}
+	return Reply;
+}
+	
+void SDataflowSplitterOverlay::OnMouseCaptureLost(const FCaptureLostEvent& CaptureLostEvent)
+{
+	SetVisibility(EVisibility::SelfHitTestInvisible);
+	SOverlay::OnMouseCaptureLost(CaptureLostEvent);
+}
+
+FReply SDataflowSplitterOverlay::OnMouseButtonUp( const FGeometry& MyGeometry, const FPointerEvent& MouseEvent )
+{
+	FReply Reply = SplitterWidget->OnMouseButtonUp(MyGeometry, MouseEvent);
+	if (Reply.ShouldReleaseMouse())
+	{
+		SetVisibility(EVisibility::SelfHitTestInvisible);
+	}
+	return Reply;
+}
+
+FReply SDataflowSplitterOverlay::OnMouseMove( const FGeometry& MyGeometry, const FPointerEvent& MouseEvent )
+{
+	return SplitterWidget->OnMouseMove(MyGeometry, MouseEvent);
+}
+
+void SDataflowSplitterOverlay::OnMouseLeave( const FPointerEvent& MouseEvent )
+{
+	return SplitterWidget->OnMouseLeave(MouseEvent);
+}
+
diff --git a/Engine/Plugins/Experimental/Dataflow/Source/DataflowEditor/Private/Widgets/SDataflowSplitterOverlay.h b/Engine/Plugins/Experimental/Dataflow/Source/DataflowEditor/Private/Widgets/SDataflowSplitterOverlay.h
new file mode 100644
--- /dev/null
+++ b/Engine/Plugins/Experimental/Dataflow/Source/DataflowEditor/Private/Widgets/SDataflowSplitterOverlay.h
@@ -0,0 +1,46 @@
+﻿// Copyright Epic Games, Inc. All Rights Reserved.
+
+#pragma once
+
+#include "Containers/Array.h"
+#include "HAL/Platform.h"
+#include "Input/CursorReply.h"
+#include "Input/Reply.h"
+#include "Layout/Margin.h"
+#include "Templates/SharedPointer.h"
+#include "Widgets/Layout/SSplitter.h"
+#include "Widgets/SOverlay.h"
+
+class FArrangedChildren;
+struct FCaptureLostEvent;
+struct FGeometry;
+struct FPointerEvent;
+
+/** Splitter used on the dataflow timeline as an overlay. Input is disabled on all areas except the draggable positions */
+class SDataflowSplitterOverlay : public SOverlay
+{
+public:
+	typedef SSplitter::FArguments FArguments;
+
+	void Construct( const FArguments& InArgs );
+
+	//~ Begin SWidget interface 
+	virtual void OnArrangeChildren( const FGeometry& AllottedGeometry, FArrangedChildren& ArrangedChildren ) const override;
+	virtual FCursorReply OnCursorQuery( const FGeometry& MyGeometry, const FPointerEvent& CursorEvent ) const override;
+	virtual FReply OnMouseButtonDown( const FGeometry& MyGeometry, const FPointerEvent& MouseEvent ) override;
+	virtual void OnMouseCaptureLost(const FCaptureLostEvent& CaptureLostEvent) override;
+	virtual FReply OnMouseButtonUp( const FGeometry& MyGeometry, const FPointerEvent& MouseEvent ) override;
+	virtual FReply OnMouseMove( const FGeometry& MyGeometry, const FPointerEvent& MouseEvent ) override;
+	virtual void OnMouseLeave( const FPointerEvent& MouseEvent ) override;
+	//~ End SWidget interface 
+
+private:
+	/** Get the padding in between several tracks */
+	FMargin GetSplitterHandlePadding(int32 Index) const;
+
+	/** Splitter widget to be used in between the tracks*/
+	TSharedPtr<SSplitter> SplitterWidget;
+
+	/** Slot padding */
+	mutable TArray<FMargin> SlotPadding;
+};
diff --git a/Engine/Plugins/Experimental/Dataflow/Source/DataflowEditor/Private/Widgets/SDataflowTimelineOverlay.cpp b/Engine/Plugins/Experimental/Dataflow/Source/DataflowEditor/Private/Widgets/SDataflowTimelineOverlay.cpp
new file mode 100644
--- /dev/null
+++ b/Engine/Plugins/Experimental/Dataflow/Source/DataflowEditor/Private/Widgets/SDataflowTimelineOverlay.cpp
@@ -0,0 +1,25 @@
+﻿// Copyright Epic Games, Inc. All Rights Reserved.
+
+#include "SDataflowTimelineOverlay.h"
+#include "DataflowTimeSliderController.h"
+
+int32 SDataflowTimelineOverlay::OnPaint(const FPaintArgs& Args, const FGeometry& AllottedGeometry,
+	const FSlateRect& CullingRect, FSlateWindowElementList& DrawElements, int32 LayerId, const FWidgetStyle& WidgetStyle, bool bParentEnabled) const
+{
+	FPaintViewAreaArgs PaintArgs;
+	PaintArgs.bDisplayTickLines = bDisplayTickLines.Get();
+	PaintArgs.bDisplayScrubPosition = bDisplayScrubPosition.Get();
+
+	if (PaintPlaybackRangeArgs.IsSet())
+	{
+		PaintArgs.PlaybackRangeArgs = PaintPlaybackRangeArgs.Get();
+	}
+
+	if(TimeSliderController)
+	{
+		TimeSliderController->OnPaintViewArea(AllottedGeometry, CullingRect, DrawElements, LayerId, ShouldBeEnabled(bParentEnabled), PaintArgs);
+	}
+
+	return SCompoundWidget::OnPaint(Args, AllottedGeometry, CullingRect, DrawElements, LayerId, WidgetStyle, bParentEnabled);
+}
+
diff --git a/Engine/Plugins/Experimental/Dataflow/Source/DataflowEditor/Private/Widgets/SDataflowTimelineOverlay.h b/Engine/Plugins/Experimental/Dataflow/Source/DataflowEditor/Private/Widgets/SDataflowTimelineOverlay.h
new file mode 100644
--- /dev/null
+++ b/Engine/Plugins/Experimental/Dataflow/Source/DataflowEditor/Private/Widgets/SDataflowTimelineOverlay.h
@@ -0,0 +1,57 @@
+﻿// Copyright Epic Games, Inc. All Rights Reserved.
+
+#pragma once
+
+#include "HAL/Platform.h"
+#include "ITimeSlider.h"
+#include "Misc/Attribute.h"
+#include "Templates/SharedPointer.h"
+#include "Widgets/DeclarativeSyntaxSupport.h"
+#include "Widgets/SCompoundWidget.h"
+
+class FDataflowTimeSliderController;
+class FPaintArgs;
+class FSlateRect;
+class FSlateWindowElementList;
+class FWidgetStyle;
+struct FGeometry;
+
+/** An overlay that displays global information in the track area. */
+class SDataflowTimelineOverlay : public SCompoundWidget
+{
+public:
+	SLATE_BEGIN_ARGS(SDataflowTimelineOverlay)
+		: _DisplayTickLines(true)
+		, _DisplayScrubPosition(false)
+		{}
+
+		SLATE_ATTRIBUTE(bool, DisplayTickLines)
+		SLATE_ATTRIBUTE(bool, DisplayScrubPosition)
+		SLATE_ATTRIBUTE(FPaintPlaybackRangeArgs, PaintPlaybackRangeArgs)
+	SLATE_END_ARGS()
+
+	/** Construct the widget given args and a slider controller */
+	void Construct(const FArguments& InArgs, TSharedRef<FDataflowTimeSliderController> InTimeSliderController)
+	{
+		bDisplayScrubPosition = InArgs._DisplayScrubPosition;
+		bDisplayTickLines = InArgs._DisplayTickLines;
+		PaintPlaybackRangeArgs = InArgs._PaintPlaybackRangeArgs;
+		TimeSliderController = InTimeSliderController;
+	}
+
+private:
+	/** SWidget Interface */
+	virtual int32 OnPaint(const FPaintArgs& Args, const FGeometry& AllottedGeometry, const FSlateRect& MyCullingRect, FSlateWindowElementList& OutDrawElements, int32 LayerId, const FWidgetStyle& InWidgetStyle, bool bParentEnabled) const override;
+	
+	/** Controller for manipulating time */
+	TSharedPtr<FDataflowTimeSliderController> TimeSliderController;
+	
+	/** Whether or not to display the scrub position */
+	TAttribute<bool> bDisplayScrubPosition;
+	
+	/** Whether or not to display tick lines */
+	TAttribute<bool> bDisplayTickLines;
+	
+	/** User-supplied options for drawing playback range */
+	TAttribute<FPaintPlaybackRangeArgs> PaintPlaybackRangeArgs;
+};
\ No newline at end of file
diff --git a/Engine/Plugins/Experimental/Dataflow/Source/DataflowEditor/Private/Widgets/SDataflowTransportControl.cpp b/Engine/Plugins/Experimental/Dataflow/Source/DataflowEditor/Private/Widgets/SDataflowTransportControl.cpp
new file mode 100644
--- /dev/null
+++ b/Engine/Plugins/Experimental/Dataflow/Source/DataflowEditor/Private/Widgets/SDataflowTransportControl.cpp
@@ -0,0 +1,338 @@
+// Copyright Epic Games, Inc. All Rights Reserved.
+
+#include "SDataflowTransportControl.h"
+#include "DataflowSimulationBinding.h"
+#include "Dataflow/DataflowEditorStyle.h"
+#include "Delegates/Delegate.h"
+#include "Framework/SlateDelegates.h"
+#include "Misc/EnumRange.h"
+#include "SlateOptMacros.h"
+#include "Editor/EditorWidgets/Public/EditorWidgetsModule.h"
+#include "Editor/EditorWidgets/Public/ITransportControl.h"
+#include "Widgets/Images/SImage.h"
+#include "Widgets/Input/SButton.h"
+#include "Widgets/Input/SCheckBox.h"
+#include "Widgets/SBoxPanel.h"
+#include "Modules/ModuleManager.h"
+
+#define LOCTEXT_NAMESPACE "DataflowSimulationTransportControl"
+
+void SDataflowTransportControl::Construct(const FArguments& InArgs, const TSharedRef<FDataflowSimulationBinding>& InBindingAsset)
+{
+	SimulationBinding = InBindingAsset;
+	
+	// Skip adding the the Loop button so we can add our own
+	TArray<FTransportControlWidget> TransportControlWidgets;
+	for (const ETransportControlWidgetType Type : TEnumRange<ETransportControlWidgetType>())
+	{
+		if ((Type != ETransportControlWidgetType::Custom) && (Type != ETransportControlWidgetType::Loop))
+		{
+			TransportControlWidgets.Add(FTransportControlWidget(Type));
+		}
+	}
+	const FTransportControlWidget PreviewModeWidget(FOnMakeTransportWidget::CreateSP(this, &SDataflowTransportControl::OnCreateModeButton));
+	TransportControlWidgets.Add(PreviewModeWidget);
+
+	const FTransportControlWidget ResetSimulationWidget(FOnMakeTransportWidget::CreateSP(this, &SDataflowTransportControl::OnCreateResetButton));
+	TransportControlWidgets.Add(ResetSimulationWidget);
+
+	FTransportControlArgs TransportControlArgs;
+	TransportControlArgs.OnForwardPlay = FOnClicked::CreateSP(this, &SDataflowTransportControl::OnForwardPlay);
+	TransportControlArgs.OnBackwardPlay = FOnClicked::CreateSP(this, &SDataflowTransportControl::OnBackwardPlay);
+	TransportControlArgs.OnForwardStep = FOnClicked::CreateSP(this, &SDataflowTransportControl::OnForwardStep);
+	TransportControlArgs.OnBackwardStep = FOnClicked::CreateSP(this, &SDataflowTransportControl::OnBackwardStep);
+	TransportControlArgs.OnForwardEnd = FOnClicked::CreateSP(this, &SDataflowTransportControl::OnForwardEnd);
+	TransportControlArgs.OnBackwardEnd = FOnClicked::CreateSP(this, &SDataflowTransportControl::OnBackwardEnd);
+	TransportControlArgs.OnRecord = FOnClicked::CreateSP(this, &SDataflowTransportControl::OnRecord);
+	TransportControlArgs.OnTickPlayback = FOnTickPlayback::CreateSP(this, &SDataflowTransportControl::OnTickPlayback);
+	TransportControlArgs.OnGetPlaybackMode = FOnGetPlaybackMode::CreateSP(this, &SDataflowTransportControl::GetPlaybackMode);
+	TransportControlArgs.WidgetsToCreate = TransportControlWidgets;
+	
+	FEditorWidgetsModule& EditorWidgetsModule = FModuleManager::LoadModuleChecked<FEditorWidgetsModule>("EditorWidgets");
+
+	TSharedPtr<SHorizontalBox> HorizontalBox;
+	ChildSlot
+	[
+		SAssignNew(HorizontalBox, SHorizontalBox)
+		+ SHorizontalBox::Slot()
+		.HAlign(HAlign_Center)
+		.FillWidth(1.0)
+		[
+			EditorWidgetsModule.CreateTransportControl(TransportControlArgs)
+		]
+	];
+}
+
+FReply SDataflowTransportControl::OnForwardStep()
+{
+	if (const TSharedPtr<FDataflowSimulationBinding> SharedBinding = SimulationBinding.Pin())
+	{
+		SharedBinding->SetScrubTime(SharedBinding->GetScrubTime() + SharedBinding->GetDeltaTime());
+	}
+	return FReply::Handled();
+}
+
+FReply SDataflowTransportControl::OnForwardEnd()
+{
+	if (const TSharedPtr<FDataflowSimulationBinding> SharedBinding = SimulationBinding.Pin())
+	{
+		SharedBinding->SetScrubTime(SharedBinding->GetSequenceLength());
+	}
+	return FReply::Handled();
+}
+
+FReply SDataflowTransportControl::OnBackwardStep()
+{
+	if (const TSharedPtr<FDataflowSimulationBinding> SharedBinding = SimulationBinding.Pin())
+	{
+		SharedBinding->SetScrubTime(SharedBinding->GetScrubTime() - SharedBinding->GetDeltaTime());
+	}
+	return FReply::Handled();
+}
+
+FReply SDataflowTransportControl::OnBackwardEnd()
+{
+	if (const TSharedPtr<FDataflowSimulationBinding> SharedBinding = SimulationBinding.Pin())
+	{
+		SharedBinding->SetScrubTime(0);
+	}
+	return FReply::Handled();
+}
+
+FReply SDataflowTransportControl::OnForwardPlay()
+{
+	if (PlaybackMode == EPlaybackMode::PlayingForward)
+	{
+		PlaybackMode = EPlaybackMode::Stopped;
+	}
+	else
+	{
+		PlaybackMode = EPlaybackMode::PlayingForward;
+	}
+	return FReply::Handled();
+}
+
+FReply SDataflowTransportControl::OnBackwardPlay()
+{
+	if(PlaybackMode == EPlaybackMode::PlayingReverse)
+	{
+		PlaybackMode = EPlaybackMode::Stopped;
+	}
+	else
+	{
+		PlaybackMode = EPlaybackMode::PlayingReverse;
+	}
+
+	return FReply::Handled();
+}
+
+FReply SDataflowTransportControl::OnPlaybackMode()
+{
+	if (PreviewMode == EPreviewMode::Default)
+	{
+		PreviewMode = EPreviewMode::Looping;
+	}
+	else if (PreviewMode == EPreviewMode::Looping)
+	{
+		PreviewMode = EPreviewMode::PingPong;
+	}
+	else
+	{
+		PreviewMode = EPreviewMode::Default;
+	}
+
+	return FReply::Handled();
+}
+
+void SDataflowTransportControl::OnTickPlayback(double InCurrentTime, float InDeltaTime)
+{
+	if (const TSharedPtr<FDataflowSimulationBinding> SharedBinding = SimulationBinding.Pin())
+	{
+		const float SequenceLength = SharedBinding->GetSequenceLength();
+		const float ScrubTime = (PlaybackMode == EPlaybackMode::PlayingForward) ?
+			SharedBinding->GetScrubTime() + InDeltaTime : SharedBinding->GetScrubTime() - InDeltaTime;
+		
+		if(PreviewMode == EPreviewMode::Looping)
+		{
+			SharedBinding->SetScrubTime(ScrubTime - SequenceLength * FMath::Floor(ScrubTime / SequenceLength));
+		}
+		else
+		{
+			if(PreviewMode == EPreviewMode::PingPong)
+			{
+				if((PlaybackMode == EPlaybackMode::PlayingForward) && (ScrubTime >= SequenceLength))
+				{
+					PlaybackMode = EPlaybackMode::PlayingReverse;
+				}
+
+				if((PlaybackMode == EPlaybackMode::PlayingReverse) && (ScrubTime <= 0))
+				{
+					PlaybackMode = EPlaybackMode::PlayingForward;
+				}
+			}
+			SharedBinding->SetScrubTime(FMath::Clamp(ScrubTime, 0.f, SequenceLength));
+		}
+	}
+}
+
+FReply SDataflowTransportControl::OnRecord()
+{
+	if (const TSharedPtr<FDataflowSimulationBinding> SharedBinding = SimulationBinding.Pin())
+	{
+		SharedBinding->ResetSimulationScene();
+
+		SharedBinding->RecordSimulationCache();
+	}
+	return FReply::Handled();
+}
+
+TSharedRef<SWidget> SDataflowTransportControl::OnCreateResetButton()
+{
+	if (const TSharedPtr<FDataflowSimulationBinding> SharedBinding = SimulationBinding.Pin())
+	{
+		ResetSimulationButton = SNew(SButton)
+				.ButtonStyle(FAppStyle::Get(), "Animation.PlayControlsButton")
+				.ToolTipText(FText::FromString("The button reset the dataflow simulation."))
+				.OnClicked_Lambda([this, SharedBinding]()
+					{
+						SharedBinding->ResetSimulationScene();
+
+						return FReply::Handled();
+					})
+				.ContentPadding(0.0f)
+				.IsFocusable(true)
+				.HAlign(HAlign_Center)
+					[
+						SNew(SImage)
+							.ColorAndOpacity(FSlateColor::UseSubduedForeground())
+							.Image(FDataflowEditorStyle::Get().GetBrush("Dataflow.ResetSimulation"))
+					];
+	}
+
+	TSharedRef<SHorizontalBox> ResetSimulationWidget = SNew(SHorizontalBox);
+	ResetSimulationWidget->AddSlot()
+	.AutoWidth()
+	[
+		ResetSimulationButton.ToSharedRef()
+	];
+
+	return ResetSimulationWidget;
+}
+
+TSharedRef<SWidget> SDataflowTransportControl::OnCreateLockButton()
+{
+	if (const TSharedPtr<FDataflowSimulationBinding> SharedBinding = SimulationBinding.Pin())
+	{
+		LockSimulationBox = SNew(SCheckBox)
+				.Style(FAppStyle::Get(), "ToggleButtonCheckbox")
+				.ToolTipText(FText::FromString("The button locks the refresh of the values in the panel."))
+				.IsChecked_Lambda([this, SharedBinding]()-> ECheckBoxState
+					{
+						if (SharedBinding->IsSimulationLocked())
+						{
+							return ECheckBoxState::Checked;
+						}
+						return ECheckBoxState::Unchecked;
+					})
+				.OnCheckStateChanged_Lambda([this, SharedBinding](ECheckBoxState State)
+					{
+						//const bool bIsPlaying = PlaybackMode != EPlaybackMode::Stopped;
+						if(State == ECheckBoxState::Checked)
+						{
+							SharedBinding->SetSimulationLocked(true);
+						}
+						else
+						{
+							SharedBinding->SetSimulationLocked(false);
+						}
+					})
+				.Padding(0.0f)
+				.HAlign(HAlign_Center)
+				[
+					SNew(SImage)
+						.ColorAndOpacity(FSlateColor::UseForeground())
+						.Image(this, &SDataflowTransportControl::GetLockButtonImage)  
+				];
+	}
+	TSharedRef<SHorizontalBox> LockSimulationWidget = SNew(SHorizontalBox);
+	LockSimulationWidget->AddSlot()
+	.AutoWidth()
+	[
+		LockSimulationBox.ToSharedRef()
+	];
+
+	return LockSimulationWidget;
+}
+
+TSharedRef<SWidget> SDataflowTransportControl::OnCreateModeButton()
+{
+	PreviewModeButton = SNew(SButton)
+		.OnClicked(this, &SDataflowTransportControl::OnPlaybackMode)
+		.ButtonStyle( FAppStyle::Get(), "Animation.PlayControlsButton" )
+		.IsFocusable(false)
+		.ToolTipText_Lambda([&]()
+		{
+			if (PreviewMode == EPreviewMode::Default)
+			{
+				return LOCTEXT("PlaybackModeDefaultTooltip", "Linear playback");
+			}
+			else if (PreviewMode == EPreviewMode::Looping)
+			{
+				return LOCTEXT("PlaybackModeLoopingTooltip", "Looping playback");
+			}
+			else
+			{
+				return LOCTEXT("PlaybackModePingPongTooltip", "Ping pong playback");
+			}
+		})
+		.ContentPadding(0.0f);
+
+	TWeakPtr<SButton> WeakButton = PreviewModeButton;
+
+	PreviewModeButton->SetContent(SNew(SImage)
+		.ColorAndOpacity(FSlateColor::UseSubduedForeground())
+		.Image_Lambda([&, WeakButton]()
+		{
+			if (PreviewMode == EPreviewMode::Default)
+			{
+				return FAppStyle::Get().GetBrush("Animation.Loop.Disabled");
+			}
+			else if (PreviewMode == EPreviewMode::Looping)
+			{
+				return FAppStyle::Get().GetBrush("Animation.Loop.Enabled");
+			}
+			else
+			{
+				return FAppStyle::Get().GetBrush("Animation.Loop.SelectionRange");		// TODO: Replace with a back and forth type icon
+			}
+		})
+	);
+
+	TSharedRef<SHorizontalBox> PreviewPlaybackModeBox = SNew(SHorizontalBox);
+	PreviewPlaybackModeBox->AddSlot()
+	.AutoWidth()
+	[
+		PreviewModeButton.ToSharedRef()
+	];
+
+	return PreviewPlaybackModeBox;
+}
+
+const FSlateBrush* SDataflowTransportControl::GetLockButtonImage() const
+{
+	if (const TSharedPtr<FDataflowSimulationBinding> SharedBinding = SimulationBinding.Pin())
+	{
+		if (SharedBinding->IsSimulationLocked())
+		{
+			return FDataflowEditorStyle::Get().GetBrush("Dataflow.LockSimulation");
+		}
+		else
+		{
+			return FDataflowEditorStyle::Get().GetBrush("Dataflow.UnlockSimulation");
+		}
+	}
+	return nullptr;
+}
+
+#undef LOCTEXT_NAMESPACE
+
diff --git a/Engine/Plugins/Experimental/Dataflow/Source/DataflowEditor/Private/Widgets/SDataflowTransportControl.h b/Engine/Plugins/Experimental/Dataflow/Source/DataflowEditor/Private/Widgets/SDataflowTransportControl.h
new file mode 100644
--- /dev/null
+++ b/Engine/Plugins/Experimental/Dataflow/Source/DataflowEditor/Private/Widgets/SDataflowTransportControl.h
@@ -0,0 +1,79 @@
+// Copyright Epic Games, Inc. All Rights Reserved.
+
+#pragma once
+
+#include "Input/Reply.h"
+#include "ITransportControl.h"
+#include "Templates/SharedPointer.h"
+#include "TickableEditorObject.h"
+#include "Editor/EditorWidgets/Public/ITransportControl.h"
+#include "Widgets/DeclarativeSyntaxSupport.h"
+
+class FDataflowSimulationBinding;
+class SButton;
+class SCheckBox;
+
+/** Dataflow simulation widget to control an animation/simulation */
+class SDataflowTransportControl : public SCompoundWidget
+{
+public:
+	SLATE_BEGIN_ARGS(SDataflowTransportControl) {}
+	SLATE_END_ARGS()
+
+	/** Construct the widget given args and a simulation binding */
+	void Construct(const FArguments& InArgs, const TSharedRef<FDataflowSimulationBinding>& InSimulationBinding);
+
+private:
+
+	/** Data flow preview mode */
+	enum class EPreviewMode : uint8
+	{
+		Default,
+		Looping,
+		PingPong
+	};
+
+	/** Delegate when the simulation controls are pressed*/
+	FReply OnForwardPlay();
+	FReply OnBackwardPlay();
+	FReply OnForwardStep();
+	FReply OnBackwardStep();
+	FReply OnForwardEnd();
+	FReply OnBackwardEnd();
+	FReply OnPlaybackMode();
+	FReply OnRecord();
+	void OnTickPlayback(double InCurrentTime, float InDeltaTime);
+
+	/** Create preview mode widget */
+	TSharedRef<SWidget> OnCreateModeButton();
+
+	/** Create lock widget */
+	TSharedRef<SWidget> OnCreateLockButton();
+
+	/** Create reset widget */
+	TSharedRef<SWidget> OnCreateResetButton();
+
+	/** Get the lock button image */
+	const FSlateBrush* GetLockButtonImage() const;
+
+	/** Get the playback mode */
+	EPlaybackMode::Type GetPlaybackMode() const {return PlaybackMode;}
+	
+	/** Simulation binding to extract dataflow information */
+	TWeakPtr<FDataflowSimulationBinding> SimulationBinding;
+
+	/** Preview mode (looping...) */
+	EPreviewMode PreviewMode = EPreviewMode::Looping;
+
+	/** Playback mode */
+	EPlaybackMode::Type PlaybackMode = EPlaybackMode::Stopped;
+
+	/** Playback mode button */
+	TSharedPtr<SButton> PreviewModeButton;
+
+	/** Widget showing if a simulation is locked or not */
+	TSharedPtr<SCheckBox> LockSimulationBox;
+
+	/** Widget to reset a simulation */
+	TSharedPtr<SButton> ResetSimulationButton;
+};
\ No newline at end of file
diff --git a/Engine/Plugins/Experimental/Dataflow/Source/DataflowEditor/Public/Dataflow/DataflowEditorCommands.h b/Engine/Plugins/Experimental/Dataflow/Source/DataflowEditor/Public/Dataflow/DataflowEditorCommands.h
--- a/Engine/Plugins/Experimental/Dataflow/Source/DataflowEditor/Public/Dataflow/DataflowEditorCommands.h
+++ b/Engine/Plugins/Experimental/Dataflow/Source/DataflowEditor/Public/Dataflow/DataflowEditorCommands.h
@@ -59,6 +59,7 @@ public:
 	
 	TSharedPtr<FUICommandInfo> EvaluateNode;
 	TSharedPtr<FUICommandInfo> EvaluateGraph;
+	TSharedPtr<FUICommandInfo> ToggleSimulation;
 	TSharedPtr<FUICommandInfo> EvaluateGraphAutomatic;
 	TSharedPtr<FUICommandInfo> EvaluateGraphManual;
 	TSharedPtr<FUICommandInfo> ClearGraphCache;
@@ -81,6 +82,7 @@ public:
 	TSharedPtr<FUICommandInfo> CollapseToSubGraph;
 
 	TMap<FName, TSharedPtr<FUICommandInfo>> SetConstructionViewModeCommands;
+	TMap<FName, TSharedPtr<FUICommandInfo>> ToolPaletteCommands;
 
 	UE_DEPRECATED(5.5, "Dataflow Tool commands are now stored in FDataflowToolRegistry")
 	const static FString BeginWeightMapPaintToolIdentifier;
@@ -89,25 +91,6 @@ public:
 	
 	const static FString AddWeightMapNodeIdentifier;
 	TSharedPtr<FUICommandInfo> AddWeightMapNode;
-
-	const static FString RebuildSimulationSceneIdentifier;
-	TSharedPtr<FUICommandInfo> RebuildSimulationScene;
-	
-	const static FString PauseSimulationSceneIdentifier;
-	TSharedPtr<FUICommandInfo> PauseSimulationScene;
-
-	const static FString StartSimulationSceneIdentifier;
-	TSharedPtr<FUICommandInfo> StartSimulationScene;
-	
-	const static FString StepSimulationSceneIdentifier;
-	TSharedPtr<FUICommandInfo> StepSimulationScene;
-	
-	// @todo(brice) Remove Example Tools
-	//const static FString BeginAttributeEditorToolIdentifier;
-	//TSharedPtr<FUICommandInfo> BeginAttributeEditorTool;
-	//
-	//const static FString BeginMeshSelectionToolIdentifier;
-	//TSharedPtr<FUICommandInfo> BeginMeshSelectionTool;
 };
 
 //@todo(brice) Merge this into the above class
diff --git a/Engine/Plugins/Experimental/Dataflow/Source/DataflowEditor/Public/Dataflow/DataflowEditorModeToolkit.h b/Engine/Plugins/Experimental/Dataflow/Source/DataflowEditor/Public/Dataflow/DataflowEditorModeToolkit.h
--- a/Engine/Plugins/Experimental/Dataflow/Source/DataflowEditor/Public/Dataflow/DataflowEditorModeToolkit.h
+++ b/Engine/Plugins/Experimental/Dataflow/Source/DataflowEditor/Public/Dataflow/DataflowEditorModeToolkit.h
@@ -11,6 +11,22 @@ class SDataflowConstructionViewport;
 class SDataflowSimulationViewport;
 class SBaseCharacterFXEditorViewport;
 
+/** Tool pallete data for the dataflow mode toolkit */
+struct FDataflowEditorModeToolPalette
+{
+	/**  Mode area to display the warning messages */ 
+	TSharedPtr<STextBlock> ModeWarningArea;
+	
+	/**  Mode header for the tools names*/ 
+	TSharedPtr<STextBlock> ModeHeaderArea;
+
+	/** Too warning area*/
+	TSharedPtr<STextBlock> ToolWarningArea;
+
+	/** Handle when the active palette is changed */
+	FDelegateHandle ActivePaletteChangedHandle;
+};
+
 /**
  * The dataflow editor mode toolkit is responsible for the panel on the side in the dataflow editor
  * that shows mode and tool properties. Tool buttons would go in Init().
@@ -23,30 +39,55 @@ class FDataflowEditorModeToolkit : public FBaseCharacterFXEditorModeToolkit
 public:
 	UE_API void Init(const TSharedPtr<IToolkitHost>& InitToolkitHost, TWeakObjectPtr<UEdMode> InOwningMode);
 
-	// FBaseCharacterFXEditorModeToolkit interface
+	//~ Begin FBaseCharacterFXEditorModeToolkit interface
 	UE_API virtual const FSlateBrush* GetActiveToolIcon(const FString& Identifier) const override;
+	//~ End FBaseCharacterFXEditorModeToolkit interface
 
-	// FModeToolkit interface
+	//~ Begin FModeToolkit interface
 	UE_API virtual void OnToolStarted(UInteractiveToolManager* Manager, UInteractiveTool* Tool) override;
 	UE_API virtual void OnToolEnded(UInteractiveToolManager* Manager, UInteractiveTool* Tool) override;
+	//~ End FModeToolkit interface
 
-	// IToolkit interface
+	//~ Begin IToolkit interface
 	UE_API virtual FName GetToolkitFName() const override;
 	UE_API virtual FText GetBaseToolkitName() const override;
+	//~ End IToolkit interface
 
 	UE_API void SetConstructionViewportWidget(TWeakPtr<SDataflowConstructionViewport>);
 	UE_API void SetSimulationViewportWidget(TWeakPtr<SDataflowSimulationViewport>);
 
 private:
 
+	/** Register the toolkit builder palettes */
+	void RegisterPalettes();
+
+	/** Make the accept/cancel overlay widget */
+	void MakeToolAcceptCancelWidget();
+
+	/** Update the active tool properties */
+	void UpdateActiveToolProperties(UInteractiveTool* Tool);
+
+	/** Invalidate the cache state */
+	void InvalidateCachedDetailPanelState(UObject* ChangedObject);
+
+	/** Bind all the tools commands */
+	void BindCommands();
+
 	// Get the viewport widget associated with the given manager
 	// TODO: This should not be necessary any more as we do not run tools in the Simulation Viewport (JIRA UE-201248)
 	UE_API SBaseCharacterFXEditorViewport* GetViewportWidgetForManager(UInteractiveToolManager* Manager);
 
+	/** Current tool context used to create nodes */
 	UE_API UEditorInteractiveToolsContext* GetCurrentToolsContext();
 
+	/** Construction viewport */
 	TWeakPtr<SDataflowConstructionViewport> ConstructionViewportWidget;
+
+	/** Simulation viewport*/
 	TWeakPtr<SDataflowSimulationViewport> SimulationViewportWidget;
+
+	/** Dataflow toolkit mode tool palette data */
+	FDataflowEditorModeToolPalette ModeToolPallete;
 };
 
 #undef UE_API
diff --git a/Engine/Plugins/Experimental/Dataflow/Source/DataflowEditor/Public/Dataflow/DataflowEditorToolkit.h b/Engine/Plugins/Experimental/Dataflow/Source/DataflowEditor/Public/Dataflow/DataflowEditorToolkit.h
--- a/Engine/Plugins/Experimental/Dataflow/Source/DataflowEditor/Public/Dataflow/DataflowEditorToolkit.h
+++ b/Engine/Plugins/Experimental/Dataflow/Source/DataflowEditor/Public/Dataflow/DataflowEditorToolkit.h
@@ -33,7 +33,9 @@ class SDataflowMembersWidget;
 class SDataflowConstructionViewport;
 class SDataflowSimulationViewport;
 class SGraphEditor;
+class SDataflowSimulationTimeline;
 struct FDataflowPath;
+class FDataflowSimulationBinding;
 
 enum class EDataflowEditorEvaluationMode : uint8;
 
@@ -107,6 +109,15 @@ public:
 	/** Get the toolkit evaluation mode */
 	const EDataflowEditorEvaluationMode& GetEvaluationMode() const { return EvaluationMode; }
 
+	/** Toggle dataflow simulation */
+	void ToggleDataflowSimulation();
+
+	/** Check if the simulation can be toggled or not */
+	bool HasSimulationManager() const;
+
+	/** Check if the simulation is enabled or not */
+	bool IsSimulationEnabled() const;
+
 protected:
 
 	UDataflowEditor* DataflowEditor = nullptr;
@@ -157,6 +168,7 @@ private:
 	UE_API TSharedRef<SDockTab> SpawnTab_SimulationVisualization(const FSpawnTabArgs& Args);
 	UE_API TSharedRef<SDockTab> SpawnTab_MembersWidget(const FSpawnTabArgs& Args);
 	UE_API TSharedRef<SDockTab> SpawnTab_OutputLog(const FSpawnTabArgs& Args);
+	UE_API TSharedRef<SDockTab> SpawnTab_SimulationTimeline(const FSpawnTabArgs& Args);
 
 	// FTickableEditorObject interface
 	UE_API virtual void Tick(float DeltaTime) override;
@@ -199,6 +211,7 @@ private:
 	static UE_API const FName SimulationViewportTabId;
 	static UE_API const FName PreviewSceneTabId;
 	static UE_API const FName SimulationVisualizationTabId;
+	static UE_API const FName SimulationTimelineTabId;
 	static UE_API const FName MembersWidgetTabId;
 	static UE_API const FName OutputLogTabId;
 
@@ -223,6 +236,7 @@ private:
 	TSharedPtr<FDataflowCollectionSpreadSheet> DataflowCollectionSpreadSheet_4;
 	TSharedPtr<SWidget> AdvancedPreviewSettingsWidget;
 	TSharedPtr<SWidget> SimulationVisualizationWidget;
+	TSharedPtr<SDataflowSimulationTimeline> SimulationTimelineWidget;
 	TSharedPtr<FDataflowOutputLog> DataflowOutputLog;
 
 	/** Customize preview scene with editor/terminal contents */
@@ -342,6 +356,9 @@ private:
 	bool bViewsNeedRefresh = false;
 
 	FString DebugDrawOverlayString;
+
+	/** Simulation Binding to control the simulation scene from the timeline */
+	TSharedPtr<FDataflowSimulationBinding> SimulationBinding;
 };
 
 #undef UE_API
diff --git a/Engine/Plugins/Experimental/Dataflow/Source/DataflowEditor/Public/Dataflow/DataflowSimulationScene.h b/Engine/Plugins/Experimental/Dataflow/Source/DataflowEditor/Public/Dataflow/DataflowSimulationScene.h
--- a/Engine/Plugins/Experimental/Dataflow/Source/DataflowEditor/Public/Dataflow/DataflowSimulationScene.h
+++ b/Engine/Plugins/Experimental/Dataflow/Source/DataflowEditor/Public/Dataflow/DataflowSimulationScene.h
@@ -134,8 +134,11 @@ public:
 	/** Gets whether the simulation is running */
 	UE_API bool IsSimulationEnabled() const;
 
+	/** Check if caching is enabled or not */
+	UE_API bool HasCachingEnabled() const;
+
 	/** Rebuild the simulation scene */
-	UE_API void RebuildSimulationScene(const bool bIsSimulationEnabled);
+	UE_API void RebuildSimulationScene();
 
 	/** Check if there is something to render */
 	bool HasRenderableGeometry() { return true; }
@@ -146,6 +149,9 @@ public:
 	/** Record the simulation cache */
 	UE_API void RecordSimulationCache();
 
+	/** Update the simulation cache */
+	UE_API void UpdateSimulationCache();
+
 	/** Get the simulation time range */
 	const FVector2f& GetTimeRange() const {return TimeRange;}
 
@@ -162,7 +168,10 @@ public:
 	float GetDeltaTime() const { return DeltaTime; }
 
 	/** Simulation time used to drive the cache loading */
-	float SimulationTime;
+	float SimulationTime = 0.0f;
+
+	/** Previous simulation time */
+	float PreviousTime = 0.0f;
 	
 	/** Preview actor accessors */
 	TObjectPtr<AActor> GetPreviewActor() { return PreviewActor; }
@@ -172,10 +181,13 @@ public:
 	UE_API void SetPreviewLOD(int32 InLOD);
 	UE_API int32 GetPreviewLOD() const;
 
-private:
+	/** Check if the simulation is locked or not */
+	bool IsSimulationLocked() const  {return bIsSimulationLocked;}
 
-	/** Update the simulation cache at some point in time (loading/recording)*/
-	UE_API void UpdateSimulationCache(float& DeltaSeconds);
+	/** Set the locked simulation flag */
+	void SetSimulationLocked(const bool bSimulationLocked, const bool bIsPlaying);
+
+private:
 
 	/** Bind the scene selection to the components */
 	UE_API void BindSceneSelection();
@@ -212,6 +224,9 @@ private:
 
 	/** Boolean to check if we are recording the cache or not */
 	bool bIsRecordingCache = false;
+
+	/** Boolean to check if the simulation is locked onto the timeline */
+	bool bIsSimulationLocked = false;
 };
 
 
diff --git a/Engine/Plugins/Experimental/Dataflow/Source/DataflowEditor/Public/Dataflow/DataflowSimulationViewport.h b/Engine/Plugins/Experimental/Dataflow/Source/DataflowEditor/Public/Dataflow/DataflowSimulationViewport.h
--- a/Engine/Plugins/Experimental/Dataflow/Source/DataflowEditor/Public/Dataflow/DataflowSimulationViewport.h
+++ b/Engine/Plugins/Experimental/Dataflow/Source/DataflowEditor/Public/Dataflow/DataflowSimulationViewport.h
@@ -10,6 +10,7 @@ class ADataflowActor;
 class FAdvancedPreviewScene;
 class UDataflowEditorMode;
 class FDataflowConstructionViewportClient;
+struct FToolMenuSection;
 
 // ----------------------------------------------------------------------------------
 
@@ -46,6 +47,9 @@ public:
 	/** Get the simulation scene */
 	const TSharedPtr<class FDataflowSimulationScene>& GetSimulationScene() const;
 
+	/** Build visualisation menu */
+	static void BuildVisualizationMenu(FToolMenuSection& MenuSection);
+
 private:
 	float GetViewMinInput() const;
 	float GetViewMaxInput() const;
diff --git a/Engine/Source/Runtime/Experimental/Dataflow/Engine/Public/Dataflow/DataflowContent.h b/Engine/Source/Runtime/Experimental/Dataflow/Engine/Public/Dataflow/DataflowContent.h
--- a/Engine/Source/Runtime/Experimental/Dataflow/Engine/Public/Dataflow/DataflowContent.h
+++ b/Engine/Source/Runtime/Experimental/Dataflow/Engine/Public/Dataflow/DataflowContent.h
@@ -81,6 +81,12 @@ public:
 	UE_API UDataflowBaseContent();
 	UE_API ~UDataflowBaseContent();
 
+	/** Notification when content data changed */
+	DECLARE_MULTICAST_DELEGATE_OneParam(FOnContentDataChanged, TObjectPtr<AActor>& SimulationActor);
+
+	/** Delegate member to be called when data is changed */
+	FOnContentDataChanged OnContentDataChanged;
+
 	/** 
 	*	Dirty - State Invalidation
 	*   Check if non-graph specific data has been changed, this usually requires a re-render 
diff --git a/Engine/Source/Runtime/Experimental/Dataflow/Engine/Public/Dataflow/DataflowPreview.h b/Engine/Source/Runtime/Experimental/Dataflow/Engine/Public/Dataflow/DataflowPreview.h
--- a/Engine/Source/Runtime/Experimental/Dataflow/Engine/Public/Dataflow/DataflowPreview.h
+++ b/Engine/Source/Runtime/Experimental/Dataflow/Engine/Public/Dataflow/DataflowPreview.h
@@ -42,5 +42,5 @@ public :
 
 	/** Boolean to check if the caching will be done on an async thread (if yes no GT dependency) */
 	UPROPERTY(EditAnywhere, AdvancedDisplay, Category = "Caching")
-	bool bAsyncCaching = true;
+	bool bAsyncCaching = false;
 };
