change DynamicMeshOverlay's CreateFromPredicate calls that were creating per-vertex attributes to use the new, faster/more-direct CreatePerVertex method

diff --git a/Engine/Plugins/ChaosClothAssetEditor/Source/ChaosClothAssetEditorTools/Private/ChaosClothAsset/ClothWeightMapPaintTool.cpp b/Engine/Plugins/ChaosClothAssetEditor/Source/ChaosClothAssetEditorTools/Private/ChaosClothAsset/ClothWeightMapPaintTool.cpp
--- a/Engine/Plugins/ChaosClothAssetEditor/Source/ChaosClothAssetEditorTools/Private/ChaosClothAsset/ClothWeightMapPaintTool.cpp
+++ b/Engine/Plugins/ChaosClothAssetEditor/Source/ChaosClothAssetEditorTools/Private/ChaosClothAsset/ClothWeightMapPaintTool.cpp
@@ -406,7 +406,7 @@ void UClothEditorWeightMapPaintTool::InitializeSculptMeshFromTarget()
 	FDynamicMesh3* Mesh = GetSculptMesh();
 	Mesh->EnableVertexColors(FVector3f::One());
 	Mesh->Attributes()->EnablePrimaryColors();
-	Mesh->Attributes()->PrimaryColors()->CreateFromPredicate([](int ParentVID, int TriIDA, int TriIDB) {return true; }, 0.f);
+	Mesh->Attributes()->PrimaryColors()->CreatePerVertex(0.f);
 	FAxisAlignedBox3d Bounds = Mesh->GetBounds(true);
 
 	TFuture<void> PrecomputeFuture = Async(UE::Chaos::ClothAsset::Private::WeightPaintToolAsyncExecTarget, [this]()
diff --git a/Engine/Plugins/Experimental/Dataflow/Source/DataflowAssetTools/Private/Dataflow/CollectionRenderingPatternUtility.cpp b/Engine/Plugins/Experimental/Dataflow/Source/DataflowAssetTools/Private/Dataflow/CollectionRenderingPatternUtility.cpp
--- a/Engine/Plugins/Experimental/Dataflow/Source/DataflowAssetTools/Private/Dataflow/CollectionRenderingPatternUtility.cpp
+++ b/Engine/Plugins/Experimental/Dataflow/Source/DataflowAssetTools/Private/Dataflow/CollectionRenderingPatternUtility.cpp
@@ -72,7 +72,7 @@ namespace UE::Dataflow
 				}
 
 				FDynamicMeshNormalOverlay* const NormalOverlay = DynamicMesh.Attributes()->PrimaryNormals();
-				NormalOverlay->CreateFromPredicate([](int ParentVID, int TriIDA, int TriIDB) { return true; }, 0.f);
+				NormalOverlay->CreatePerVertex(0.f);
 				auto SetOverlayNormalsFromVertexNormals = [&](int TriangleID)
 				{
 					const FIndex3i Tri = DynamicMesh.GetTriangle(TriangleID);
@@ -89,7 +89,7 @@ namespace UE::Dataflow
 				}
 
 				DynamicMesh.Attributes()->EnablePrimaryColors();
-				DynamicMesh.Attributes()->PrimaryColors()->CreateFromPredicate([](int ParentVID, int TriIDA, int TriIDB) {return true; }, 0.f);
+				DynamicMesh.Attributes()->PrimaryColors()->CreatePerVertex(0.f);
 				DynamicMesh.EnableVertexColors(FVector3f::Zero());
 				FDynamicMeshColorOverlay* const ColorOverlay = DynamicMesh.Attributes()->PrimaryColors();
 				auto SetColorsFromWeights = [&](int TriangleID)
@@ -129,7 +129,7 @@ namespace UE::Dataflow
 					for (int32 UVLayerIndex = 0; UVLayerIndex < NumUVLayers; ++UVLayerIndex)
 					{
 						FDynamicMeshUVOverlay* const UVLayer = DynamicMesh.Attributes()->GetUVLayer(UVLayerIndex);
-						UVLayer->CreateFromPredicate([](int ParentVID, int TriIDA, int TriIDB) { return true; }, 0.f);
+						UVLayer->CreatePerVertex(0.f);
 
 						for (const int TriangleID : DynamicMesh.TriangleIndicesItr())
 						{
diff --git a/Engine/Plugins/Experimental/Dataflow/Source/DataflowEditor/Private/DataflowEditorTools/DataflowEditorVertexAttributePaintTool.cpp b/Engine/Plugins/Experimental/Dataflow/Source/DataflowEditor/Private/DataflowEditorTools/DataflowEditorVertexAttributePaintTool.cpp
--- a/Engine/Plugins/Experimental/Dataflow/Source/DataflowEditor/Private/DataflowEditorTools/DataflowEditorVertexAttributePaintTool.cpp
+++ b/Engine/Plugins/Experimental/Dataflow/Source/DataflowEditor/Private/DataflowEditorTools/DataflowEditorVertexAttributePaintTool.cpp
@@ -481,7 +481,7 @@ void UDataflowEditorVertexAttributePaintTool::Setup()
 	FDynamicMesh3* Mesh = GetSculptMesh();
 	Mesh->EnableVertexColors(FVector3f::One());
 	Mesh->Attributes()->EnablePrimaryColors();
-	Mesh->Attributes()->PrimaryColors()->CreateFromPredicate([](int ParentVID, int TriIDA, int TriIDB) { return true; }, 0.f);
+	Mesh->Attributes()->PrimaryColors()->CreatePerVertex(0.f);
 	UE::Geometry::FAxisAlignedBox3d Bounds = Mesh->GetBounds(true);
 
 	TFuture<void> PrecomputeFuture = Async(UE::DataflowEditorVertexAttributePaintTool::Private::WeightPaintToolAsyncExecTarget, [&]()
diff --git a/Engine/Plugins/Experimental/Dataflow/Source/DataflowEditor/Private/DataflowEditorTools/DataflowEditorWeightMapPaintTool.cpp b/Engine/Plugins/Experimental/Dataflow/Source/DataflowEditor/Private/DataflowEditorTools/DataflowEditorWeightMapPaintTool.cpp
--- a/Engine/Plugins/Experimental/Dataflow/Source/DataflowEditor/Private/DataflowEditorTools/DataflowEditorWeightMapPaintTool.cpp
+++ b/Engine/Plugins/Experimental/Dataflow/Source/DataflowEditor/Private/DataflowEditorTools/DataflowEditorWeightMapPaintTool.cpp
@@ -222,7 +222,7 @@ void UDataflowEditorWeightMapPaintTool::Setup()
 	FDynamicMesh3* Mesh = GetSculptMesh();
 	Mesh->EnableVertexColors(FVector3f::One());
 	Mesh->Attributes()->EnablePrimaryColors();
-	Mesh->Attributes()->PrimaryColors()->CreateFromPredicate([](int ParentVID, int TriIDA, int TriIDB) {return true; }, 0.f);
+	Mesh->Attributes()->PrimaryColors()->CreatePerVertex(0.f);
 	FAxisAlignedBox3d Bounds = Mesh->GetBounds(true);
 
 	TFuture<void> PrecomputeFuture = Async(UE::Dataflow::Private::WeightPaintToolAsyncExecTarget, [&]()
diff --git a/Engine/Plugins/Runtime/GeometryProcessing/Source/DynamicMesh/Private/Operations/TransferDynamicMeshAttributes.cpp b/Engine/Plugins/Runtime/GeometryProcessing/Source/DynamicMesh/Private/Operations/TransferDynamicMeshAttributes.cpp
--- a/Engine/Plugins/Runtime/GeometryProcessing/Source/DynamicMesh/Private/Operations/TransferDynamicMeshAttributes.cpp
+++ b/Engine/Plugins/Runtime/GeometryProcessing/Source/DynamicMesh/Private/Operations/TransferDynamicMeshAttributes.cpp
@@ -55,7 +55,7 @@ namespace TransferValuesLocals
 			MeshAttributes->EnablePrimaryColors();
 
 			// Start with a clean attribute and elements we can write to.
-			MeshAttributes->PrimaryColors()->CreateFromPredicate([](int /*ParentVID*/, int /*TriIDA*/, int /*TriIDB*/){return true;}, 0.f);
+			MeshAttributes->PrimaryColors()->CreatePerVertex(0.f);
 		}
 		
 		if (InSplit)
diff --git a/Engine/Plugins/Runtime/GeometryScripting/Source/GeometryScriptingCore/Private/MeshNormalsFunctions.cpp b/Engine/Plugins/Runtime/GeometryScripting/Source/GeometryScriptingCore/Private/MeshNormalsFunctions.cpp
--- a/Engine/Plugins/Runtime/GeometryScripting/Source/GeometryScriptingCore/Private/MeshNormalsFunctions.cpp
+++ b/Engine/Plugins/Runtime/GeometryScripting/Source/GeometryScriptingCore/Private/MeshNormalsFunctions.cpp
@@ -211,7 +211,7 @@ UDynamicMesh* UGeometryScriptLibrary_MeshNormalsFunctions::RecomputeNormals(
 		if (EditMesh.Attributes()->PrimaryNormals()->ElementCount() == 0)
 		{
 			UE::Geometry::AppendWarning(Debug, EGeometryScriptErrorType::InvalidInputs, LOCTEXT("RecomputeNormals_NothingToRecompute", "RecomputeNormals: TargetMesh did not have normals to recompute; falling back to per-vertex normals. Consider using 'Set Mesh To Per Vertex Normals' or 'Compute Split Normals' instead."));
-			EditMesh.Attributes()->PrimaryNormals()->CreateFromPredicate([](int, int, int)->bool {return true;}, 0.0f);
+			EditMesh.Attributes()->PrimaryNormals()->CreatePerVertex(0.0f);
 		}
 		MeshNormals.RecomputeOverlayNormals(EditMesh.Attributes()->PrimaryNormals(), CalculateOptions.bAreaWeighted, CalculateOptions.bAngleWeighted);
 		MeshNormals.CopyToOverlay(EditMesh.Attributes()->PrimaryNormals(), false);
diff --git a/Engine/Plugins/Runtime/MeshModelingToolset/Source/MeshModelingTools/Private/MeshAttributePaintTool.cpp b/Engine/Plugins/Runtime/MeshModelingToolset/Source/MeshModelingTools/Private/MeshAttributePaintTool.cpp
--- a/Engine/Plugins/Runtime/MeshModelingToolset/Source/MeshModelingTools/Private/MeshAttributePaintTool.cpp
+++ b/Engine/Plugins/Runtime/MeshModelingToolset/Source/MeshModelingTools/Private/MeshAttributePaintTool.cpp
@@ -339,7 +339,7 @@ void UMeshAttributePaintTool::Setup()
 		Mesh.Attributes()->DisablePrimaryColors();
 		Mesh.Attributes()->EnablePrimaryColors();
 		// Create an overlay that has no split elements, init with zero value.
-		Mesh.Attributes()->PrimaryColors()->CreateFromPredicate([](int ParentVID, int TriIDA, int TriIDB){return true;}, 0.f);
+		Mesh.Attributes()->PrimaryColors()->CreatePerVertex(0.f);
 
 		if (GeometrySelection.IsSet())
 		{
