Dataflow : support more type display in collection spreadsheet

diff --git a/Engine/Plugins/Experimental/Dataflow/Source/DataflowEditor/Private/Dataflow/DataflowCollectionSpreadSheetHelpers.cpp b/Engine/Plugins/Experimental/Dataflow/Source/DataflowEditor/Private/Dataflow/DataflowCollectionSpreadSheetHelpers.cpp
--- a/Engine/Plugins/Experimental/Dataflow/Source/DataflowEditor/Private/Dataflow/DataflowCollectionSpreadSheetHelpers.cpp
+++ b/Engine/Plugins/Experimental/Dataflow/Source/DataflowEditor/Private/Dataflow/DataflowCollectionSpreadSheetHelpers.cpp
@@ -57,22 +57,37 @@ namespace UE::Dataflow::CollectionSpreadSheetHelpers
 		return FString::SanitizeFloat(Value, 2);
 	}
 
+	FString AttributeValueToString(double Value)
+	{
+		return FString::SanitizeFloat(Value, 2);
+	}
+
+	FString AttributeValueToString(uint8 Value)
+	{
+		return FString::FromInt((int32)Value);
+	}
+
 	FString AttributeValueToString(int32 Value)
 	{
 		return FString::FromInt(Value);
 	}
 
-	FString AttributeValueToString(FString Value)
+	FString AttributeValueToString(const FString& Value)
 	{
 		return Value;
 	}
 
-	FString AttributeValueToString(FLinearColor Value)
+	FString AttributeValueToString(const FName& Value)
+	{
+		return Value.ToString();
+	}
+
+	FString AttributeValueToString(const FLinearColor& Value)
 	{
 		return FString::Printf(TEXT("(R=%.2f G=%.2f B=%.2f A=%.2f)"), Value.R, Value.G, Value.B, Value.A);
 	}
 
-	FString AttributeValueToString(FVector Value)
+	FString AttributeValueToString(const FVector& Value)
 	{
 		return FString::Printf(TEXT("(X=%2.2f Y=%2.2f Z=%2.2f)"), Value.X, Value.Y, Value.Z);
 	}
@@ -87,7 +102,7 @@ namespace UE::Dataflow::CollectionSpreadSheetHelpers
 		return Value ? FString("true") : FString("false");
 	}
 
-	FString AttributeValueToString(TSet<int32> Value)
+	FString AttributeValueToString(const TSet<int32>& Value)
 	{
 		TArray<int32> Array = Value.Array();
 		FString Out;
@@ -105,7 +120,7 @@ namespace UE::Dataflow::CollectionSpreadSheetHelpers
 		return Out;
 	}
 
-	FString AttributeValueToString(FTransform3f Value)
+	FString AttributeValueToString(const FTransform3f& Value)
 	{
 		const FVector3f Translation = Value.GetTranslation();
 		const FVector3f Rotation = Value.GetRotation().Euler();
@@ -114,7 +129,7 @@ namespace UE::Dataflow::CollectionSpreadSheetHelpers
 		return FString::Printf(TEXT("T:(%s) R:(%s) S:(%s)"), *Translation.ToString(), *Rotation.ToString(), *Scale.ToString());
 	}
 
-	FString AttributeValueToString(FTransform Value)
+	FString AttributeValueToString(const FTransform& Value)
 	{
 		const FVector Translation = Value.GetTranslation();
 		const FVector Rotation = Value.GetRotation().Euler();
@@ -123,29 +138,60 @@ namespace UE::Dataflow::CollectionSpreadSheetHelpers
 		return FString::Printf(TEXT("T:(%s) R:(%s) S:(%s)"), *Translation.ToString(), *Rotation.ToString(), *Scale.ToString());
 	}
 
-	FString AttributeValueToString(FBox Value)
+	FString AttributeValueToString(const FBox& Value)
 	{
 		FVector Center, Extents;
 		Value.GetCenterAndExtents(Center, Extents);
 		return FString::Printf(TEXT("Center:(%s) Extents:(%s)"), *Center.ToString(), *Extents.ToString());
 	}
 
-	FString AttributeValueToString(FIntVector Value)
+	FString AttributeValueToString(const FVector4f& Value)
+	{
+		return Value.ToString();
+	}
+
+	FString AttributeValueToString(const FVector3f& Value)
+	{
+		return Value.ToString();
+	}
+
+	FString AttributeValueToString(const FVector2f& Value)
+	{
+		return Value.ToString();
+	}
+
+	FString AttributeValueToString(const FIntVector& Value)
 	{
 		return FString::Printf(TEXT("%d %d %d"), Value.X, Value.Y, Value.Z);
 	}
 
-	FString AttributeValueToString(FIntVector4 Value)
+	FString AttributeValueToString(const FIntVector2& Value)
+	{
+		return FString::Printf(TEXT("%d %d"), Value.X, Value.Y);
+	}
+
+	FString AttributeValueToString(const FIntVector4& Value)
 	{
 		return FString::Printf(TEXT("%d %d %d %d"), Value[0], Value[1], Value[2], Value[3]);
 	}
 
-	FString AttributeValueToString(FGuid Value)
+	FString AttributeValueToString(const FGuid& Value)
 	{
 		return Value.ToString();
 	}
 
-	FString AttributeValueToString(Chaos::FConvexPtr Value)
+	FString AttributeValueToString(const FSoftObjectPath& Value)
+	{
+		return Value.ToString();
+	}
+
+	FString AttributeValueToString(const TObjectPtr<UObject>& Value)
+	{
+		static const FString NullObjectStr(TEXT("<Null Object>"));
+		return Value ? Value.GetName() : NullObjectStr;
+	}
+
+	FString AttributeValueToString(const Chaos::FConvexPtr& Value)
 	{
 		if (Value)
 		{
@@ -168,6 +214,14 @@ namespace UE::Dataflow::CollectionSpreadSheetHelpers
 			ValueAsString = AttributeValueToString<float>(InCollection, InAttributeName, InGroupName, InIdxColumn);
 			break;
 
+		case FManagedArrayCollection::EArrayType::FDoubleType:
+			ValueAsString = AttributeValueToString<double>(InCollection, InAttributeName, InGroupName, InIdxColumn);
+			break;
+
+		case FManagedArrayCollection::EArrayType::FUInt8Type:
+			ValueAsString = AttributeValueToString<uint8>(InCollection, InAttributeName, InGroupName, InIdxColumn);
+			break;
+
 		case FManagedArrayCollection::EArrayType::FInt32Type:
 			ValueAsString = AttributeValueToString<int32>(InCollection, InAttributeName, InGroupName, InIdxColumn);
 			break;
@@ -180,6 +234,14 @@ namespace UE::Dataflow::CollectionSpreadSheetHelpers
 			ValueAsString = AttributeValueToString<FString>(InCollection, InAttributeName, InGroupName, InIdxColumn);
 			break;
 
+		case FManagedArrayCollection::EArrayType::FNameType:
+			ValueAsString = AttributeValueToString<FName>(InCollection, InAttributeName, InGroupName, InIdxColumn);
+			break;
+
+		case FManagedArrayCollection::EArrayType::FSoftObjectPaathType:
+			ValueAsString = AttributeValueToString<FSoftObjectPath>(InCollection, InAttributeName, InGroupName, InIdxColumn);
+			break;
+
 		case FManagedArrayCollection::EArrayType::FLinearColorType:
 			ValueAsString = AttributeValueToString<FLinearColor>(InCollection, InAttributeName, InGroupName, InIdxColumn);
 			break;
@@ -196,10 +258,34 @@ namespace UE::Dataflow::CollectionSpreadSheetHelpers
 			ValueAsString = AttributeValueToString<FVector2f>(InCollection, InAttributeName, InGroupName, InIdxColumn);
 			break;
 
+		case FManagedArrayCollection::EArrayType::FVector4fType:
+			ValueAsString = AttributeValueToString<FVector4f>(InCollection, InAttributeName, InGroupName, InIdxColumn);
+			break;
+
 		case FManagedArrayCollection::EArrayType::FIntVectorType:
 			ValueAsString = AttributeValueToString<FIntVector>(InCollection, InAttributeName, InGroupName, InIdxColumn);
 			break;
 
+		case FManagedArrayCollection::EArrayType::FIntVector2Type:
+			ValueAsString = AttributeValueToString<FIntVector2>(InCollection, InAttributeName, InGroupName, InIdxColumn);
+			break;
+
+		case FManagedArrayCollection::EArrayType::FFVectorArrayType:
+			ValueAsString = AttributeValueToString<TArray<FVector3f>>(InCollection, InAttributeName, InGroupName, InIdxColumn);
+			break;
+
+		case FManagedArrayCollection::EArrayType::FIntVector3ArrayType:
+			ValueAsString = AttributeValueToString<TArray<FIntVector3>>(InCollection, InAttributeName, InGroupName, InIdxColumn);
+			break;
+
+		case FManagedArrayCollection::EArrayType::FIntVector2ArrayType:
+			ValueAsString = AttributeValueToString<TArray<FIntVector2>>(InCollection, InAttributeName, InGroupName, InIdxColumn);
+			break;
+
+		case FManagedArrayCollection::EArrayType::FVector4fArrayType:
+			ValueAsString = AttributeValueToString<TArray<FVector4f>>(InCollection, InAttributeName, InGroupName, InIdxColumn);
+			break;
+
 		case FManagedArrayCollection::EArrayType::FTransformType:
 			ValueAsString = AttributeValueToString<FTransform>(InCollection, InAttributeName, InGroupName, InIdxColumn);
 			break;
@@ -240,6 +326,10 @@ namespace UE::Dataflow::CollectionSpreadSheetHelpers
 			ValueAsString = AttributeValueToString<Chaos::FConvexPtr>(InCollection, InAttributeName, InGroupName, InIdxColumn);
 			break;
 
+		case FManagedArrayCollection::EArrayType::FUObjectArrayType:
+			ValueAsString = AttributeValueToString<TObjectPtr<UObject>>(InCollection, InAttributeName, InGroupName, InIdxColumn);
+			break;
+
 		default:
 			//ensure(false);
 			ValueAsString = "<Unknown Data Type>";
diff --git a/Engine/Plugins/Experimental/Dataflow/Source/DataflowEditor/Public/Dataflow/DataflowCollectionSpreadSheetHelpers.h b/Engine/Plugins/Experimental/Dataflow/Source/DataflowEditor/Public/Dataflow/DataflowCollectionSpreadSheetHelpers.h
--- a/Engine/Plugins/Experimental/Dataflow/Source/DataflowEditor/Public/Dataflow/DataflowCollectionSpreadSheetHelpers.h
+++ b/Engine/Plugins/Experimental/Dataflow/Source/DataflowEditor/Public/Dataflow/DataflowCollectionSpreadSheetHelpers.h
@@ -7,6 +7,8 @@
 
 class SWidget;
 
+struct FSoftObjectPath;
+
 namespace UE::Dataflow::CollectionSpreadSheetHelpers
 {
 	struct FAttrInfo
@@ -20,12 +22,14 @@ namespace UE::Dataflow::CollectionSpreadSheetHelpers
 		{ TEXT("Transform"),	600 },
 		{ TEXT("Transform3f"),	600 },
 		{ TEXT("String"),		200 },
+		{ TEXT("Name"),			200 },
 		{ TEXT("LinearColor"),	80 },
 		{ TEXT("int32"),		100 },
 		{ TEXT("IntArray"),		200 },
 		{ TEXT("Vector"),		250 },
 		{ TEXT("Vector2D"),		160 },
 		{ TEXT("Float"),		150 },
+		{ TEXT("Double"),		150 },
 		{ TEXT("IntVector"),	220 },
 		{ TEXT("Bool"),			75 },
 		{ TEXT("Box"),			550 },
@@ -36,33 +40,51 @@ namespace UE::Dataflow::CollectionSpreadSheetHelpers
 
 	FString AttributeValueToString(float Value);
 
+	FString AttributeValueToString(double Value);
+
+	FString AttributeValueToString(uint8 Value);
+
 	FString AttributeValueToString(int32 Value);
 
-	FString AttributeValueToString(FString Value);
+	FString AttributeValueToString(const FString& Value);
 
-	FString AttributeValueToString(FLinearColor Value);
+	FString AttributeValueToString(const FName& Value);
 
-	FString AttributeValueToString(FVector Value);
+	FString AttributeValueToString(const FLinearColor& Value);
+
+	FString AttributeValueToString(const FVector& Value);
 
 	FString AttributeValueToString(bool Value);
 
 	FString AttributeValueToString(const FConstBitReference& Value);
 
-	FString AttributeValueToString(TSet<int32> Value);
+	FString AttributeValueToString(const TSet<int32>& Value);
+
+	FString AttributeValueToString(const FTransform3f& Value);
+
+	FString AttributeValueToString(const FTransform& Value);
+
+	FString AttributeValueToString(const FBox& Value);
+
+	FString AttributeValueToString(const FVector4f& Value);
+
+	FString AttributeValueToString(const FVector3f& Value);
+	
+	FString AttributeValueToString(const FVector2f& Value);
 
-	FString AttributeValueToString(FTransform3f Value);
+	FString AttributeValueToString(const FIntVector& Value);
 
-	FString AttributeValueToString(FTransform Value);
+	FString AttributeValueToString(const FIntVector2& Value);
 
-	FString AttributeValueToString(FBox Value);
+	FString AttributeValueToString(const FIntVector4& Value);
 
-	FString AttributeValueToString(FIntVector Value);
+	FString AttributeValueToString(const FGuid& Value);
 
-	FString AttributeValueToString(FIntVector4 Value);
+	FString AttributeValueToString(const FSoftObjectPath& Value);
 
-	FString AttributeValueToString(FGuid Value);
+	FString AttributeValueToString(const TObjectPtr<UObject>& Value);
 
-	FString AttributeValueToString(Chaos::FConvexPtr Value);
+	FString AttributeValueToString(const Chaos::FConvexPtr& Value);
 
 	FString AttributeValueToString(const FManagedArrayCollection& InCollection, const FName& InAttributeName, const FName& InGroupName, int32 InIdxColumn);
 
