Dataflow Editor : change default thickness scale smaller and boundary edge color purple for better readability

diff --git a/Engine/Plugins/Experimental/Dataflow/Source/DataflowEditor/Private/DataflowEditorTools/DataflowEditorWeightMapPaintTool.cpp b/Engine/Plugins/Experimental/Dataflow/Source/DataflowEditor/Private/DataflowEditorTools/DataflowEditorWeightMapPaintTool.cpp
--- a/Engine/Plugins/Experimental/Dataflow/Source/DataflowEditor/Private/DataflowEditorTools/DataflowEditorWeightMapPaintTool.cpp
+++ b/Engine/Plugins/Experimental/Dataflow/Source/DataflowEditor/Private/DataflowEditorTools/DataflowEditorWeightMapPaintTool.cpp
@@ -355,6 +355,8 @@ void UDataflowEditorWeightMapPaintTool::Setup()
 	if (ensure(MeshElementsDisplay->Settings))
 	{
 		MeshElementsDisplay->Settings->bShowNormalSeams = false;
+		MeshElementsDisplay->Settings->ThicknessScale = 0.5f;
+		MeshElementsDisplay->Settings->BoundaryEdgeColor = FColor(164, 73, 164);
 		MeshElementsDisplay->Settings->RestoreProperties(this, TEXT("DataflowEditorWeightMapPaintTool2"));
 		AddToolPropertySource(MeshElementsDisplay->Settings);
 	}
