Fix typo in SoftObjectPaath->SoftObjectPath in ManagedArrayTypeValues.inl from last commit.

diff --git a/Engine/Plugins/Experimental/Dataflow/Source/DataflowEditor/Private/Dataflow/DataflowCollectionSpreadSheetHelpers.cpp b/Engine/Plugins/Experimental/Dataflow/Source/DataflowEditor/Private/Dataflow/DataflowCollectionSpreadSheetHelpers.cpp
--- a/Engine/Plugins/Experimental/Dataflow/Source/DataflowEditor/Private/Dataflow/DataflowCollectionSpreadSheetHelpers.cpp
+++ b/Engine/Plugins/Experimental/Dataflow/Source/DataflowEditor/Private/Dataflow/DataflowCollectionSpreadSheetHelpers.cpp
@@ -238,7 +238,7 @@ namespace UE::Dataflow::CollectionSpreadSheetHelpers
 			ValueAsString = AttributeValueToString<FName>(InCollection, InAttributeName, InGroupName, InIdxColumn);
 			break;
 
-		case FManagedArrayCollection::EArrayType::FSoftObjectPaathType:
+		case FManagedArrayCollection::EArrayType::FSoftObjectPathType:
 			ValueAsString = AttributeValueToString<FSoftObjectPath>(InCollection, InAttributeName, InGroupName, InIdxColumn);
 			break;
 
diff --git a/Engine/Source/Runtime/Experimental/Chaos/Public/GeometryCollection/ManagedArrayTypeValues.inl b/Engine/Source/Runtime/Experimental/Chaos/Public/GeometryCollection/ManagedArrayTypeValues.inl
--- a/Engine/Source/Runtime/Experimental/Chaos/Public/GeometryCollection/ManagedArrayTypeValues.inl
+++ b/Engine/Source/Runtime/Experimental/Chaos/Public/GeometryCollection/ManagedArrayTypeValues.inl
@@ -63,7 +63,7 @@ MANAGED_ARRAY_TYPE(FUintVector2, UintVector2)
 MANAGED_ARRAY_TYPE(TObjectPtr<UObject>, UObjectArray)
 MANAGED_ARRAY_TYPE(Chaos::FLinearCurve, LinearCurve)
 MANAGED_ARRAY_TYPE(FName, Name)
-MANAGED_ARRAY_TYPE(FSoftObjectPath, SoftObjectPaath)
+MANAGED_ARRAY_TYPE(FSoftObjectPath, SoftObjectPath)
 // NOTE: new types must be added at the bottom to keep serialization from breaking
 
 
