Dataflow Editor : Allow adjusting the simulation / playback range from the timeline widget ( will update the preview settings automatically )
- the other way does not work yet - will be part of a future CL

diff --git a/Engine/Plugins/Experimental/Dataflow/Source/DataflowEditor/Private/Dataflow/DataflowSimulationScene.cpp b/Engine/Plugins/Experimental/Dataflow/Source/DataflowEditor/Private/Dataflow/DataflowSimulationScene.cpp
--- a/Engine/Plugins/Experimental/Dataflow/Source/DataflowEditor/Private/Dataflow/DataflowSimulationScene.cpp
+++ b/Engine/Plugins/Experimental/Dataflow/Source/DataflowEditor/Private/Dataflow/DataflowSimulationScene.cpp
@@ -347,6 +347,15 @@ void FDataflowSimulationScene::UpdateSimulationCache()
 	}
 }
 
+void FDataflowSimulationScene::SetTimeRange(const FVector2f& NewTimeRange)
+{
+	TimeRange = NewTimeRange;
+	if (SceneDescription)
+	{
+		SceneDescription->CacheParams.TimeRange = NewTimeRange;
+	}
+}
+
 void FDataflowSimulationScene::SetPreviewLOD(int32 InLOD)
 {
 	CurrentPreviewLOD = InLOD;
diff --git a/Engine/Plugins/Experimental/Dataflow/Source/DataflowEditor/Private/Widgets/DataflowSimulationBinding.cpp b/Engine/Plugins/Experimental/Dataflow/Source/DataflowEditor/Private/Widgets/DataflowSimulationBinding.cpp
--- a/Engine/Plugins/Experimental/Dataflow/Source/DataflowEditor/Private/Widgets/DataflowSimulationBinding.cpp
+++ b/Engine/Plugins/Experimental/Dataflow/Source/DataflowEditor/Private/Widgets/DataflowSimulationBinding.cpp
@@ -62,6 +62,22 @@ TRange<FFrameNumber> FDataflowSimulationBinding::GetPlaybackRange() const
 								FFrameNumber(FMath::RoundToInt32(PlaybackRange.GetUpperBoundValue() * TickResolution)));
 }
 
+
+void FDataflowSimulationBinding::SetPlaybackRange(const TRange<FFrameNumber>& NewRange)
+{
+	const double TickResolution = static_cast<double>(GetTickResolution());
+	if (TickResolution > 0)
+	{
+		PlaybackRange.SetLowerBoundValue((double)NewRange.GetLowerBoundValue().Value / TickResolution);
+		PlaybackRange.SetUpperBoundValue((double)NewRange.GetUpperBoundValue().Value / TickResolution);
+	}
+	if (const TSharedPtr<FDataflowSimulationScene> PreviewScene = SimulationScene.Pin())
+	{
+		PreviewScene->SetTimeRange(FVector2f{ (float)PlaybackRange.GetLowerBoundValue(), (float)PlaybackRange.GetUpperBoundValue() });
+	}
+}
+
+
 FFrameNumber FDataflowSimulationBinding::GetScrubPosition() const
 {
 	if (const TSharedPtr<const FDataflowSimulationScene> PreviewScene = SimulationScene.Pin())
diff --git a/Engine/Plugins/Experimental/Dataflow/Source/DataflowEditor/Private/Widgets/DataflowSimulationBinding.h b/Engine/Plugins/Experimental/Dataflow/Source/DataflowEditor/Private/Widgets/DataflowSimulationBinding.h
--- a/Engine/Plugins/Experimental/Dataflow/Source/DataflowEditor/Private/Widgets/DataflowSimulationBinding.h
+++ b/Engine/Plugins/Experimental/Dataflow/Source/DataflowEditor/Private/Widgets/DataflowSimulationBinding.h
@@ -1,4 +1,4 @@
-﻿// Copyright Epic Games, Inc. All Rights Reserved.
+// Copyright Epic Games, Inc. All Rights Reserved.
 
 #pragma once
 
@@ -36,6 +36,9 @@ public:
 	/** Get the playback range of the model's data */
 	TRange<FFrameNumber> GetPlaybackRange() const;
 
+	/** Set the playback range  */
+	void SetPlaybackRange(const TRange<FFrameNumber>& NewRange);
+
 	/** Get the current scrub position */
 	FFrameNumber GetScrubPosition() const;
 
diff --git a/Engine/Plugins/Experimental/Dataflow/Source/DataflowEditor/Private/Widgets/SDataflowSimulationTimeline.cpp b/Engine/Plugins/Experimental/Dataflow/Source/DataflowEditor/Private/Widgets/SDataflowSimulationTimeline.cpp
--- a/Engine/Plugins/Experimental/Dataflow/Source/DataflowEditor/Private/Widgets/SDataflowSimulationTimeline.cpp
+++ b/Engine/Plugins/Experimental/Dataflow/Source/DataflowEditor/Private/Widgets/SDataflowSimulationTimeline.cpp
@@ -64,10 +64,11 @@ void SDataflowSimulationTimeline::Construct(const FArguments& InArgs, const TSha
 		TimeSliderArgs.TickResolution = TickResolution;
 		TimeSliderArgs.OnViewRangeChanged = FOnViewRangeChanged::CreateSP(&InSimulationBinding.Get(), &FDataflowSimulationBinding::HandleViewRangeChanged);
 		TimeSliderArgs.OnClampRangeChanged = FOnTimeRangeChanged::CreateSP(&InSimulationBinding.Get(), &FDataflowSimulationBinding::HandleWorkingRangeChanged);
-		TimeSliderArgs.IsPlaybackRangeLocked = true;
+		TimeSliderArgs.IsPlaybackRangeLocked = false;
 		TimeSliderArgs.PlaybackStatus = EMovieScenePlayerStatus::Stopped;
 		TimeSliderArgs.NumericTypeInterface = NumericTypeInterface;
 		TimeSliderArgs.OnScrubPositionChanged = FOnScrubPositionChanged::CreateSP(this, &SDataflowSimulationTimeline::HandleScrubPositionChanged);
+		TimeSliderArgs.OnPlaybackRangeChanged = FOnFrameRangeChanged::CreateSP(this, &SDataflowSimulationTimeline::HandlePlaybackRangeChanged);
 	}
 	
 	TimeSliderController = MakeShareable(new FDataflowTimeSliderController(TimeSliderArgs, SharedThis(this)));
@@ -318,6 +319,14 @@ void SDataflowSimulationTimeline::HandleScrubPositionChanged(FFrameTime NewScrub
 	}
 }
 
+void SDataflowSimulationTimeline::HandlePlaybackRangeChanged(TRange<FFrameNumber> NewRange)
+{
+	if (const TSharedPtr<FDataflowSimulationBinding> SharedBinding = SimulationBinding.Pin())
+	{
+		SharedBinding->SetPlaybackRange(NewRange);
+	}
+}
+
 void SDataflowSimulationTimeline::InitTrackNames()
 {
 	TrackNames.Empty();
diff --git a/Engine/Plugins/Experimental/Dataflow/Source/DataflowEditor/Private/Widgets/SDataflowSimulationTimeline.h b/Engine/Plugins/Experimental/Dataflow/Source/DataflowEditor/Private/Widgets/SDataflowSimulationTimeline.h
--- a/Engine/Plugins/Experimental/Dataflow/Source/DataflowEditor/Private/Widgets/SDataflowSimulationTimeline.h
+++ b/Engine/Plugins/Experimental/Dataflow/Source/DataflowEditor/Private/Widgets/SDataflowSimulationTimeline.h
@@ -1,4 +1,4 @@
-﻿// Copyright Epic Games, Inc. All Rights Reserved.
+// Copyright Epic Games, Inc. All Rights Reserved.
 
 #pragma once
 
@@ -32,6 +32,8 @@ private:
 	/** Handle any changes to the scrub position */
 	void HandleScrubPositionChanged(FFrameTime NewScrubPosition, bool bIsScrubbing, bool bEvaluate) const;
 
+	void HandlePlaybackRangeChanged(TRange<FFrameNumber> NewRange);
+
 	/** Init track names from the simulation binding */
 	void InitTrackNames();
 
diff --git a/Engine/Plugins/Experimental/Dataflow/Source/DataflowEditor/Public/Dataflow/DataflowSimulationScene.h b/Engine/Plugins/Experimental/Dataflow/Source/DataflowEditor/Public/Dataflow/DataflowSimulationScene.h
--- a/Engine/Plugins/Experimental/Dataflow/Source/DataflowEditor/Public/Dataflow/DataflowSimulationScene.h
+++ b/Engine/Plugins/Experimental/Dataflow/Source/DataflowEditor/Public/Dataflow/DataflowSimulationScene.h
@@ -155,6 +155,9 @@ public:
 	/** Get the simulation time range */
 	const FVector2f& GetTimeRange() const {return TimeRange;}
 
+	/** Get the simulation time range */
+	void SetTimeRange(const FVector2f& NewTimeRange);
+
 	/** Get the number of frames */
 	const int32& GetNumFrames() const {return NumFrames;}
 
