Dataflow Editor : make the timeline change when the time range in the preview properties is changed

diff --git a/Engine/Plugins/Experimental/Dataflow/Source/DataflowEditor/Private/Widgets/DataflowSimulationBinding.cpp b/Engine/Plugins/Experimental/Dataflow/Source/DataflowEditor/Private/Widgets/DataflowSimulationBinding.cpp
--- a/Engine/Plugins/Experimental/Dataflow/Source/DataflowEditor/Private/Widgets/DataflowSimulationBinding.cpp
+++ b/Engine/Plugins/Experimental/Dataflow/Source/DataflowEditor/Private/Widgets/DataflowSimulationBinding.cpp
@@ -12,7 +12,8 @@ FDataflowSimulationBinding::FDataflowSimulationBinding(const TWeakPtr<FDataflowS
 {
 	if (const TSharedPtr<const FDataflowSimulationScene> PreviewScene = SimulationScene.Pin())
 	{
-		PlaybackRange = FAnimatedRange(PreviewScene->GetTimeRange()[0], PreviewScene->GetTimeRange()[1]);
+		// todo use the animation and playback range and find the common part
+		const FAnimatedRange PlaybackRange(PreviewScene->GetTimeRange()[0], PreviewScene->GetTimeRange()[1]);
 		SetViewRange(PlaybackRange);
 	}
 }
@@ -57,23 +58,33 @@ FAnimatedRange FDataflowSimulationBinding::GetWorkingRange() const
 
 TRange<FFrameNumber> FDataflowSimulationBinding::GetPlaybackRange() const
 {
-	const float TickResolution = static_cast<float>(GetTickResolution());
-	return TRange<FFrameNumber>(FFrameNumber(FMath::RoundToInt32(PlaybackRange.GetLowerBoundValue() * TickResolution)),
-								FFrameNumber(FMath::RoundToInt32(PlaybackRange.GetUpperBoundValue() * TickResolution)));
+	if (const TSharedPtr<const FDataflowSimulationScene> PreviewScene = SimulationScene.Pin())
+	{
+		const float TickResolution = static_cast<float>(GetTickResolution());
+		const float StartRange = PreviewScene->GetTimeRange()[0];
+		const float EndRange = PreviewScene->GetTimeRange()[1];
+		return TRange<FFrameNumber>(
+			FFrameNumber(FMath::RoundToInt32(StartRange * TickResolution)),
+			FFrameNumber(FMath::RoundToInt32(EndRange * TickResolution))
+		);
+	}
+	return TRange<FFrameNumber>(FFrameNumber(0), FFrameNumber(150)); // default 
 }
 
 
 void FDataflowSimulationBinding::SetPlaybackRange(const TRange<FFrameNumber>& NewRange)
 {
-	const double TickResolution = static_cast<double>(GetTickResolution());
-	if (TickResolution > 0)
-	{
-		PlaybackRange.SetLowerBoundValue((double)NewRange.GetLowerBoundValue().Value / TickResolution);
-		PlaybackRange.SetUpperBoundValue((double)NewRange.GetUpperBoundValue().Value / TickResolution);
-	}
 	if (const TSharedPtr<FDataflowSimulationScene> PreviewScene = SimulationScene.Pin())
 	{
-		PreviewScene->SetTimeRange(FVector2f{ (float)PlaybackRange.GetLowerBoundValue(), (float)PlaybackRange.GetUpperBoundValue() });
+		const float TickResolution = static_cast<double>(GetTickResolution());
+		if (TickResolution > 0)
+		{
+			FVector2f TimeRange = {
+				((float)NewRange.GetLowerBoundValue().Value / TickResolution),
+				((float)NewRange.GetUpperBoundValue().Value / TickResolution),
+			};
+			PreviewScene->SetTimeRange(TimeRange);
+		}
 	}
 }
 
diff --git a/Engine/Plugins/Experimental/Dataflow/Source/DataflowEditor/Private/Widgets/DataflowSimulationBinding.h b/Engine/Plugins/Experimental/Dataflow/Source/DataflowEditor/Private/Widgets/DataflowSimulationBinding.h
--- a/Engine/Plugins/Experimental/Dataflow/Source/DataflowEditor/Private/Widgets/DataflowSimulationBinding.h
+++ b/Engine/Plugins/Experimental/Dataflow/Source/DataflowEditor/Private/Widgets/DataflowSimulationBinding.h
@@ -87,8 +87,4 @@ private:
 
 	/** Working simulation range */
 	FAnimatedRange WorkingRange;
-
-	/** Playback simulation range */
-	FAnimatedRange PlaybackRange;
-
 };
diff --git a/Engine/Plugins/Experimental/Dataflow/Source/DataflowEditor/Public/Dataflow/DataflowSimulationScene.h b/Engine/Plugins/Experimental/Dataflow/Source/DataflowEditor/Public/Dataflow/DataflowSimulationScene.h
--- a/Engine/Plugins/Experimental/Dataflow/Source/DataflowEditor/Public/Dataflow/DataflowSimulationScene.h
+++ b/Engine/Plugins/Experimental/Dataflow/Source/DataflowEditor/Public/Dataflow/DataflowSimulationScene.h
@@ -155,7 +155,7 @@ public:
 	/** Get the simulation time range */
 	const FVector2f& GetTimeRange() const {return TimeRange;}
 
-	/** Get the simulation time range */
+	/** Set the simulation time range */
 	void SetTimeRange(const FVector2f& NewTimeRange);
 
 	/** Get the number of frames */
