Dataflow : make sure we use the fallback tool builder for selecting compatible views
- This fixes view modes available when using cloth weightmap tool in Dataflow editor

diff --git a/Engine/Plugins/Experimental/Dataflow/Source/DataflowEditor/Private/DataflowEditorTools/DataflowEditorVertexAttributePaintTool.cpp b/Engine/Plugins/Experimental/Dataflow/Source/DataflowEditor/Private/DataflowEditorTools/DataflowEditorVertexAttributePaintTool.cpp
--- a/Engine/Plugins/Experimental/Dataflow/Source/DataflowEditor/Private/DataflowEditorTools/DataflowEditorVertexAttributePaintTool.cpp
+++ b/Engine/Plugins/Experimental/Dataflow/Source/DataflowEditor/Private/DataflowEditorTools/DataflowEditorVertexAttributePaintTool.cpp
@@ -154,6 +154,13 @@ namespace UE::DataflowEditorVertexAttributePaintTool::Private
 
 void UDataflowEditorVertexAttributePaintToolBuilder::GetSupportedConstructionViewModes(const UDataflowContextObject& ContextObject, TArray<const UE::Dataflow::IDataflowConstructionViewMode*>& Modes) const
 {
+	if (!UE::DataflowEditorVertexAttributePaintTool::CVars::DataflowEditorUseNewWeightMapTool && FallbackToolBuilder)
+	{
+		if (const IDataflowEditorToolBuilder* DataflowFallbackToolBuilder = Cast<IDataflowEditorToolBuilder>(FallbackToolBuilder))
+		{
+			DataflowFallbackToolBuilder->GetSupportedConstructionViewModes(ContextObject, Modes);
+		}
+	}
 	//Modes.Add(UE::Dataflow::EDataflowPatternVertexType::Sim3D);
 	//Modes.Add(UE::Dataflow::EDataflowPatternVertexType::Sim2D);
 }
@@ -1886,17 +1893,14 @@ void UDataflowEditorVertexAttributePaintTool::OnTick(float DeltaTime)
 		PolyLassoMechanic->SetIsEnabled(bIsLasso);
 	}
 
-	const bool bIsGradient = false; // TODO(ccaillaud) (FilterProperties->SubToolType == EDataflowEditorWeightMapPaintInteractionType::Gradient);
 	const bool bIsMeshEditMode = (ToolProperties->EditingMode == EDataflowEditorToolEditMode::Mesh);
 	if (MeshSelector)
 	{
-		MeshSelector->SetIsEnabled(bIsGradient || bIsMeshEditMode);
+		MeshSelector->SetIsEnabled(bIsMeshEditMode);
 	}
 
-	check(!(bIsLasso && bIsGradient));
-
 	ConfigureIndicator(IsVolumetricBrush());
-	SetIndicatorVisibility(!bIsLasso && !bIsGradient && !bIsMeshEditMode);
+	SetIndicatorVisibility(!bIsLasso && !bIsMeshEditMode);
 
 	SCOPE_CYCLE_COUNTER(VertexAttributePaintTool_Tick);
 
