Dataflow editor : Make cloth weight tool editor to properly display the right view mode in the new weight tool
- Added a interface for vertex attribute editing nodes to expose the type of view mode they need for editing
- ( more change to come to also handle remapping of vertices on the node itself )

diff --git a/Engine/Plugins/ChaosClothAssetEditor/Source/ChaosClothAssetDataflowNodes/Private/ChaosClothAsset/WeightMapNode.cpp b/Engine/Plugins/ChaosClothAssetEditor/Source/ChaosClothAssetDataflowNodes/Private/ChaosClothAsset/WeightMapNode.cpp
--- a/Engine/Plugins/ChaosClothAssetEditor/Source/ChaosClothAssetDataflowNodes/Private/ChaosClothAsset/WeightMapNode.cpp
+++ b/Engine/Plugins/ChaosClothAssetEditor/Source/ChaosClothAssetDataflowNodes/Private/ChaosClothAsset/WeightMapNode.cpp
@@ -4,6 +4,7 @@
 #include "ChaosClothAsset/ClothAsset.h"
 #include "ChaosClothAsset/ClothCollectionGroup.h"
 #include "ChaosClothAsset/ClothDataflowTools.h"
+#include "ChaosClothAsset/ClothDataflowViewModes.h"
 #include "ChaosClothAsset/ClothGeometryTools.h"
 #include "ChaosClothAsset/CollectionClothFacade.h"
 #include "ChaosClothAsset/WeightedValue.h"
@@ -245,6 +246,20 @@ FName FChaosClothAssetWeightMapNode::GetInputName(UE::Dataflow::FContext& Contex
 	return InInputName != NAME_None ? InInputName : FName(OutputName.StringValue);
 }
 
+void FChaosClothAssetWeightMapNode::GetSupportedViewModes(UE::Dataflow::FContext& Context, TArray<FName>& OutViewModeNames) const
+{
+	switch (MeshTarget)
+	{
+	case EChaosClothAssetWeightMapMeshTarget::Simulation:
+		OutViewModeNames.Add(UE::Chaos::ClothAsset::FCloth2DSimViewMode::Name);
+		OutViewModeNames.Add(UE::Chaos::ClothAsset::FCloth3DSimViewMode::Name);
+		break;
+	case EChaosClothAssetWeightMapMeshTarget::Render:
+		OutViewModeNames.Add(UE::Chaos::ClothAsset::FClothRenderViewMode::Name);
+		break;
+	} 
+}
+
 void FChaosClothAssetWeightMapNode::GetVertexAttributeValues(UE::Dataflow::FContext& Context, TArray<float>& OutValues) const
 {
 	using namespace UE::Chaos::ClothAsset;
diff --git a/Engine/Plugins/ChaosClothAssetEditor/Source/ChaosClothAssetDataflowNodes/Public/ChaosClothAsset/WeightMapNode.h b/Engine/Plugins/ChaosClothAssetEditor/Source/ChaosClothAssetDataflowNodes/Public/ChaosClothAsset/WeightMapNode.h
--- a/Engine/Plugins/ChaosClothAssetEditor/Source/ChaosClothAssetDataflowNodes/Public/ChaosClothAsset/WeightMapNode.h
+++ b/Engine/Plugins/ChaosClothAssetEditor/Source/ChaosClothAssetDataflowNodes/Public/ChaosClothAsset/WeightMapNode.h
@@ -100,6 +100,7 @@ private:
 	//~ End FDataflowNode interface
 
 	//~ Begin - FDataflowVertexAttributeEditableNode API
+	virtual void GetSupportedViewModes(UE::Dataflow::FContext& Context, TArray<FName>& OutViewModeNames) const override;
 	virtual void GetVertexAttributeValues(UE::Dataflow::FContext& Context, TArray<float>& OutValues) const override;
 	virtual void SetVertexAttributeValues(UE::Dataflow::FContext& Context, const TArray<float>& InValues) override;
 	//~ End - FDataflowVertexAttributeEditableNode API
diff --git a/Engine/Plugins/ChaosClothAssetEditor/Source/ChaosClothAssetEditorTools/Private/ChaosClothAsset/ClothWeightMapPaintTool.cpp b/Engine/Plugins/ChaosClothAssetEditor/Source/ChaosClothAssetEditorTools/Private/ChaosClothAsset/ClothWeightMapPaintTool.cpp
--- a/Engine/Plugins/ChaosClothAssetEditor/Source/ChaosClothAssetEditorTools/Private/ChaosClothAsset/ClothWeightMapPaintTool.cpp
+++ b/Engine/Plugins/ChaosClothAssetEditor/Source/ChaosClothAssetEditorTools/Private/ChaosClothAsset/ClothWeightMapPaintTool.cpp
@@ -480,8 +480,7 @@ void UClothEditorWeightMapPaintTool::InitializeSculptMeshFromTarget()
 				WeightToDynamicMesh = Cloth.GetSimVertex2DLookup();
 			}
 
-			const EClothPatternVertexType ViewMode = UE::Chaos::ClothAsset::DataflowViewModeToClothViewMode(DataflowContextObject->GetConstructionViewMode());
-			const bool bIsRenderMode = ViewMode == EClothPatternVertexType::Render;
+			const bool bIsRenderMode = (WeightMapNodeToUpdate->MeshTarget == EChaosClothAssetWeightMapMeshTarget::Render);
 
 			// Find the map if it exists.
 			if (TSharedPtr<UE::Dataflow::FEngineContext> DataflowContext = DataflowContextObject->GetDataflowContext())
@@ -535,8 +534,9 @@ void UClothEditorWeightMapPaintTool::UpdateShowHideProperties()
 			FCollectionClothConstFacade Cloth(ClothCollection.ToSharedRef());
 			check(Cloth.IsValid());
 
-			const EClothPatternVertexType ViewMode = UE::Chaos::ClothAsset::DataflowViewModeToClothViewMode(DataflowContextObject->GetConstructionViewMode());
-			const bool bIsRenderMode = ViewMode == EClothPatternVertexType::Render;
+			const bool bIsRenderMode = WeightMapNodeToUpdate
+				? (WeightMapNodeToUpdate->MeshTarget == EChaosClothAssetWeightMapMeshTarget::Render)
+				: true;
 
 			const int32 NumPatterns = bIsRenderMode ? Cloth.GetNumRenderPatterns() : Cloth.GetNumSimPatterns();
 
diff --git a/Engine/Plugins/Experimental/Dataflow/Source/DataflowEditor/Private/DataflowEditorTools/DataflowEditorVertexAttributePaintTool.cpp b/Engine/Plugins/Experimental/Dataflow/Source/DataflowEditor/Private/DataflowEditorTools/DataflowEditorVertexAttributePaintTool.cpp
--- a/Engine/Plugins/Experimental/Dataflow/Source/DataflowEditor/Private/DataflowEditorTools/DataflowEditorVertexAttributePaintTool.cpp
+++ b/Engine/Plugins/Experimental/Dataflow/Source/DataflowEditor/Private/DataflowEditorTools/DataflowEditorVertexAttributePaintTool.cpp
@@ -13,6 +13,7 @@
 #include "Dataflow/DataflowEditorCollectionComponent.h"
 #include "Dataflow/DataflowEditorMode.h"
 #include "Dataflow/DataflowNode.h"
+#include "Dataflow/DataflowRenderingViewMode.h"
 #include "DataflowEditorTools/DataflowEditorWeightMapPaintBrushOps.h"
 #include "DataflowEditorTools/DataflowEditorWeightMapPaintTool.h"
 #include "DynamicMesh/DynamicMeshAttributeSet.h"
@@ -38,6 +39,7 @@
 #include "Util/BufferUtil.h"
 #include "Widgets/Notifications/SNotificationList.h"
 
+
 #include UE_INLINE_GENERATED_CPP_BY_NAME(DataflowEditorVertexAttributePaintTool)
 
 DEFINE_LOG_CATEGORY_STATIC(LogDataflowEditorVertexAttributePaintTool, Warning, All);
@@ -156,13 +158,30 @@ void UDataflowEditorVertexAttributePaintToolBuilder::GetSupportedConstructionVie
 {
 	if (!UE::DataflowEditorVertexAttributePaintTool::CVars::DataflowEditorUseNewWeightMapTool && FallbackToolBuilder)
 	{
-		if (const IDataflowEditorToolBuilder* DataflowFallbackToolBuilder = Cast<IDataflowEditorToolBuilder>(FallbackToolBuilder))
+		if (const IDataflowEditorToolBuilder* FallbackDataflowToolBuilder = Cast<IDataflowEditorToolBuilder>(FallbackToolBuilder))
+		{
+			FallbackDataflowToolBuilder->GetSupportedConstructionViewModes(ContextObject, Modes);
+			return;
+		}
+	}
+
+	if (const FDataflowVertexAttributeEditableNode* SelectedNode = ContextObject.GetSelectedNodeOfType<FDataflowVertexAttributeEditableNode>())
+	{
+		if (const TSharedPtr<UE::Dataflow::FEngineContext>& EvalContext = ContextObject.GetDataflowContext())
 		{
-			DataflowFallbackToolBuilder->GetSupportedConstructionViewModes(ContextObject, Modes);
+			TArray<FName> ViewModeNames;
+			SelectedNode->GetSupportedViewModes(*EvalContext, ViewModeNames);
+
+			const UE::Dataflow::FRenderingViewModeFactory& Factory = UE::Dataflow::FRenderingViewModeFactory::GetInstance();
+			for (FName ViewModeName : ViewModeNames)
+			{
+				if (const UE::Dataflow::IDataflowConstructionViewMode* ViewMode = Factory.GetViewMode(ViewModeName))
+				{
+					Modes.Add(ViewMode);
+				}
+			}
 		}
 	}
-	//Modes.Add(UE::Dataflow::EDataflowPatternVertexType::Sim3D);
-	//Modes.Add(UE::Dataflow::EDataflowPatternVertexType::Sim2D);
 }
 
 bool UDataflowEditorVertexAttributePaintToolBuilder::CanBuildTool(const FToolBuilderState& SceneState) const
@@ -597,6 +616,8 @@ void UDataflowEditorVertexAttributePaintTool::Setup()
 	if (ensure(MeshElementsDisplay->Settings))
 	{
 		MeshElementsDisplay->Settings->bShowNormalSeams = false;
+		MeshElementsDisplay->Settings->ThicknessScale = 0.5f;
+		MeshElementsDisplay->Settings->BoundaryEdgeColor = FColor(164, 73, 164);
 		MeshElementsDisplay->Settings->RestoreProperties(this, TEXT("DataflowEditorVertexAttributePaintTool"));
 		AddToolPropertySource(MeshElementsDisplay->Settings);
 	}
diff --git a/Engine/Plugins/Experimental/Dataflow/Source/DataflowNodes/Private/Dataflow/DataflowCollectionAddScalarVertexPropertyNode.cpp b/Engine/Plugins/Experimental/Dataflow/Source/DataflowNodes/Private/Dataflow/DataflowCollectionAddScalarVertexPropertyNode.cpp
--- a/Engine/Plugins/Experimental/Dataflow/Source/DataflowNodes/Private/Dataflow/DataflowCollectionAddScalarVertexPropertyNode.cpp
+++ b/Engine/Plugins/Experimental/Dataflow/Source/DataflowNodes/Private/Dataflow/DataflowCollectionAddScalarVertexPropertyNode.cpp
@@ -232,6 +232,13 @@ FCollectionAttributeKey FDataflowCollectionAddScalarVertexPropertyNode::GetWeigh
 	return Key;
 }
 
+void FDataflowCollectionAddScalarVertexPropertyNode::GetSupportedViewModes(UE::Dataflow::FContext& Context, TArray<FName>& OutViewModeNames) const
+{
+	// matches UE::Dataflow::FDataflowConstruction3DViewMode::Name - but we cannot include the module to get it directly 
+	static const FName DataflowConstruction3DViewModeName = TEXT("3DView");
+	OutViewModeNames.Add(DataflowConstruction3DViewModeName);
+}
+
 void FDataflowCollectionAddScalarVertexPropertyNode::GetVertexAttributeValues(UE::Dataflow::FContext& Context, TArray<float>& OutValues) const
 {
 	const FCollectionAttributeKey InKey = GetWeightAttributeKey(Context);
diff --git a/Engine/Plugins/Experimental/Dataflow/Source/DataflowNodes/Public/Dataflow/DataflowCollectionAddScalarVertexPropertyNode.h b/Engine/Plugins/Experimental/Dataflow/Source/DataflowNodes/Public/Dataflow/DataflowCollectionAddScalarVertexPropertyNode.h
--- a/Engine/Plugins/Experimental/Dataflow/Source/DataflowNodes/Public/Dataflow/DataflowCollectionAddScalarVertexPropertyNode.h
+++ b/Engine/Plugins/Experimental/Dataflow/Source/DataflowNodes/Public/Dataflow/DataflowCollectionAddScalarVertexPropertyNode.h
@@ -90,10 +90,13 @@ struct FDataflowVertexAttributeEditableNode : public FDataflowNode
 	UE_API FDataflowVertexAttributeEditableNode(const UE::Dataflow::FNodeParameters& InParam, FGuid InGuid = FGuid::NewGuid());
 
 public:
-	/** set the absolute vertex attribute values */
-	virtual void GetVertexAttributeValues(UE::Dataflow::FContext& Context, TArray<float>& OutValues) const { ensure(false); }
+	/** Get the list of view mode supported by this node when being edited */
+	virtual void GetSupportedViewModes(UE::Dataflow::FContext& Context, TArray<FName>& OutViewModeNames) const { ensure(false); }
 
 	/** Get the absolute vertex attribute values */
+	virtual void GetVertexAttributeValues(UE::Dataflow::FContext& Context, TArray<float>& OutValues) const { ensure(false); }
+
+	/** Set the absolute vertex attribute values */
 	virtual void SetVertexAttributeValues(UE::Dataflow::FContext& Context, const TArray<float>& InValues) { ensure(false); }
 };
 
@@ -151,6 +154,7 @@ private:
 	UE_API virtual void Evaluate(UE::Dataflow::FContext& Context, const FDataflowOutput* Out) const override;
 
 	//~ Begin - FDataflowVertexAttributeEditableNode API
+	virtual void GetSupportedViewModes(UE::Dataflow::FContext& Context, TArray<FName>& OutViewModeNames) const override;
 	virtual void GetVertexAttributeValues(UE::Dataflow::FContext& Context, TArray<float>& OutValues) const override;
 	virtual void SetVertexAttributeValues(UE::Dataflow::FContext& Context, const TArray<float>& InValues) override;
 	//~ End - FDataflowVertexAttributeEditableNode API
