Dataflow Skin weight correct nodes : make the relax operation to use the selection as a multiplier to the global strength
- also reorganise the way parameters are passed to various operations to make futurechanges easier and code more readable

diff --git a/Engine/Plugins/Experimental/Dataflow/Source/DataflowEditor/Private/DataflowEditorTools/DataflowEditorCorrectSkinWeightsNode.cpp b/Engine/Plugins/Experimental/Dataflow/Source/DataflowEditor/Private/DataflowEditorTools/DataflowEditorCorrectSkinWeightsNode.cpp
--- a/Engine/Plugins/Experimental/Dataflow/Source/DataflowEditor/Private/DataflowEditorTools/DataflowEditorCorrectSkinWeightsNode.cpp
+++ b/Engine/Plugins/Experimental/Dataflow/Source/DataflowEditor/Private/DataflowEditorTools/DataflowEditorCorrectSkinWeightsNode.cpp
@@ -24,6 +24,22 @@
 
 namespace UE::Dataflow::Private
 {
+	namespace SkinWeightOperations
+	{
+		struct FWeightsAndIndicesRef
+		{
+			TArray<TArray<float>>& Weights;
+			TArray<TArray<int32>>& Indices;
+		};
+
+		struct FOperationParams
+		{
+			const FWeightsAndIndicesRef Setup;
+			TConstArrayView<float> SelectionMap;
+			FWeightsAndIndicesRef Final;
+			bool bUseSelectionAsPerVertexFactor = false;
+		};
+	};
 
 static void BuildDynamicMeshes(const FDataflowNode* DataflowNode, UE::Dataflow::FContext& Context, TArray<Geometry::FDynamicMesh3>& DynamicMeshes)
 {
@@ -261,86 +277,92 @@ static void MergeSkinWeights(const Geometry::FDynamicMesh3& DynamicMesh, const T
 	}
 }
 
-static void PruneSkinWeights(const float PruningThreshold, const TArray<TArray<int32>>& SetupIndices, const TArray<TArray<float>>& SetupWeights,
-	const TArray<float>& SelectionMap, TArray<TArray<int32>>& FinalIndices, TArray<TArray<float>>& FinalWeights)
+static void PruneSkinWeights(const float PruningThreshold, const SkinWeightOperations::FOperationParams& Params)
 {
-	for(int32 VertexIndex = 0, NumVertices = FinalIndices.Num(); VertexIndex < NumVertices; ++VertexIndex)
+	const int32 NumVertices = Params.Final.Indices.Num();
+	for (int32 VertexIndex = 0; VertexIndex < NumVertices; ++VertexIndex)
 	{
-		if(SelectionMap[VertexIndex] > 0)
+		if (Params.SelectionMap[VertexIndex] > 0)
 		{
 			float TotalWeight = 0.0f;
-			for(int32 WeightIndex = 0, NumWeights = SetupWeights[VertexIndex].Num(); WeightIndex < NumWeights; ++WeightIndex)
+			for (int32 WeightIndex = 0, NumWeights = Params.Setup.Weights[VertexIndex].Num(); WeightIndex < NumWeights; ++WeightIndex)
 			{
-				if(SetupWeights[VertexIndex][WeightIndex] >= PruningThreshold)
+				if (Params.Setup.Weights[VertexIndex][WeightIndex] >= PruningThreshold)
 				{
-					FinalIndices[VertexIndex].Add(SetupIndices[VertexIndex][WeightIndex]);
-					FinalWeights[VertexIndex].Add(SetupWeights[VertexIndex][WeightIndex]);
-					TotalWeight += SetupWeights[VertexIndex][WeightIndex];
+					Params.Final.Indices[VertexIndex].Add(Params.Setup.Indices[VertexIndex][WeightIndex]);
+					Params.Final.Weights[VertexIndex].Add(Params.Setup.Weights[VertexIndex][WeightIndex]);
+					TotalWeight += Params.Setup.Weights[VertexIndex][WeightIndex];
 				}
 			}
-			if(TotalWeight != 0.0f)
+			if (TotalWeight != 0.0f)
 			{
-				for(int32 WeightIndex = 0, NumWeights = FinalWeights[VertexIndex].Num(); WeightIndex < NumWeights; ++WeightIndex)
+				for (int32 WeightIndex = 0, NumWeights = Params.Final.Weights[VertexIndex].Num(); WeightIndex < NumWeights; ++WeightIndex)
 				{
-					FinalWeights[VertexIndex][WeightIndex] /= TotalWeight;
+					Params.Final.Weights[VertexIndex][WeightIndex] /= TotalWeight;
 				}
 			}
 		}
 		else
 		{
-			FinalIndices[VertexIndex] = SetupIndices[VertexIndex];
-			FinalWeights[VertexIndex] = SetupWeights[VertexIndex];
+			Params.Final.Indices[VertexIndex] = Params.Setup.Indices[VertexIndex];
+			Params.Final.Weights[VertexIndex] = Params.Setup.Weights[VertexIndex];
 		}
 	}
 }
 	
-static void NormalizeSkinWeights(const TArray<TArray<int32>>& SetupIndices, const TArray<TArray<float>>& SetupWeights,
-			const TArray<float>& SelectionMap, TArray<TArray<int32>>& FinalIndices, TArray<TArray<float>>& FinalWeights)
+static void NormalizeSkinWeights(const SkinWeightOperations::FOperationParams& Params)
 {
-	for(int32 VertexIndex = 0, NumVertices = FinalIndices.Num(); VertexIndex < NumVertices; ++VertexIndex)
+	const int32 NumVertices = Params.Final.Indices.Num();
+	for (int32 VertexIndex = 0; VertexIndex < NumVertices; ++VertexIndex)
 	{
-		if(SelectionMap[VertexIndex] > 0)
+		if (Params.SelectionMap[VertexIndex] > 0)
 		{
 			float TotalWeight = 0.0f;
-			for(int32 WeightIndex = 0, NumWeights = SetupWeights[VertexIndex].Num(); WeightIndex < NumWeights; ++WeightIndex)
+			for (int32 WeightIndex = 0, NumWeights = Params.Setup.Weights[VertexIndex].Num(); WeightIndex < NumWeights; ++WeightIndex)
 			{
-				FinalIndices[VertexIndex].Add(SetupIndices[VertexIndex][WeightIndex]);
-				FinalWeights[VertexIndex].Add(SetupWeights[VertexIndex][WeightIndex]);
-				TotalWeight += SetupWeights[VertexIndex][WeightIndex];
+				Params.Final.Indices[VertexIndex].Add(Params.Setup.Indices[VertexIndex][WeightIndex]);
+				Params.Final.Weights[VertexIndex].Add(Params.Setup.Weights[VertexIndex][WeightIndex]);
+				TotalWeight += Params.Setup.Weights[VertexIndex][WeightIndex];
 			}
 			if(TotalWeight != 0.0f)
 			{
-				for(int32 WeightIndex = 0, NumWeights = FinalWeights[VertexIndex].Num(); WeightIndex < NumWeights; ++WeightIndex)
+				for(int32 WeightIndex = 0, NumWeights = Params.Final.Weights[VertexIndex].Num(); WeightIndex < NumWeights; ++WeightIndex)
 				{
-					FinalWeights[VertexIndex][WeightIndex] /= TotalWeight;
+					Params.Final.Weights[VertexIndex][WeightIndex] /= TotalWeight;
 				}
 			}
 		}
 		else
 		{
-			FinalIndices[VertexIndex] = SetupIndices[VertexIndex];
-			FinalWeights[VertexIndex] = SetupWeights[VertexIndex];
+			Params.Final.Indices[VertexIndex] = Params.Setup.Indices[VertexIndex];
+			Params.Final.Weights[VertexIndex] = Params.Setup.Weights[VertexIndex];
 		}
 	}
 }
 
-static void ClampSkinWeights(const int32 ClampingNumber, const TArray<TArray<int32>>& SetupIndices, const TArray<TArray<float>>& SetupWeights,
-		const TArray<float>& SelectionMap, TArray<TArray<int32>>& FinalIndices, TArray<TArray<float>>& FinalWeights)
+static void ClampSkinWeights(const int32 ClampingNumber, const SkinWeightOperations::FOperationParams& Params)
 {
-	for(int32 VertexIndex = 0, NumVertices = FinalIndices.Num(); VertexIndex < NumVertices; ++VertexIndex)
+	const int32 NumVertices = Params.Final.Indices.Num();
+	for (int32 VertexIndex = 0; VertexIndex < NumVertices; ++VertexIndex)
 	{
-		FinalIndices[VertexIndex] = SetupIndices[VertexIndex];
-		FinalWeights[VertexIndex] = SetupWeights[VertexIndex];
+		Params.Final.Indices[VertexIndex] = Params.Setup.Indices[VertexIndex];
+		Params.Final.Weights[VertexIndex] = Params.Setup.Weights[VertexIndex];
 
-		if((SelectionMap[VertexIndex] > 0) && (SetupWeights[VertexIndex].Num() > ClampingNumber))
+		if((Params.SelectionMap[VertexIndex] > 0) && (Params.Setup.Weights[VertexIndex].Num() > ClampingNumber))
 		{
-			RestrictSkinWeights(ClampingNumber, FinalIndices[VertexIndex], FinalWeights[VertexIndex]);
+			RestrictSkinWeights(ClampingNumber, Params.Final.Indices[VertexIndex], Params.Final.Weights[VertexIndex]);
 		}
 	}
 }
 	
-static void SmoothVertexWeights(const UE::Geometry::FDynamicMesh3& DynamicMesh, const TArray<int32>& CollocatedVertices, const int32 VertexOffset, const float SmoothingFactor, 
-		TArray<TArray<int32>>& FinalIndices, TArray<TArray<float>>& FinalWeights)
+static void SmoothVertexWeights(
+	const UE::Geometry::FDynamicMesh3& DynamicMesh, 
+	const TArray<int32>& CollocatedVertices, 
+	const int32 VertexOffset, 
+	const float SmoothingFactor, 
+	const SkinWeightOperations::FWeightsAndIndicesRef& Final,
+	TConstArrayView<float> PerVertexSmoothFactors
+)
 {
 	// Get list of all neighboring vertices, AND this vertex
 	TSet<int32> NeighborVertices;
@@ -354,7 +376,7 @@ static void SmoothVertexWeights(const UE::Geometry::FDynamicMesh3& DynamicMesh,
 	}
 	// Average the per bone weights values
 	TMap<int32, FAveragedBoneWeight> NeighborWeights;
-	if (AverageSkinWeights(NeighborVertices, VertexOffset, FinalIndices, FinalWeights, NeighborWeights))
+	if (AverageSkinWeights(NeighborVertices, VertexOffset, Final.Indices, Final.Weights, NeighborWeights))
 	{
 		for (const int32 CollocatedVertex : CollocatedVertices)
 		{
@@ -370,15 +392,16 @@ static void SmoothVertexWeights(const UE::Geometry::FDynamicMesh3& DynamicMesh,
 				float SmoothWeight = NeighborWeight.Value.WeightValue;
 				float StoredWeight = 0.0f;
 				
-				for (int32 WeightIndex = 0, NumWeights = FinalIndices[GlobalVertex].Num(); WeightIndex < NumWeights; ++WeightIndex)
+				for (int32 WeightIndex = 0, NumWeights = Final.Indices[GlobalVertex].Num(); WeightIndex < NumWeights; ++WeightIndex)
 				{
-					if (BoneIndex == FinalIndices[GlobalVertex][WeightIndex])
+					if (BoneIndex == Final.Indices[GlobalVertex][WeightIndex])
 					{
-						StoredWeight = FinalWeights[GlobalVertex][WeightIndex];
+						StoredWeight = Final.Weights[GlobalVertex][WeightIndex];
 						break;
 					}
 				}
-				BoneWeights.Add(FMath::Lerp<float>(StoredWeight, SmoothWeight, SmoothingFactor));
+				const float PerVertexSmoothFactor = PerVertexSmoothFactors.IsValidIndex(GlobalVertex)? PerVertexSmoothFactors[GlobalVertex]: 1.0f;
+				BoneWeights.Add(FMath::Lerp<float>(StoredWeight, SmoothWeight, SmoothingFactor * PerVertexSmoothFactor));
 				BoneIndices.Add(BoneIndex);
 
 				TotalWeight += BoneWeights.Last();
@@ -390,17 +413,22 @@ static void SmoothVertexWeights(const UE::Geometry::FDynamicMesh3& DynamicMesh,
 				{
 					VertexWeight /= TotalWeight;
 				}
-				FinalIndices[GlobalVertex] = BoneIndices;
-				FinalWeights[GlobalVertex] = BoneWeights;
+				Final.Indices[GlobalVertex] = BoneIndices;
+				Final.Weights[GlobalVertex] = BoneWeights;
 			}
 		}
 	}
 }
 
 // This is an equivalent of the RelaxSkinWeights using the collocated vertices in order to smooth skin weights across seam	
-static void SmoothSkinWeights(const UE::Geometry::FDynamicMesh3& DynamicMesh, const TArray<TArray<int32>>& CollocatedVertices, const int32 VertexOffset, const float SmoothStrength, const int32 NumIterations,
-			const TArray<TArray<int32>>& SetupIndices, const TArray<TArray<float>>& SetupWeights,
-			const TArray<float>& SelectionMap, TArray<TArray<int32>>& FinalIndices, TArray<TArray<float>>& FinalWeights)
+static void SmoothSkinWeights(
+	const UE::Geometry::FDynamicMesh3& DynamicMesh, 
+	const TArray<TArray<int32>>& CollocatedVertices, 
+	const int32 VertexOffset, 
+	const float SmoothStrength, 
+	const int32 NumIterations,
+	const SkinWeightOperations::FOperationParams& Params
+)
 {
 	const int32 NumVertices = DynamicMesh.VertexCount();
 
@@ -411,7 +439,7 @@ static void SmoothSkinWeights(const UE::Geometry::FDynamicMesh3& DynamicMesh, co
 	{
 		for (const int32 VertexIndex : CollocatedVertex)
 		{
-			if (SelectionMap[VertexIndex + VertexOffset] > 0.0f)
+			if (Params.SelectionMap[VertexIndex + VertexOffset] > 0.0f)
 			{
 				// Add the collocated vertex if one of the underlying vertices is selected 
 				SelectedVertices.Add(CollocatedVertex);
@@ -423,37 +451,43 @@ static void SmoothSkinWeights(const UE::Geometry::FDynamicMesh3& DynamicMesh, co
 	for (int32 VertexIndex = 0; VertexIndex < NumVertices; ++VertexIndex)
 	{
 		const int32 GlobalVertex = VertexOffset+VertexIndex;
-		if(SelectionMap[GlobalVertex] > 0.0f && !DynamicMesh.IsBoundaryVertex(VertexIndex))
+		if(Params.SelectionMap[GlobalVertex] > 0.0f && !DynamicMesh.IsBoundaryVertex(VertexIndex))
 		{
 			SelectedVertices.Add({VertexIndex});
 		}
 		// Initialize the final indices/weights with the setup ones
-		FinalIndices[GlobalVertex] = SetupIndices[GlobalVertex];
-		FinalWeights[GlobalVertex] = SetupWeights[GlobalVertex];
+		Params.Final.Indices[GlobalVertex] = Params.Setup.Indices[GlobalVertex];
+		Params.Final.Weights[GlobalVertex] = Params.Setup.Weights[GlobalVertex];
 	}
 
 	static constexpr float PercentPerIteration = 0.95f;
 	const float SmoothingFactor = SmoothStrength * PercentPerIteration;
 	
+	const TArray<float> NoPerVertexSmoothingFactors;
+	const TConstArrayView<float> PerVertexSmoothingFactors = Params.bUseSelectionAsPerVertexFactor ? Params.SelectionMap : NoPerVertexSmoothingFactors;
 	for (int32 IterationIndex = 0; IterationIndex < NumIterations; ++IterationIndex)
 	{
 		for (const TArray<int32>& SelectedVertex : SelectedVertices)
 		{
-			SmoothVertexWeights(DynamicMesh, SelectedVertex, VertexOffset, SmoothingFactor, FinalIndices, FinalWeights);
+			SmoothVertexWeights(DynamicMesh, SelectedVertex, VertexOffset, SmoothingFactor, Params.Final, PerVertexSmoothingFactors);
 		}
 	}
 	for (const TArray<int32>& SelectedVertex : SelectedVertices)
 	{
 		for (const int32 VertexIndex : SelectedVertex)
 		{
-			RestrictSkinWeights(MAX_TOTAL_INFLUENCES, FinalIndices[VertexIndex], FinalWeights[VertexIndex]);
+			RestrictSkinWeights(MAX_TOTAL_INFLUENCES, Params.Final.Indices[VertexIndex], Params.Final.Weights[VertexIndex]);
 		}
 	}
 }
 
-static void RelaxSkinWeights(const UE::Geometry::FDynamicMesh3& DynamicMesh, const int32 VertexOffset, const float SmoothStrength, const int32 NumIterations,
-		const TArray<TArray<int32>>& SetupIndices, const TArray<TArray<float>>& SetupWeights,
-		const TArray<float>& SelectionMap, TArray<TArray<int32>>& FinalIndices, TArray<TArray<float>>& FinalWeights)
+static void RelaxSkinWeights(
+	const UE::Geometry::FDynamicMesh3& DynamicMesh, 
+	const int32 VertexOffset, 
+	const float SmoothStrength, 
+	const int32 NumIterations,
+	const SkinWeightOperations::FOperationParams& Params
+)
 {
 	if(UE::Geometry::FDynamicMeshVertexSkinWeightsAttribute* SkinWeights = DynamicMesh.Attributes()->GetSkinWeightsAttribute(FName("Default")))
 	{
@@ -465,23 +499,23 @@ static void RelaxSkinWeights(const UE::Geometry::FDynamicMesh3& DynamicMesh, con
 		for (int32 VertexIndex = 0; VertexIndex < NumVertices; ++VertexIndex)
 		{
 			const int32 GlobalVertex = VertexOffset+VertexIndex;
-			BoneWeightsBuffer.SetNumUninitialized(SetupIndices[GlobalVertex].Num());
+			BoneWeightsBuffer.SetNumUninitialized(Params.Setup.Indices[GlobalVertex].Num());
 
-			for(int32 BoneIndex = 0; BoneIndex < SetupIndices[GlobalVertex].Num(); ++BoneIndex)
+			for(int32 BoneIndex = 0; BoneIndex < Params.Setup.Indices[GlobalVertex].Num(); ++BoneIndex)
 			{
-				BoneWeightsBuffer[BoneIndex].SetBoneIndex(SetupIndices[VertexIndex+VertexOffset][BoneIndex]);
-				BoneWeightsBuffer[BoneIndex].SetWeight(SetupWeights[VertexIndex+VertexOffset][BoneIndex]);
+				BoneWeightsBuffer[BoneIndex].SetBoneIndex(Params.Setup.Indices[GlobalVertex][BoneIndex]);
+				BoneWeightsBuffer[BoneIndex].SetWeight(Params.Setup.Weights[GlobalVertex][BoneIndex]);
 			}
 			SkinWeights->SetValue(VertexIndex, UE::AnimationCore::FBoneWeights::Create(BoneWeightsBuffer));
 			
-			if(SelectionMap[GlobalVertex] > 0.0f)
+			if(Params.SelectionMap[GlobalVertex] > 0.0f)
 			{
 				SelectedVertices.Add(VertexIndex);
 			}
 			else
 			{
-				FinalIndices[GlobalVertex] = SetupIndices[GlobalVertex];
-				FinalWeights[GlobalVertex] = SetupWeights[GlobalVertex];
+				Params.Final.Indices[GlobalVertex] = Params.Setup.Indices[GlobalVertex];
+				Params.Final.Weights[GlobalVertex] = Params.Setup.Weights[GlobalVertex];
 			}
 		}
 		
@@ -501,13 +535,17 @@ static void RelaxSkinWeights(const UE::Geometry::FDynamicMesh3& DynamicMesh, con
 		for (const int32 SelectedVertex : SelectedVertices)
 		{
 			const int32 GlobalVertex = VertexOffset+SelectedVertex;
-			SkinWeights->GetValue(SelectedVertex, FinalIndices[GlobalVertex], FinalWeights[GlobalVertex]);
+			SkinWeights->GetValue(SelectedVertex, Params.Final.Indices[GlobalVertex], Params.Final.Weights[GlobalVertex]);
 		}
 	}
 }
 
-static void HammerSkinWeights(const UE::Geometry::FDynamicMesh3& DynamicMesh, const int32 VertexOffset, const float SelectionThreshold, const TArray<TArray<int32>>& SetupIndices, const TArray<TArray<float>>& SetupWeights,
-	const TArray<float>& SelectionMap, TArray<TArray<int32>>& FinalIndices, TArray<TArray<float>>& FinalWeights)
+static void HammerSkinWeights(
+	const UE::Geometry::FDynamicMesh3& DynamicMesh, 
+	const int32 VertexOffset, 
+	const float SelectionThreshold,
+	const SkinWeightOperations::FOperationParams& Params
+)
 {
 	const int32 NumVertices = DynamicMesh.VertexCount();
 	
@@ -517,12 +555,12 @@ static void HammerSkinWeights(const UE::Geometry::FDynamicMesh3& DynamicMesh, co
 	for (int32 VertexIndex = 0; VertexIndex < NumVertices; ++VertexIndex)
 	{
 		const int32 GlobalVertex = VertexOffset+VertexIndex;
-		if(SelectionMap[GlobalVertex] > SelectionThreshold)
+		if(Params.SelectionMap[GlobalVertex] > SelectionThreshold)
 		{
 			SelectedVertices.Add(VertexIndex);
 			for (const int32 NeighborIndex : DynamicMesh.VtxVerticesItr(VertexIndex))
 			{
-				if (SelectionMap[VertexOffset+NeighborIndex] <= SelectionThreshold)
+				if (Params.SelectionMap[VertexOffset+NeighborIndex] <= SelectionThreshold)
 				{
 					NeighborVertices.Add(NeighborIndex);
 				}
@@ -530,8 +568,8 @@ static void HammerSkinWeights(const UE::Geometry::FDynamicMesh3& DynamicMesh, co
 		}
 		else
 		{
-			FinalIndices[GlobalVertex] = SetupIndices[GlobalVertex];
-			FinalWeights[GlobalVertex] = SetupWeights[GlobalVertex];
+			Params.Final.Indices[GlobalVertex] = Params.Setup.Indices[GlobalVertex];
+			Params.Final.Weights[GlobalVertex] = Params.Setup.Weights[GlobalVertex];
 		}
 	}
 	UE::Geometry::TMeshDijkstra<UE::Geometry::FDynamicMesh3> PathFinder(&DynamicMesh);
@@ -554,19 +592,19 @@ static void HammerSkinWeights(const UE::Geometry::FDynamicMesh3& DynamicMesh, co
 		TMap<int32, float> SkinWeights;
 		float TotalWeight = 0.0;
 
-		const float InterpolationWeight = FMath::Clamp((SelectionMap[GlobalVertex]-SelectionThreshold) / (1.0-SelectionThreshold), 0.0f, 1.0f);
+		const float InterpolationWeight = FMath::Clamp((Params.SelectionMap[GlobalVertex]-SelectionThreshold) / (1.0-SelectionThreshold), 0.0f, 1.0f);
 
-		AccumulateSkinWeights(InterpolationWeight, SetupIndices[GlobalClosest], SetupWeights[GlobalClosest], SkinWeights, TotalWeight);
-		AccumulateSkinWeights(1.0f-InterpolationWeight, SetupIndices[GlobalVertex], SetupWeights[GlobalVertex], SkinWeights, TotalWeight);
+		AccumulateSkinWeights(InterpolationWeight, Params.Setup.Indices[GlobalClosest], Params.Setup.Weights[GlobalClosest], SkinWeights, TotalWeight);
+		AccumulateSkinWeights(1.0f-InterpolationWeight, Params.Setup.Indices[GlobalVertex], Params.Setup.Weights[GlobalVertex], SkinWeights, TotalWeight);
 
 		if(TotalWeight != 0.0f)
 		{
-			ReportSkinWeights(FinalIndices[GlobalVertex], FinalWeights[GlobalVertex], SkinWeights, TotalWeight);
+			ReportSkinWeights(Params.Final.Indices[GlobalVertex], Params.Final.Weights[GlobalVertex], SkinWeights, TotalWeight);
 		}
 		else
 		{
-			FinalIndices[GlobalVertex] = SetupIndices[GlobalVertex];
-			FinalWeights[GlobalVertex] = SetupWeights[GlobalVertex];
+			Params.Final.Indices[GlobalVertex] = Params.Setup.Indices[GlobalVertex];
+			Params.Final.Weights[GlobalVertex] = Params.Setup.Weights[GlobalVertex];
 		}
 	}
 }
@@ -670,17 +708,25 @@ void FDataflowCorrectSkinWeightsNode::Evaluate(UE::Dataflow::FContext& Context,
 					FinalIndices.SetNum(SetupIndices.Num());
 					FinalWeights.SetNum(SetupWeights.Num());
 				
+					const UE::Dataflow::Private::SkinWeightOperations::FOperationParams Params
+					{
+						.Setup = { SetupWeights, SetupIndices },
+						.SelectionMap = SelectionMap,
+						.Final = { FinalWeights, FinalIndices },
+						.bUseSelectionAsPerVertexFactor = bUseSelectionAsPerVertexFactor,
+					};
+
 					if(CorrectionType == ESkinWeightsCorrectionType::Prune)
 					{
-						UE::Dataflow::Private::PruneSkinWeights(PruningThreshold, SetupIndices, SetupWeights, SelectionMap, FinalIndices, FinalWeights);
+						UE::Dataflow::Private::PruneSkinWeights(PruningThreshold, Params);
 					}
 					else if(CorrectionType == ESkinWeightsCorrectionType::Clamp)
 					{
-						UE::Dataflow::Private::ClampSkinWeights(ClampingNumber, SetupIndices, SetupWeights, SelectionMap, FinalIndices, FinalWeights);
+						UE::Dataflow::Private::ClampSkinWeights(ClampingNumber, Params);
 					}
 					else if(CorrectionType == ESkinWeightsCorrectionType::Normalize)
 					{
-						UE::Dataflow::Private::NormalizeSkinWeights(SetupIndices, SetupWeights, SelectionMap, FinalIndices, FinalWeights);
+						UE::Dataflow::Private::NormalizeSkinWeights(Params);
 					}
 					else
 					{
@@ -696,13 +742,11 @@ void FDataflowCorrectSkinWeightsNode::Evaluate(UE::Dataflow::FContext& Context,
 							
 							if(CorrectionType == ESkinWeightsCorrectionType::Relax)
 							{
-								UE::Dataflow::Private::SmoothSkinWeights(DynamicMesh, CollocatedVertices, VertexOffset, SmoothingFactor, SmoothingIterations,
-									SetupIndices, SetupWeights, SelectionMap, FinalIndices, FinalWeights);
+								UE::Dataflow::Private::SmoothSkinWeights(DynamicMesh, CollocatedVertices, VertexOffset, SmoothingFactor, SmoothingIterations, Params);
 							}
 							else if (CorrectionType == ESkinWeightsCorrectionType::Hammer)
 							{
-								UE::Dataflow::Private::HammerSkinWeights(DynamicMesh, VertexOffset, SelectionThreshold,
-									SetupIndices, SetupWeights, SelectionMap, FinalIndices, FinalWeights);
+								UE::Dataflow::Private::HammerSkinWeights(DynamicMesh, VertexOffset, SelectionThreshold, Params);
 							}
 
 							UE::Dataflow::Private::MergeSkinWeights(DynamicMesh, CollocatedVertices, VertexOffset, SelectionMap, FinalIndices, FinalWeights);
diff --git a/Engine/Plugins/Experimental/Dataflow/Source/DataflowEditor/Public/DataflowEditorTools/DataflowEditorCorrectSkinWeightsNode.h b/Engine/Plugins/Experimental/Dataflow/Source/DataflowEditor/Public/DataflowEditorTools/DataflowEditorCorrectSkinWeightsNode.h
--- a/Engine/Plugins/Experimental/Dataflow/Source/DataflowEditor/Public/DataflowEditorTools/DataflowEditorCorrectSkinWeightsNode.h
+++ b/Engine/Plugins/Experimental/Dataflow/Source/DataflowEditor/Public/DataflowEditorTools/DataflowEditorCorrectSkinWeightsNode.h
@@ -83,6 +83,10 @@ public:
 	UPROPERTY(EditAnywhere, Category = "Skin Weights", meta = (DisplayName = "Smooting Factor", EditCondition = "CorrectionType == ESkinWeightsCorrectionType::Relax", EditConditionHides))
 	float SmoothingFactor = 0.5;
 
+	/** When true, use the decimal values of the selection as a per vertex factor for the selected operation - Currently available for Relax operation only - false by default */
+	UPROPERTY(EditAnywhere, Category = "Skin Weights", meta = (EditCondition = "CorrectionType == ESkinWeightsCorrectionType::Relax", EditConditionHides))
+	bool bUseSelectionAsPerVertexFactor = false;
+
 	/** All weights below this threshold will be pruned */
 	UPROPERTY(EditAnywhere, Category = "Skin Weights", meta = (DisplayName = "Pruning Threshold", EditCondition = "CorrectionType == ESkinWeightsCorrectionType::Prune", EditConditionHides))
 	float PruningThreshold = 0.01;
@@ -94,7 +98,7 @@ public:
 	/** Selection threshold to consider a neighbor skin weight */
     UPROPERTY(EditAnywhere, Category = "Skin Weights", meta = (DisplayName = "Selection Threshold", EditCondition = "CorrectionType == ESkinWeightsCorrectionType::Hammer", EditConditionHides, ClampMin = "0.0", ClampMax = "1.0", UIMin = "0.0", UIMax = "1.0"))
     float SelectionThreshold = 0.5f;
-	
+
 	/** Get the weights attribute key to retrieve/set the bone indices values*/
 	UE_API FCollectionAttributeKey GetBoneIndicesKey(UE::Dataflow::FContext& Context) const;
 
