Dataflow vertex attribute paint tool : address feedback
- Remove strength property as it is not really being used and the attribute value is enough

diff --git a/Engine/Plugins/Experimental/Dataflow/Source/DataflowEditor/Private/DataflowEditorTools/DataflowEditorVertexAttributePaintTool.cpp b/Engine/Plugins/Experimental/Dataflow/Source/DataflowEditor/Private/DataflowEditorTools/DataflowEditorVertexAttributePaintTool.cpp
--- a/Engine/Plugins/Experimental/Dataflow/Source/DataflowEditor/Private/DataflowEditorTools/DataflowEditorVertexAttributePaintTool.cpp
+++ b/Engine/Plugins/Experimental/Dataflow/Source/DataflowEditor/Private/DataflowEditorTools/DataflowEditorVertexAttributePaintTool.cpp
@@ -819,10 +819,10 @@ void UDataflowEditorVertexAttributePaintTool::OnBeginStroke(const FRay& WorldRay
 
 	UpdateBrushPosition(WorldRay);
 
+	PaintBrushOpProperties->Strength = 1.0f; // Attribute painting only need the attribute value 
 	if (PaintBrushOpProperties)
 	{
 		PaintBrushOpProperties->AttributeValue = ToolProperties->BrushProperties.AttributeValue;
-		PaintBrushOpProperties->Strength = ToolProperties->BrushProperties.Strength * ToolProperties->BrushProperties.Strength;
 		PaintBrushOpProperties->EditOperation = ToolProperties->BrushProperties.BrushMode;
 	}
 	if (EraseBrushOpProperties)
diff --git a/Engine/Plugins/Experimental/Dataflow/Source/DataflowEditor/Private/DataflowEditorTools/DataflowEditorVertexAttributePaintToolDetailCustomization.cpp b/Engine/Plugins/Experimental/Dataflow/Source/DataflowEditor/Private/DataflowEditorTools/DataflowEditorVertexAttributePaintToolDetailCustomization.cpp
--- a/Engine/Plugins/Experimental/Dataflow/Source/DataflowEditor/Private/DataflowEditorTools/DataflowEditorVertexAttributePaintToolDetailCustomization.cpp
+++ b/Engine/Plugins/Experimental/Dataflow/Source/DataflowEditor/Private/DataflowEditorTools/DataflowEditorVertexAttributePaintToolDetailCustomization.cpp
@@ -205,7 +205,6 @@ namespace UE::Dataflow::Editor
 		HideProperty(DetailBuilder, GET_MEMBER_NAME_CHECKED(FDataflowEditorVertexAttributePaintToolBrushProperties, BrushAreaMode));
 		HideProperty(DetailBuilder, GET_MEMBER_NAME_CHECKED(FDataflowEditorVertexAttributePaintToolBrushProperties, ValueAtBrush));
 		HideProperty(DetailBuilder, GET_MEMBER_NAME_CHECKED(FDataflowEditorVertexAttributePaintToolBrushProperties, AttributeValue));
-		HideProperty(DetailBuilder, GET_MEMBER_NAME_CHECKED(FDataflowEditorVertexAttributePaintToolBrushProperties, Strength));
 		HideProperty(DetailBuilder, GET_MEMBER_NAME_CHECKED(FDataflowEditorVertexAttributePaintToolDisplayProperties, ColorMode));
 
 		HideProperty(DetailBuilder, GET_MEMBER_NAME_CHECKED(UMeshElementsVisualizerProperties, bVisible));
@@ -346,49 +345,6 @@ namespace UE::Dataflow::Editor
 
 	///////////////////////////////////////////////////////////////////////////////////////////////////////////////
 
-	void FVertexAttributePaintToolDetailCustomization::AddBrushStrengthRow(IDetailCategoryBuilder& BrushCategory) const
-	{
-		BrushCategory.AddCustomRow(LOCTEXT("BrushStrengthCategory", "Brush Strength"), false)
-		.NameContent()
-		[
-			SNew(STextBlock)
-			.Text(LOCTEXT("BrushStrengthLabel", "Strength"))
-			.Font(FAppStyle::GetFontStyle(TEXT("PropertyWindow.NormalFont")))
-			.ToolTipText(LOCTEXT("BrushStrengthTooltip", "The strength of the effect on the weights. Exact effect depends on the Brush mode."))
-		]
-		.ValueContent()
-		[
-			SNew(SSpinBox<float>)
-			.MinValue(0.f)
-			.MaxValue(2.0f)
-			.MaxSliderValue(1.f)
-			.Value(1.0f)
-			.SupportDynamicSliderMaxValue(true)
-			.Value_Lambda([this]()
-				{
-					return ToolProperties.IsValid() ? ToolProperties->BrushProperties.Strength : 1.f;
-				})
-			.OnValueChanged_Lambda([this](float NewValue)
-				{
-					if (ToolProperties.IsValid())
-					{
-						ToolProperties->BrushProperties.Strength = NewValue;
-						FPropertyChangedEvent PropertyChangedEvent(UBrushBaseProperties::StaticClass()->FindPropertyByName(GET_MEMBER_NAME_CHECKED(FDataflowEditorVertexAttributePaintToolBrushProperties, Strength)));
-						ToolProperties->PostEditChangeProperty(PropertyChangedEvent);
-					}
-				})
-			.OnValueCommitted_Lambda([this](float NewValue, ETextCommit::Type CommitType)
-				{
-					if (ToolProperties.IsValid())
-					{
-						// TODO(ccaillaud) ==>> ToolProperties->SaveConfig();
-					}
-				})
-		];
-	}
-
-	///////////////////////////////////////////////////////////////////////////////////////////////////////////////
-
 	void FVertexAttributePaintToolDetailCustomization::AddBrushValueRow(IDetailCategoryBuilder& BrushCategory) const
 	{
 		BrushCategory.AddCustomRow(LOCTEXT("AttributeValueCategory", "Attribute Value"), false)
@@ -447,9 +403,6 @@ namespace UE::Dataflow::Editor
 
 			// Brush value field
 			AddBrushValueRow(BrushCategory);
-
-			// Brush strength field
-			AddBrushStrengthRow(BrushCategory);
 		}
 	}
 
diff --git a/Engine/Plugins/Experimental/Dataflow/Source/DataflowEditor/Private/DataflowEditorTools/DataflowEditorVertexAttributePaintToolDetailCustomization.h b/Engine/Plugins/Experimental/Dataflow/Source/DataflowEditor/Private/DataflowEditorTools/DataflowEditorVertexAttributePaintToolDetailCustomization.h
--- a/Engine/Plugins/Experimental/Dataflow/Source/DataflowEditor/Private/DataflowEditorTools/DataflowEditorVertexAttributePaintToolDetailCustomization.h
+++ b/Engine/Plugins/Experimental/Dataflow/Source/DataflowEditor/Private/DataflowEditorTools/DataflowEditorVertexAttributePaintToolDetailCustomization.h
@@ -39,7 +39,6 @@ namespace UE::Dataflow::Editor
 		void AddBrushFalloffModeRow(IDetailCategoryBuilder& BrushCategory) const;
 		void AddBrushRadiusRow(IDetailCategoryBuilder& BrushCategory) const;
 		void AddBrushValueRow(IDetailCategoryBuilder& BrushCategory) const;
-		void AddBrushStrengthRow(IDetailCategoryBuilder& BrushCategory) const;
 
 		void AddSelectionUI(IDetailLayoutBuilder& DetailBuilder) const;
 		TSharedRef<SWidget> MakeSelectionElementsToolbar() const;
diff --git a/Engine/Plugins/Experimental/Dataflow/Source/DataflowEditor/Public/DataflowEditorTools/DataflowEditorVertexAttributePaintTool.h b/Engine/Plugins/Experimental/Dataflow/Source/DataflowEditor/Public/DataflowEditorTools/DataflowEditorVertexAttributePaintTool.h
--- a/Engine/Plugins/Experimental/Dataflow/Source/DataflowEditor/Public/DataflowEditorTools/DataflowEditorVertexAttributePaintTool.h
+++ b/Engine/Plugins/Experimental/Dataflow/Source/DataflowEditor/Public/DataflowEditorTools/DataflowEditorVertexAttributePaintTool.h
@@ -100,15 +100,10 @@ struct FDataflowEditorVertexAttributePaintToolBrushProperties
 	UPROPERTY()
 	float AttributeValue = 1.0;
 
-	/** How quickly each brush stroke will drive mesh values towards the desired value */
-	UPROPERTY()
-	float Strength = 1.0;
-
 	/** Area Mode specifies the shape of the brush and which triangles will be included relative to the cursor */
 	UPROPERTY(EditAnywhere, Category = Filters, meta = (DisplayName = "Brush Area Mode"))
 	EMeshVertexPaintBrushAreaType BrushAreaMode = EMeshVertexPaintBrushAreaType::Connected;
 
-
 	/** The Region affected by the current operation will be bounded by edge angles larger than this threshold */
 	UPROPERTY(EditAnywhere, Category = Filters, meta = (UIMin = "0.0", ClampMin = "0.0", UIMax = "180.0", ClampMax = "180.0"))
 	float AngleThreshold = 180.0f;
