Dataflow vertex attribute paint tool  improvements
* Show radius and value when using B + mouse drag
* Rename AttributeValue to Value
* invert value when CTRL is pressed
* handle CancelStroke properly ( It was causing issues when debugging )
* save brush config so each brush mode keeps its own settings ( still have a bug where Add does not persists across editor session but the other modes do )
* also persist the display setting ( no the ramp yet , will be part of a later CL )

diff --git a/Engine/Plugins/Experimental/Dataflow/Source/DataflowEditor/Private/DataflowEditorTools/DataflowEditorVertexAttributePaintTool.cpp b/Engine/Plugins/Experimental/Dataflow/Source/DataflowEditor/Private/DataflowEditorTools/DataflowEditorVertexAttributePaintTool.cpp
--- a/Engine/Plugins/Experimental/Dataflow/Source/DataflowEditor/Private/DataflowEditorTools/DataflowEditorVertexAttributePaintTool.cpp
+++ b/Engine/Plugins/Experimental/Dataflow/Source/DataflowEditor/Private/DataflowEditorTools/DataflowEditorVertexAttributePaintTool.cpp
@@ -248,6 +248,57 @@ UMeshSurfacePointTool* UDataflowEditorVertexAttributePaintToolBuilder::CreateNew
 	return PaintTool;
 }
 
+const FDataflowEditorVertexAttributePaintToolBrushConfig& UDataflowEditorVertexAttributePaintToolProperties::GetBrushConfig(EDataflowEditorToolEditOperation BrushMode) const
+{
+	static FDataflowEditorVertexAttributePaintToolBrushConfig DefaultConfig;
+
+	switch (BrushMode)
+	{
+	case EDataflowEditorToolEditOperation::Add:
+		return BrushConfigAdd;
+
+	case EDataflowEditorToolEditOperation::Replace:
+		return BrushConfigReplace;
+
+	case EDataflowEditorToolEditOperation::Multiply:
+		return BrushConfigMultiply;
+
+	case EDataflowEditorToolEditOperation::Relax:
+		return BrushConfigRelax;
+
+	case EDataflowEditorToolEditOperation::Invert:
+	default:
+		return DefaultConfig;
+	};
+}
+
+void UDataflowEditorVertexAttributePaintToolProperties::SetBrushConfig(EDataflowEditorToolEditOperation BrushMode, const FDataflowEditorVertexAttributePaintToolBrushConfig& BrushConfig)
+{
+	switch (BrushMode)
+	{
+	case EDataflowEditorToolEditOperation::Add:
+		BrushConfigAdd = BrushConfig;
+		break;
+
+	case EDataflowEditorToolEditOperation::Replace:
+		BrushConfigReplace = BrushConfig;
+		break;
+
+	case EDataflowEditorToolEditOperation::Multiply:
+		BrushConfigMultiply = BrushConfig;
+		break;
+
+	case EDataflowEditorToolEditOperation::Relax:
+		BrushConfigRelax = BrushConfig;
+		break;
+
+	case EDataflowEditorToolEditOperation::Invert:
+	default:
+		// nothing to update
+		break;
+	}
+}
+
 ///////////////////////////////////////////////////////////////////////////////////////////////////////////////////////////////////////////////////
 //
 // Data
@@ -462,7 +513,7 @@ void UDataflowEditorVertexAttributePaintTool::SetupBrushEditBehaviorSetup(ULocal
 			ToolProperties->PostEditChangeProperty(PropertyChangedEvent);
 #endif
 		};
-	OutBehavior.VerticalProperty.Name = LOCTEXT("AttributeValue", "Attribute Value");
+	OutBehavior.VerticalProperty.Name = LOCTEXT("AttributeValue", "Value");
 	OutBehavior.VerticalProperty.EditRate = 0.005f;
 	OutBehavior.VerticalProperty.bEnabled = true;
 }
@@ -800,7 +851,40 @@ float UDataflowEditorVertexAttributePaintTool::GetBrushMaxRadius() const
 	return UMeshSculptToolBase::BrushProperties->BrushSize.WorldSizeRange.Max * 0.5f;
 }
 
-bool UDataflowEditorVertexAttributePaintTool::IsInBrushSubMode() const
+void UDataflowEditorVertexAttributePaintTool::SetBrushMode(EDataflowEditorToolEditOperation NewBrushMode)
+{
+	if (ToolProperties)
+	{
+		const EDataflowEditorToolEditOperation CurrentBrushMode = ToolProperties->BrushProperties.BrushMode;
+		if (NewBrushMode != CurrentBrushMode)
+		{
+			ToolProperties->BrushProperties.BrushMode = NewBrushMode;
+
+			// Save the current brush config 
+			FDataflowEditorVertexAttributePaintToolBrushConfig BrushConfigToSave;
+			BrushConfigToSave.BrushSize = ToolProperties->BrushProperties.BrushSize;
+			BrushConfigToSave.Value = ToolProperties->BrushProperties.AttributeValue;
+			ToolProperties->SetBrushConfig(CurrentBrushMode, BrushConfigToSave);
+
+			// Load the new brush config 
+			const FDataflowEditorVertexAttributePaintToolBrushConfig& NewBrushConfig = ToolProperties->GetBrushConfig(NewBrushMode);
+			ToolProperties->BrushProperties.BrushSize = NewBrushConfig.BrushSize;
+			ToolProperties->BrushProperties.AttributeValue = NewBrushConfig.Value;
+
+			// Send notification for the changed properties
+			FPropertyChangedEvent BrushModePropertyChangedEvent(UBrushBaseProperties::StaticClass()->FindPropertyByName(GET_MEMBER_NAME_CHECKED(FDataflowEditorVertexAttributePaintToolBrushProperties, BrushMode)));
+			ToolProperties->PostEditChangeProperty(BrushModePropertyChangedEvent);
+			FPropertyChangedEvent BrushSizePropertyChangedEvent(UBrushBaseProperties::StaticClass()->FindPropertyByName(GET_MEMBER_NAME_CHECKED(FDataflowEditorVertexAttributePaintToolBrushProperties, BrushSize)));
+			ToolProperties->PostEditChangeProperty(BrushSizePropertyChangedEvent);
+			FPropertyChangedEvent ValuePropertyChangedEvent(UBrushBaseProperties::StaticClass()->FindPropertyByName(GET_MEMBER_NAME_CHECKED(FDataflowEditorVertexAttributePaintToolBrushProperties, AttributeValue)));
+			ToolProperties->PostEditChangeProperty(ValuePropertyChangedEvent);
+
+			ToolProperties->SaveConfig();
+		}
+	}
+}
+
+bool UDataflowEditorVertexAttributePaintTool::IsInBrushMode() const
 {
 	return (ToolProperties->EditingMode == EDataflowEditorToolEditMode::Brush);
 }
@@ -822,6 +906,10 @@ void UDataflowEditorVertexAttributePaintTool::OnBeginStroke(const FRay& WorldRay
 	if (PaintBrushOpProperties)
 	{
 		PaintBrushOpProperties->AttributeValue = ToolProperties->BrushProperties.AttributeValue;
+		if (GetInInvertStroke())
+		{
+			PaintBrushOpProperties->AttributeValue = FMath::Clamp(1.0f - PaintBrushOpProperties->AttributeValue, 0.f, 1.f);
+		}
 		PaintBrushOpProperties->EditOperation = ToolProperties->BrushProperties.BrushMode;
 		PaintBrushOpProperties->Strength = 1.0f; // Attribute painting only need the attribute value 
 	}
@@ -869,6 +957,14 @@ void UDataflowEditorVertexAttributePaintTool::OnEndStroke()
 	EndChange();
 }
 
+void UDataflowEditorVertexAttributePaintTool::OnCancelStroke()
+{
+	if (const TUniquePtr<FMeshSculptBrushOp>& UseBrushOp = GetActiveBrushOp())
+	{
+		UseBrushOp->CancelStroke();
+	}
+}
+
 void UDataflowEditorVertexAttributePaintTool::UpdateROI(const FSculptBrushStamp& BrushStamp)
 {
 	SCOPE_CYCLE_COUNTER(VertexAttributePaintTool_UpdateROI);
@@ -1021,7 +1117,7 @@ bool UDataflowEditorVertexAttributePaintTool::ApplyStamp()
 
 	// yuck
 	FDataflowMeshVertexWeightMapEditBrushOp* WeightBrushOp = (FDataflowMeshVertexWeightMapEditBrushOp*)UseBrushOp.Get();
-	WeightBrushOp->bApplyRadiusLimit = IsInBrushSubMode();
+	WeightBrushOp->bApplyRadiusLimit = IsInBrushMode();
 
 	FDynamicMesh3* Mesh = GetSculptMesh();
 	WeightBrushOp->ApplyStampByVertices(Mesh, CurrentStamp, VertexROI, ROIWeightValueBuffer);
@@ -1747,7 +1843,7 @@ void UDataflowEditorVertexAttributePaintTool::ApplyVisibilityFilter(const TArray
 
 int32 UDataflowEditorVertexAttributePaintTool::FindHitSculptMeshTriangleConst(const FRay3d& LocalRay) const
 {
-	if (!IsInBrushSubMode())
+	if (!IsInBrushMode())
 	{
 		return IndexConstants::InvalidID;
 	}
@@ -1842,30 +1938,37 @@ void UDataflowEditorVertexAttributePaintTool::DrawHUD(FCanvas* Canvas, IToolsCon
 		}
 	}
 
-	if (ToolProperties && ToolProperties->EditingMode == EDataflowEditorToolEditMode::Brush)
+	if (BrushEditBehavior.IsValid())
+	{
+		BrushEditBehavior->DrawHUD(Canvas, RenderAPI);
+	}
+
+	const bool bInteractiveBrushIsAdjusting = (BrushEditBehavior.IsValid() && BrushEditBehavior->IsEditing());
+
+	// when the user adjust the brush interactively we hide the value on brush
+	if (!bInteractiveBrushIsAdjusting && IsInBrushMode())
 	{
 		if (Canvas && RenderAPI && BrushIndicator)
 		{
-			// display the value under the brush
+			// Display the value under the brush
 			const float BrushRadius = BrushIndicator->BrushRadius;
-			const FVector Offset = (-RenderAPI->GetSceneView()->GetViewUp() * BrushRadius * 1.1f);
-			const FVector BrushPos = BrushIndicator->BrushPosition + Offset;
+			FVector Offset{ 0 };
+			const FVector HoverPosition = HoverStamp.WorldFrame.ToFTransform().GetTranslation();
+
+			const FVector BrushPos = HoverPosition + Offset;
 
 			FVector4 ScreenBrushPos = RenderAPI->GetSceneView()->WorldToScreen(BrushPos);
 			FVector2D ScreenBrushPos2D = FVector2D(ScreenBrushPos.X, ScreenBrushPos.Y);
 
 			RenderAPI->GetSceneView()->ScreenToPixel(ScreenBrushPos, ScreenBrushPos2D);
 
-			const bool bShouldShowValueAtBrush = IsInBrushSubMode();
-			if (bShouldShowValueAtBrush)
-			{
-				const FText BaseText = LOCTEXT("DataflowWeightMapPaintTool_ValueAtBrush_TextFormat", "Value: {0}");
-				const FText ValueAtBrushText = FText::Format(BaseText, FText::AsNumber(ToolProperties->BrushProperties.ValueAtBrush));
-				FCanvasTextItem TextItem(ScreenBrushPos2D, ValueAtBrushText, GEngine->GetLargeFont(), FLinearColor::White);
-				TextItem.EnableShadow(FLinearColor::Black);
-				TextItem.bCentreX = true;
-				Canvas->DrawItem(TextItem);
-			}
+			// Render value under at brush even in mesh mode
+			const FText BaseText = LOCTEXT("DataflowWeightMapPaintTool_ValueAtBrush_TextFormat", "Current Value: {0}");
+			const FText ValueAtBrushText = FText::Format(BaseText, FText::AsNumber(ToolProperties->BrushProperties.ValueAtBrush));
+			FCanvasTextItem TextItem(ScreenBrushPos2D, ValueAtBrushText, GEngine->GetLargeFont(), FLinearColor::White);
+			TextItem.EnableShadow(FLinearColor::Black);
+			TextItem.bCentreX = true;
+			Canvas->DrawItem(TextItem);
 		}
 	}
 }
@@ -1942,9 +2045,9 @@ void UDataflowEditorVertexAttributePaintTool::OnTick(float DeltaTime)
 
 	// Get value at brush location
 	const bool bShouldPickWeight = bPendingPickWeight && IsStampPending() == false;
-	const bool bShouldUpdateValueAtBrush = IsInBrushSubMode();
+//	const bool bShouldUpdateValueAtBrush = IsInBrushSubMode();
 
-	if (bShouldPickWeight || bShouldUpdateValueAtBrush)
+	//if (bShouldPickWeight /*|| bShouldUpdateValueAtBrush*/)
 	{
 		if (GetSculptMesh()->IsTriangle(GetBrushTriangleID()))
 		{
@@ -1956,7 +2059,7 @@ void UDataflowEditorVertexAttributePaintTool::OnTick(float DeltaTime)
 				NotifyOfPropertyChangeByTool(ToolProperties);
 			}
 
-			if (bShouldUpdateValueAtBrush)
+			/*if (bShouldUpdateValueAtBrush)*/
 			{
 				ToolProperties->BrushProperties.ValueAtBrush = HitWeightValue;
 			}
@@ -1994,7 +2097,7 @@ void UDataflowEditorVertexAttributePaintTool::OnTick(float DeltaTime)
 		};
 
 
-	if (IsInBrushSubMode())
+	if (IsInBrushMode())
 	{
 		ProcessPerTickStamps(
 			[this](const FRay& StampRay) -> bool {
@@ -2134,7 +2237,7 @@ int32 UDataflowEditorVertexAttributePaintTool::GetBrushNearestVertex() const
 
 void UDataflowEditorVertexAttributePaintTool::UpdateBrushType(EDataflowEditorToolEditOperation BrushMode)
 {
-	static const FText BaseMessage = LOCTEXT("OnStartTool", "Hold Shift to Erase. [/] and S/D change Size (+Shift to small-step)");
+	static const FText BaseMessage = LOCTEXT("OnStartTool", "Hold Shift to Erase. Ctrl to Invert Value, [/] and S/D change Size (+Shift to small-step)");
 	FTextBuilder Builder;
 	Builder.AppendLine(BaseMessage);
 
@@ -2154,6 +2257,10 @@ void UDataflowEditorVertexAttributePaintTool::UpdateBrushType(EDataflowEditorToo
 
 void UDataflowEditorVertexAttributePaintTool::OnColorRampChanged(TArray<FRichCurve*> Curves)
 {
+	if (ToolProperties)
+	{
+		ToolProperties->SaveConfig();
+	}
 	UpdateVertexColorOverlay();
 	GetToolManager()->PostInvalidation();
 }
@@ -2163,6 +2270,7 @@ void UDataflowEditorVertexAttributePaintTool::SetColorMode(EDataflowEditorToolCo
 	if (ToolProperties)
 	{
 		ToolProperties->DisplayProperties.ColorMode = NewColorMode;
+		ToolProperties->SaveConfig();
 		UpdateVertexColorOverlay();
 		GetToolManager()->PostInvalidation();
 	}
diff --git a/Engine/Plugins/Experimental/Dataflow/Source/DataflowEditor/Private/DataflowEditorTools/DataflowEditorVertexAttributePaintToolDetailCustomization.cpp b/Engine/Plugins/Experimental/Dataflow/Source/DataflowEditor/Private/DataflowEditorTools/DataflowEditorVertexAttributePaintToolDetailCustomization.cpp
--- a/Engine/Plugins/Experimental/Dataflow/Source/DataflowEditor/Private/DataflowEditorTools/DataflowEditorVertexAttributePaintToolDetailCustomization.cpp
+++ b/Engine/Plugins/Experimental/Dataflow/Source/DataflowEditor/Private/DataflowEditorTools/DataflowEditorVertexAttributePaintToolDetailCustomization.cpp
@@ -234,11 +234,9 @@ namespace UE::Dataflow::Editor
 					})
 				.OnValueChanged_Lambda([this](EDataflowEditorToolEditOperation Mode)
 					{
-						if (ToolProperties.IsValid())
+						if (Tool.IsValid())
 						{
-							ToolProperties->BrushProperties.BrushMode = Mode;
-							FPropertyChangedEvent PropertyChangedEvent(UBrushBaseProperties::StaticClass()->FindPropertyByName(GET_MEMBER_NAME_CHECKED(FDataflowEditorVertexAttributePaintToolBrushProperties, BrushMode)));
-							ToolProperties->PostEditChangeProperty(PropertyChangedEvent);
+							Tool->SetBrushMode(Mode);
 						}
 					})
 				+ SSegmentedControl<EDataflowEditorToolEditOperation>::Slot(EDataflowEditorToolEditOperation::Add)
@@ -347,13 +345,13 @@ namespace UE::Dataflow::Editor
 
 	void FVertexAttributePaintToolDetailCustomization::AddBrushValueRow(IDetailCategoryBuilder& BrushCategory) const
 	{
-		BrushCategory.AddCustomRow(LOCTEXT("AttributeValueCategory", "Attribute Value"), false)
+		BrushCategory.AddCustomRow(LOCTEXT("ValueCategory", "Value"), false)
 		.NameContent()
 		[
 			SNew(STextBlock)
-			.Text(LOCTEXT("AttributeValueLabel", "Attribute Value"))
+			.Text(LOCTEXT("ValueLabel", "Value"))
 			.Font(FAppStyle::GetFontStyle(TEXT("PropertyWindow.NormalFont")))
-			.ToolTipText(LOCTEXT("AttributeValueTooltip", "The value of the attribute to be applied. Exact effect depends on the Brush strength and mode."))
+			.ToolTipText(LOCTEXT("ValueTooltip", "The value of the attribute to be applied. Exact effect depends on the selected brush mode."))
 		]
 		.ValueContent()
 		[
diff --git a/Engine/Plugins/Experimental/Dataflow/Source/DataflowEditor/Public/DataflowEditorTools/DataflowEditorVertexAttributePaintTool.h b/Engine/Plugins/Experimental/Dataflow/Source/DataflowEditor/Public/DataflowEditorTools/DataflowEditorVertexAttributePaintTool.h
--- a/Engine/Plugins/Experimental/Dataflow/Source/DataflowEditor/Public/DataflowEditorTools/DataflowEditorVertexAttributePaintTool.h
+++ b/Engine/Plugins/Experimental/Dataflow/Source/DataflowEditor/Public/DataflowEditorTools/DataflowEditorVertexAttributePaintTool.h
@@ -125,6 +125,21 @@ struct FDataflowEditorVertexAttributePaintToolBrushProperties
 	double ValueAtBrush = 0;
 };
 
+USTRUCT()
+struct FDataflowEditorVertexAttributePaintToolBrushConfig
+{
+	GENERATED_BODY()
+
+public:
+	/** Adaptive size of brush - Relative to the model size */
+	UPROPERTY(Config)
+	float BrushSize = 0.25f;
+
+	/** Value to paint on the mesh (dependent on the BrushMode) */
+	UPROPERTY(Config)
+	float Value = 1.0;
+};
+
 ///////////////////////////////////////////////////////////////////////////////////////////////////////////////////////////////////////////////////
 
 USTRUCT()
@@ -172,7 +187,7 @@ struct FDataflowEditorVertexAttributePaintToolDisplayProperties
 
 ///////////////////////////////////////////////////////////////////////////////////////////////////////////////////////////////////////////////////
 
-UCLASS(MinimalAPI)
+UCLASS(MinimalAPI, config = EditorSettings)
 class UDataflowEditorVertexAttributePaintToolProperties : public UInteractiveToolPropertySet
 {
 	GENERATED_BODY()
@@ -192,8 +207,24 @@ public:
 	UPROPERTY()
 	FDataflowEditorVertexAttributePaintToolSelectionProperties SelectionProperties;
 
-	UPROPERTY()
+	UPROPERTY(config)
 	FDataflowEditorVertexAttributePaintToolDisplayProperties DisplayProperties;
+
+	const FDataflowEditorVertexAttributePaintToolBrushConfig& GetBrushConfig(EDataflowEditorToolEditOperation BrushMode) const;
+	void SetBrushConfig(EDataflowEditorToolEditOperation BrushMode, const FDataflowEditorVertexAttributePaintToolBrushConfig& BrushConfig);
+
+	UPROPERTY(config)
+	FDataflowEditorVertexAttributePaintToolBrushConfig BrushConfigAdd;
+
+	UPROPERTY(config)
+	FDataflowEditorVertexAttributePaintToolBrushConfig BrushConfigReplace;
+
+	UPROPERTY(config)
+	FDataflowEditorVertexAttributePaintToolBrushConfig BrushConfigMultiply;
+
+	UPROPERTY(config)
+	FDataflowEditorVertexAttributePaintToolBrushConfig BrushConfigRelax;
+
 };
 
 ///////////////////////////////////////////////////////////////////////////////////////////////////////////////////////////////////////////////////
@@ -280,7 +311,8 @@ public:
 
 	UE_API virtual void OnPropertyModified(UObject* PropertySet, FProperty* Property) override;
 
-	UE_API bool IsInBrushSubMode() const;
+	UE_API void SetBrushMode(EDataflowEditorToolEditOperation BrushMode);
+	UE_API bool IsInBrushMode() const;
 	UE_API bool IsVolumetricBrush() const;
 
 	UE_API virtual void CommitResult(UBaseDynamicMeshComponent* Component, bool bModifiedTopology) override;
@@ -348,6 +380,7 @@ protected:
 
 	UE_API virtual void OnBeginStroke(const FRay& WorldRay) override;
 	UE_API virtual void OnEndStroke() override;
+	UE_API virtual void OnCancelStroke() override;
 
 	UE_API void OnColorRampChanged(TArray<FRichCurve*> Curves);
 
