Fix warning C6235: (<non-zero constant> || <expression>) is always a non-zero constant.

#rnx

diff --git a/Engine/Plugins/Experimental/Dataflow/Source/DataflowEditor/Private/DataflowEditorTools/DataflowEditorVertexAttributePaintTool.cpp b/Engine/Plugins/Experimental/Dataflow/Source/DataflowEditor/Private/DataflowEditorTools/DataflowEditorVertexAttributePaintTool.cpp
--- a/Engine/Plugins/Experimental/Dataflow/Source/DataflowEditor/Private/DataflowEditorTools/DataflowEditorVertexAttributePaintTool.cpp
+++ b/Engine/Plugins/Experimental/Dataflow/Source/DataflowEditor/Private/DataflowEditorTools/DataflowEditorVertexAttributePaintTool.cpp
@@ -2011,7 +2011,7 @@ void UDataflowEditorVertexAttributePaintTool::OnTick(float DeltaTime)
 	UMeshSculptToolBase::OnTick(DeltaTime);
 	MeshElementsDisplay->OnTick(DeltaTime);
 
-	const bool bIsLasso = false; // TODO(ccaillaud) (FilterProperties->SubToolType == EDataflowEditorWeightMapPaintInteractionType::PolyLasso);
+	bool bIsLasso = false; // TODO(ccaillaud) const bool bIsLasso = (FilterProperties->SubToolType == EDataflowEditorWeightMapPaintInteractionType::PolyLasso);
 	if (PolyLassoMechanic)
 	{
 		PolyLassoMechanic->SetIsEnabled(bIsLasso);
