Dataflow : automatically select the first tool palette when opening the Dataflow Editor

diff --git a/Engine/Plugins/Experimental/Dataflow/Source/DataflowEditor/Private/Dataflow/DataflowEditorModeToolkit.cpp b/Engine/Plugins/Experimental/Dataflow/Source/DataflowEditor/Private/Dataflow/DataflowEditorModeToolkit.cpp
--- a/Engine/Plugins/Experimental/Dataflow/Source/DataflowEditor/Private/Dataflow/DataflowEditorModeToolkit.cpp
+++ b/Engine/Plugins/Experimental/Dataflow/Source/DataflowEditor/Private/Dataflow/DataflowEditorModeToolkit.cpp
@@ -254,6 +254,18 @@ void FDataflowEditorModeToolkit::RegisterPalettes()
 				ToolkitBuilder->AddPalette(MakeShareable(new FToolPalette(NewCommandInfo.ToSharedRef(), CategoryCommands)));
 			}
 		}
+
+		// Select the first tool palette if available
+		if (ToolCategories.Num() > 0)
+		{
+			if (const TSharedPtr<FUICommandInfo>* CommandInfo = Commands.ToolPaletteCommands.Find(ToolCategories[0]))
+			{
+				if (CommandInfo->IsValid())
+				{
+					ToolkitBuilder->SetActivePaletteOnLoad((*CommandInfo).Get());
+				}
+			}
+		}
 	}
 	
 	//ToolkitBuilder->SetActivePaletteOnLoad(Commands.LoadGeneralTools.Get());
