Dataflow : new Weight paint tool improvements
* better handle of brush cancelling
* shift + brush now shift to using Relax tool ( got rid of erase tool )

diff --git a/Engine/Plugins/Experimental/Dataflow/Source/DataflowEditor/Private/DataflowEditorTools/DataflowEditorVertexAttributePaintTool.cpp b/Engine/Plugins/Experimental/Dataflow/Source/DataflowEditor/Private/DataflowEditorTools/DataflowEditorVertexAttributePaintTool.cpp
--- a/Engine/Plugins/Experimental/Dataflow/Source/DataflowEditor/Private/DataflowEditorTools/DataflowEditorVertexAttributePaintTool.cpp
+++ b/Engine/Plugins/Experimental/Dataflow/Source/DataflowEditor/Private/DataflowEditorTools/DataflowEditorVertexAttributePaintTool.cpp
@@ -248,6 +248,50 @@ UMeshSurfacePointTool* UDataflowEditorVertexAttributePaintToolBuilder::CreateNew
 	return PaintTool;
 }
 
+///////////////////////////////////////////////////////////////////////////////////////////////////////////////////////////////////////////////////
+//
+// Properties
+//
+///////////////////////////////////////////////////////////////////////////////////////////////////////////////////////////////////////////////////
+
+UDataflowEditorVertexAttributePaintToolProperties::UDataflowEditorVertexAttributePaintToolProperties()
+	: UInteractiveToolPropertySet()
+{
+	LoadConfig();
+
+	if (DisplayProperties.ColorRamp.ColorRamp.IsEmpty())
+	{
+		CreateDefaultColorRamp();
+	}
+}
+
+void UDataflowEditorVertexAttributePaintToolProperties::CreateDefaultColorRamp()
+{
+	constexpr FLinearColor HeatMapColors[] =
+	{
+		FLinearColor(0.8f, 0.4f, 0.8f), // Purple
+		FLinearColor(0.0f, 0.0f, 0.5f), // Dark Blue
+		FLinearColor(0.2f, 0.2f, 1.0f), // Light Blue
+		FLinearColor(0.0f, 1.0f, 0.0f), // Green
+		FLinearColor(1.0f, 1.0f, 0.0f), // Yellow
+		FLinearColor(1.0f, 0.65f, 0.0f), // Orange
+		FLinearColor(1.0f, 0.0f, 0.0f, 0.0f), // Red
+	};
+	constexpr int32 NumColors = sizeof(HeatMapColors) / sizeof(FLinearColor);
+	ensure(NumColors > 1);
+
+	// Set ColorRamp
+	FDataflowColorCurveOwner& Ramp = DisplayProperties.ColorRamp.ColorRamp;
+	for (int32 Index = 0; Index < NumColors; ++Index)
+	{
+		const float Time = (float)Index / ((float)NumColors - 1.f);
+		Ramp.GetCurves()[0].CurveToEdit->AddKey(Time, HeatMapColors[Index].R);
+		Ramp.GetCurves()[1].CurveToEdit->AddKey(Time, HeatMapColors[Index].G);
+		Ramp.GetCurves()[2].CurveToEdit->AddKey(Time, HeatMapColors[Index].B);
+	}
+	Ramp.GetCurves()[3].CurveToEdit->AddKey(0.f, 1.f); // single alpha  value
+}
+
 const FDataflowEditorVertexAttributePaintToolBrushConfig& UDataflowEditorVertexAttributePaintToolProperties::GetBrushConfig(EDataflowEditorToolEditOperation BrushMode) const
 {
 	static FDataflowEditorVertexAttributePaintToolBrushConfig DefaultConfig;
@@ -421,6 +465,17 @@ TUniquePtr<UE::DataflowEditorVertexAttributePaintTool::Private::FMeshChange> FDa
 	return {};
 }
 
+void FDataflowEditorVertexAttributePaintToolData::CancelChange()
+{
+	check(ActiveWeightEditChangeTracker);
+
+	if (ActiveWeightEditChangeTracker)
+	{
+		ActiveWeightEditChangeTracker->EndChange();
+		ActiveWeightEditChangeTracker = nullptr;
+	}
+}
+
 float FDataflowEditorVertexAttributePaintToolData::GetValue(int32 VertexIdx) const
 {
 	float Value = 0.f;
@@ -606,10 +661,9 @@ void UDataflowEditorVertexAttributePaintTool::Setup()
 	{
 		ToolProperties->WatchProperty(ToolProperties->BrushProperties.BrushSize,
 			[this](float NewSize) { UMeshSculptToolBase::BrushProperties->BrushSize.AdaptiveSize = NewSize; });
-		ToolProperties->BrushProperties.BrushSize = UMeshSculptToolBase::BrushProperties->BrushSize.AdaptiveSize;
-		CreateDefaultColorRamp(ToolProperties->DisplayProperties.ColorRamp);
+		//ToolProperties->BrushProperties.BrushSize = UMeshSculptToolBase::BrushProperties->BrushSize.AdaptiveSize;
 		ToolProperties->DisplayProperties.ColorRamp.ColorRamp.OnColorCurveChangedDelegate.AddUObject(this, &UDataflowEditorVertexAttributePaintTool::OnColorRampChanged);
-		ToolProperties->RestoreProperties(this);
+		//ToolProperties->RestoreProperties(this);
 	}
 	AddToolPropertySource(ToolProperties);
 
@@ -645,10 +699,10 @@ void UDataflowEditorVertexAttributePaintTool::Setup()
 		PaintBrushOpProperties);
 
 	// secondary brushes
-	EraseBrushOpProperties = NewObject<UDataflowWeightMapEraseBrushOpProps>(this);
-	RegisterSecondaryBrushType(EraseBrushId, LOCTEXT("Erase", "Erase"),
-		MakeUnique<TBasicMeshSculptBrushOpFactory<FDataflowWeightMapEraseBrushOp>>(),
-		EraseBrushOpProperties);
+	SmoothBrushOpProperties = NewObject<UDataflowWeightMapSmoothBrushOpProps>(this);
+	RegisterSecondaryBrushType(SmoothBrushId, LOCTEXT("SmoothBrushType", "Smooth"),
+		MakeUnique<FLambdaMeshSculptBrushOpFactory>([this]() { return MakeUnique<FDataflowWeightMapSmoothBrushOp>(); }),
+		SmoothBrushOpProperties);
 
 	AddToolPropertySource(UMeshSculptToolBase::ViewProperties);
 	SetToolPropertySourceEnabled(UMeshSculptToolBase::ViewProperties, false);
@@ -659,6 +713,7 @@ void UDataflowEditorVertexAttributePaintTool::Setup()
 	// must call before updating brush type so that we register all brush properties?
 	UMeshSculptToolBase::OnCompleteSetup();
 
+
 	UpdateBrushType(ToolProperties->BrushProperties.BrushMode);
 
 	// mesh element display is used to show seams and other features ofthe mesh
@@ -691,33 +746,9 @@ void UDataflowEditorVertexAttributePaintTool::Setup()
 	UpdateVertexColorOverlay();
 	//DynamicMeshComponent->FastNotifyVertexAttributesUpdated(EMeshRenderAttributeFlags::VertexColors);
 	GetToolManager()->PostInvalidation();
-}
 
-void UDataflowEditorVertexAttributePaintTool::CreateDefaultColorRamp(FDataflowColorRamp& OutColorRamp)
-{
-	constexpr FLinearColor HeatMapColors[] =
-	{
-		FLinearColor(0.8f, 0.4f, 0.8f), // Purple
-		FLinearColor(0.0f, 0.0f, 0.5f), // Dark Blue
-		FLinearColor(0.2f, 0.2f, 1.0f), // Light Blue
-		FLinearColor(0.0f, 1.0f, 0.0f), // Green
-		FLinearColor(1.0f, 1.0f, 0.0f), // Yellow
-		FLinearColor(1.0f, 0.65f, 0.0f), // Orange
-		FLinearColor(1.0f, 0.0f, 0.0f, 0.0f), // Red
-	};
-	constexpr int32 NumColors = sizeof(HeatMapColors) / sizeof(FLinearColor);
-	ensure(NumColors > 1);
-
-	// Set ColorRamp
-	FDataflowColorCurveOwner& Ramp = OutColorRamp.ColorRamp;
-	for (int32 Index = 0; Index < NumColors; ++Index)
-	{
-		const float Time = (float)Index / ((float)NumColors - 1.f);
-		Ramp.GetCurves()[0].CurveToEdit->AddKey(Time, HeatMapColors[Index].R);
-		Ramp.GetCurves()[1].CurveToEdit->AddKey(Time, HeatMapColors[Index].G);
-		Ramp.GetCurves()[2].CurveToEdit->AddKey(Time, HeatMapColors[Index].B);
-	}
-	Ramp.GetCurves()[3].CurveToEdit->AddKey(0.f, 1.f); // single alpha  value
+	ToolProperties->LoadConfig();
+	SetBrushMode(ToolProperties->BrushProperties.BrushMode);
 }
 
 void UDataflowEditorVertexAttributePaintTool::SetDataflowEditorContextObject(TObjectPtr<UDataflowContextObject> InDataflowEditorContextObject)
@@ -793,20 +824,6 @@ void UDataflowEditorVertexAttributePaintTool::RegisterActions(FInteractiveToolAc
 		[this]() { bPendingPickWeight = true; });
 };
 
-
-TUniquePtr<FMeshSculptBrushOp>& UDataflowEditorVertexAttributePaintTool::GetActiveBrushOp()
-{
-	if (GetInEraseStroke())
-	{
-		return SecondaryBrushOp;
-	}
-	else
-	{
-		return PrimaryBrushOp;
-	}
-}
-
-
 void UDataflowEditorVertexAttributePaintTool::OnPropertyModified(UObject* PropertySet, FProperty* Property)
 {
 	CalculateBrushRadius();
@@ -856,31 +873,28 @@ void UDataflowEditorVertexAttributePaintTool::SetBrushMode(EDataflowEditorToolEd
 	if (ToolProperties)
 	{
 		const EDataflowEditorToolEditOperation CurrentBrushMode = ToolProperties->BrushProperties.BrushMode;
-		if (NewBrushMode != CurrentBrushMode)
-		{
-			ToolProperties->BrushProperties.BrushMode = NewBrushMode;
-
-			// Save the current brush config 
-			FDataflowEditorVertexAttributePaintToolBrushConfig BrushConfigToSave;
-			BrushConfigToSave.BrushSize = ToolProperties->BrushProperties.BrushSize;
-			BrushConfigToSave.Value = ToolProperties->BrushProperties.AttributeValue;
-			ToolProperties->SetBrushConfig(CurrentBrushMode, BrushConfigToSave);
-
-			// Load the new brush config 
-			const FDataflowEditorVertexAttributePaintToolBrushConfig& NewBrushConfig = ToolProperties->GetBrushConfig(NewBrushMode);
-			ToolProperties->BrushProperties.BrushSize = NewBrushConfig.BrushSize;
-			ToolProperties->BrushProperties.AttributeValue = NewBrushConfig.Value;
-
-			// Send notification for the changed properties
-			FPropertyChangedEvent BrushModePropertyChangedEvent(UBrushBaseProperties::StaticClass()->FindPropertyByName(GET_MEMBER_NAME_CHECKED(FDataflowEditorVertexAttributePaintToolBrushProperties, BrushMode)));
-			ToolProperties->PostEditChangeProperty(BrushModePropertyChangedEvent);
-			FPropertyChangedEvent BrushSizePropertyChangedEvent(UBrushBaseProperties::StaticClass()->FindPropertyByName(GET_MEMBER_NAME_CHECKED(FDataflowEditorVertexAttributePaintToolBrushProperties, BrushSize)));
-			ToolProperties->PostEditChangeProperty(BrushSizePropertyChangedEvent);
-			FPropertyChangedEvent ValuePropertyChangedEvent(UBrushBaseProperties::StaticClass()->FindPropertyByName(GET_MEMBER_NAME_CHECKED(FDataflowEditorVertexAttributePaintToolBrushProperties, AttributeValue)));
-			ToolProperties->PostEditChangeProperty(ValuePropertyChangedEvent);
-
-			ToolProperties->SaveConfig();
-		}
+		ToolProperties->BrushProperties.BrushMode = NewBrushMode;
+
+		// Save the current brush config 
+		FDataflowEditorVertexAttributePaintToolBrushConfig BrushConfigToSave;
+		BrushConfigToSave.BrushSize = ToolProperties->BrushProperties.BrushSize;
+		BrushConfigToSave.Value = ToolProperties->BrushProperties.AttributeValue;
+		ToolProperties->SetBrushConfig(CurrentBrushMode, BrushConfigToSave);
+
+		// Load the new brush config 
+		const FDataflowEditorVertexAttributePaintToolBrushConfig& NewBrushConfig = ToolProperties->GetBrushConfig(NewBrushMode);
+		ToolProperties->BrushProperties.BrushSize = NewBrushConfig.BrushSize;
+		ToolProperties->BrushProperties.AttributeValue = NewBrushConfig.Value;
+
+		// Send notification for the changed properties
+		FPropertyChangedEvent BrushModePropertyChangedEvent(UBrushBaseProperties::StaticClass()->FindPropertyByName(GET_MEMBER_NAME_CHECKED(FDataflowEditorVertexAttributePaintToolBrushProperties, BrushMode)));
+		ToolProperties->PostEditChangeProperty(BrushModePropertyChangedEvent);
+		FPropertyChangedEvent BrushSizePropertyChangedEvent(UBrushBaseProperties::StaticClass()->FindPropertyByName(GET_MEMBER_NAME_CHECKED(FDataflowEditorVertexAttributePaintToolBrushProperties, BrushSize)));
+		ToolProperties->PostEditChangeProperty(BrushSizePropertyChangedEvent);
+		FPropertyChangedEvent ValuePropertyChangedEvent(UBrushBaseProperties::StaticClass()->FindPropertyByName(GET_MEMBER_NAME_CHECKED(FDataflowEditorVertexAttributePaintToolBrushProperties, AttributeValue)));
+		ToolProperties->PostEditChangeProperty(ValuePropertyChangedEvent);
+
+		ToolProperties->SaveConfig();
 	}
 }
 
@@ -913,10 +927,6 @@ void UDataflowEditorVertexAttributePaintTool::OnBeginStroke(const FRay& WorldRay
 		PaintBrushOpProperties->EditOperation = ToolProperties->BrushProperties.BrushMode;
 		PaintBrushOpProperties->Strength = 1.0f; // Attribute painting only need the attribute value 
 	}
-	if (EraseBrushOpProperties)
-	{
-		EraseBrushOpProperties->AttributeValue = 0.0;
-	}
 
 	// initialize first "Last Stamp", so that we can assume all stamps in stroke have a valid previous stamp
 	LastStamp.WorldFrame = GetBrushFrameWorld();
@@ -963,6 +973,7 @@ void UDataflowEditorVertexAttributePaintTool::OnCancelStroke()
 	{
 		UseBrushOp->CancelStroke();
 	}
+	CancelChange();
 }
 
 void UDataflowEditorVertexAttributePaintTool::UpdateROI(const FSculptBrushStamp& BrushStamp)
@@ -1378,8 +1389,8 @@ void UDataflowEditorVertexAttributePaintTool::OnPolyLassoFinished(const FCameraP
 	UE::Geometry::FMeshVertexSelection VertexSelection(Mesh);
 	VertexSelection.SelectByVertexID([&](int32 VertexIdx) { return TempROIBuffer[VertexIdx] == 1; });
 
-	const double SetWeightValue = GetInEraseStroke() ? 0.0 : ToolProperties->BrushProperties.AttributeValue;
-	SetVerticesToWeightMap(VertexSelection.AsSet(), SetWeightValue, GetInEraseStroke());
+	const double SetWeightValue = ToolProperties? ToolProperties->BrushProperties.AttributeValue: 0.0;
+	SetVerticesToWeightMap(VertexSelection.AsSet(), SetWeightValue);
 }
 
 
@@ -1723,7 +1734,7 @@ void UDataflowEditorVertexAttributePaintTool::OnSelectionModified()
 }
 
 
-void UDataflowEditorVertexAttributePaintTool::SetVerticesToWeightMap(const TSet<int32>& Vertices, double WeightValue, bool bIsErase)
+void UDataflowEditorVertexAttributePaintTool::SetVerticesToWeightMap(const TSet<int32>& Vertices, double WeightValue)
 {
 	BeginChange();
 
@@ -1908,9 +1919,6 @@ bool UDataflowEditorVertexAttributePaintTool::UpdateBrushPosition(const FRay& Wo
 	return bHit;
 }
 
-
-
-
 bool UDataflowEditorVertexAttributePaintTool::OnUpdateHover(const FInputDeviceRay& DevicePos)
 {
 	if (ensure(InStroke() == false))
@@ -1925,8 +1933,7 @@ void UDataflowEditorVertexAttributePaintTool::DrawHUD(FCanvas* Canvas, IToolsCon
 {
 	if (PolyLassoMechanic)
 	{
-		// because the actual Weight change is deferred until mouse release, color the lasso to let the user know whether it will erase
-		PolyLassoMechanic->LineColor = GetInEraseStroke() ? FLinearColor::Red : FLinearColor::Green;
+		PolyLassoMechanic->LineColor = FLinearColor::Green;
 		PolyLassoMechanic->DrawHUD(Canvas, RenderAPI);
 	}
 
@@ -2137,6 +2144,12 @@ void UDataflowEditorVertexAttributePaintTool::EndChange()
 	LongTransactions.Close(GetToolManager());
 }
 
+void UDataflowEditorVertexAttributePaintTool::CancelChange()
+{
+	VertexData.CancelChange();
+	LongTransactions.Close(GetToolManager());
+}
+
 
 void UDataflowEditorVertexAttributePaintTool::WaitForPendingUndoRedoUpdate()
 {
@@ -2237,7 +2250,7 @@ int32 UDataflowEditorVertexAttributePaintTool::GetBrushNearestVertex() const
 
 void UDataflowEditorVertexAttributePaintTool::UpdateBrushType(EDataflowEditorToolEditOperation BrushMode)
 {
-	static const FText BaseMessage = LOCTEXT("OnStartTool", "Hold Shift to Erase. Ctrl to Invert Value, [/] and S/D change Size (+Shift to small-step)");
+	static const FText BaseMessage = LOCTEXT("OnStartTool", "Hold [Shift] to Relax values, Hold [Ctrl] to Invert Value, [/] and S/D change Size (+Shift to small-step)");
 	FTextBuilder Builder;
 	Builder.AppendLine(BaseMessage);
 
@@ -2246,10 +2259,10 @@ void UDataflowEditorVertexAttributePaintTool::UpdateBrushType(EDataflowEditorToo
 		ToolProperties->BrushProperties.BrushMode = BrushMode;
 	}
 	SetActivePrimaryBrushType(PaintBrushId);
-	SetActiveSecondaryBrushType(EraseBrushId);
+	SetActiveSecondaryBrushType(SmoothBrushId);
 
 	SetToolPropertySourceEnabled(PaintBrushOpProperties, false);
-	SetToolPropertySourceEnabled(EraseBrushOpProperties, false);
+	SetToolPropertySourceEnabled(SmoothBrushOpProperties, false);
 	SetToolPropertySourceEnabled(GizmoProperties, false);
 
 	GetToolManager()->DisplayMessage(Builder.ToText(), EToolMessageLevel::UserNotification);
diff --git a/Engine/Plugins/Experimental/Dataflow/Source/DataflowEditor/Public/DataflowEditorTools/DataflowEditorVertexAttributePaintTool.h b/Engine/Plugins/Experimental/Dataflow/Source/DataflowEditor/Public/DataflowEditorTools/DataflowEditorVertexAttributePaintTool.h
--- a/Engine/Plugins/Experimental/Dataflow/Source/DataflowEditor/Public/DataflowEditorTools/DataflowEditorVertexAttributePaintTool.h
+++ b/Engine/Plugins/Experimental/Dataflow/Source/DataflowEditor/Public/DataflowEditorTools/DataflowEditorVertexAttributePaintTool.h
@@ -34,7 +34,7 @@
 #define UE_API DATAFLOWEDITOR_API
 
 class UMeshElementsVisualizer;
-class UDataflowWeightMapEraseBrushOpProps;
+class UDataflowWeightMapSmoothBrushOpProps;
 class UDataflowVertexAttributePaintBrushOpProps;
 class UDataflowContextObject;
 struct FDataflowVertexAttributeEditableNode;
@@ -192,6 +192,8 @@ class UDataflowEditorVertexAttributePaintToolProperties : public UInteractiveToo
 {
 	GENERATED_BODY()
 
+	UE_API UDataflowEditorVertexAttributePaintToolProperties();
+
 public:
 
 	// Edit mode (brush,mesh...)
@@ -225,6 +227,9 @@ public:
 	UPROPERTY(config)
 	FDataflowEditorVertexAttributePaintToolBrushConfig BrushConfigRelax;
 
+private:
+	void CreateDefaultColorRamp();
+
 };
 
 ///////////////////////////////////////////////////////////////////////////////////////////////////////////////////////////////////////////////////
@@ -237,6 +242,7 @@ public:
 
 	void BeginChange(UDynamicMeshComponent* Component);
 	TUniquePtr<UE::DataflowEditorVertexAttributePaintTool::Private::FMeshChange> EndChange(UDynamicMeshComponent* Component);
+	void CancelChange();
 
 	bool IsValid() const { return (ActiveWeightMap != nullptr); }
 
@@ -335,7 +341,7 @@ public:
 	TObjectPtr<UDataflowVertexAttributePaintBrushOpProps> PaintBrushOpProperties;
 
 	UPROPERTY()
-	TObjectPtr<UDataflowWeightMapEraseBrushOpProps> EraseBrushOpProperties;
+	TObjectPtr<UDataflowWeightMapSmoothBrushOpProps> SmoothBrushOpProperties;
 
 	UE_API bool HaveVisibilityFilter() const;
 	UE_API void ApplyVisibilityFilter(const TArray<int32>& Vertices, TArray<int32>& VisibleVertices);
@@ -384,10 +390,6 @@ protected:
 
 	UE_API void OnColorRampChanged(TArray<FRichCurve*> Curves);
 
-	void CreateDefaultColorRamp(FDataflowColorRamp& OutColorRamp);
-
-	UE_API virtual TUniquePtr<FMeshSculptBrushOp>& GetActiveBrushOp();
-
 	virtual bool SharesBrushPropertiesChanges() const override { return false; }
 	// end UMeshSculptToolBase API
 
@@ -461,23 +463,17 @@ protected:
 	UE_API void UpdateBrushType(EDataflowEditorToolEditOperation EditOperation);
 	UE_API bool UpdateBrushPosition(const FRay& WorldRay);
 
-	bool GetInEraseStroke()
-	{
-		// Re-use the smoothing stroke key (shift) for erase stroke in the weight paint tool
-		return GetInSmoothingStroke();
-	}
-
-
 	bool bPendingPickWeight = false;
 
 	TArray<int32> ROITriangleBuffer;
 	TArray<double> ROIWeightValueBuffer;
 	UE_API bool SyncMeshWithWeightBuffer(FDynamicMesh3* Mesh);
 	UE_API bool SyncWeightBufferWithMesh(const FDynamicMesh3* Mesh);
-	UE_API void SetVerticesToWeightMap(const TSet<int32>& Vertices, double WeightValue, bool bIsErase);
+	UE_API void SetVerticesToWeightMap(const TSet<int32>& Vertices, double WeightValue);
 
 	UE_API void BeginChange();
 	UE_API void EndChange();
+	UE_API void CancelChange();
 
 	TArray<FVector3d> TriNormals;
 	TArray<int32> UVSeamEdges;
@@ -499,7 +495,7 @@ protected:
 	UDataflowEditorMode* Mode = nullptr;
 
 	static constexpr int32 PaintBrushId = 0;
-	static constexpr int32 EraseBrushId = 1;
+	static constexpr int32 SmoothBrushId = 1;
 };
 
 
diff --git a/Engine/Source/Runtime/Experimental/Dataflow/Engine/Private/Dataflow/DataflowColorRamp.cpp b/Engine/Source/Runtime/Experimental/Dataflow/Engine/Private/Dataflow/DataflowColorRamp.cpp
--- a/Engine/Source/Runtime/Experimental/Dataflow/Engine/Private/Dataflow/DataflowColorRamp.cpp
+++ b/Engine/Source/Runtime/Experimental/Dataflow/Engine/Private/Dataflow/DataflowColorRamp.cpp
@@ -18,6 +18,11 @@ FDataflowColorCurveOwner::FDataflowColorCurveOwner()
 	}
 }
 
+bool FDataflowColorCurveOwner::IsEmpty() const
+{
+	return (RichCurves.Num() == 0 || RichCurves[0].IsEmpty());
+}
+
 FDataflowColorCurveOwner::FDataflowColorCurveOwner(const FDataflowColorCurveOwner& Other)
 {
 	this->operator=(Other);
diff --git a/Engine/Source/Runtime/Experimental/Dataflow/Engine/Public/Dataflow/DataflowColorRamp.h b/Engine/Source/Runtime/Experimental/Dataflow/Engine/Public/Dataflow/DataflowColorRamp.h
--- a/Engine/Source/Runtime/Experimental/Dataflow/Engine/Public/Dataflow/DataflowColorRamp.h
+++ b/Engine/Source/Runtime/Experimental/Dataflow/Engine/Public/Dataflow/DataflowColorRamp.h
@@ -20,6 +20,8 @@ public:
 	FDataflowColorCurveOwner(FDataflowColorCurveOwner&& Other) = delete;
 	FDataflowColorCurveOwner& operator=(FDataflowColorCurveOwner&& Other) = delete;
 
+	DATAFLOWENGINE_API bool IsEmpty() const;
+
 	/** implemenmt all the necessary virtual function from FCurveOwnerInterface */
 	DATAFLOWENGINE_API virtual TArray<FRichCurveEditInfoConst> GetCurves() const override;
 	DATAFLOWENGINE_API virtual void GetCurves(TAdderReserverRef<FRichCurveEditInfoConst> Curves) const override;
