Dataflow Editor : When restarting simulation on the timeline , do not reset the simulate button

diff --git a/Engine/Plugins/Experimental/Dataflow/Source/DataflowEditor/Private/Dataflow/DataflowSimulationScene.cpp b/Engine/Plugins/Experimental/Dataflow/Source/DataflowEditor/Private/Dataflow/DataflowSimulationScene.cpp
--- a/Engine/Plugins/Experimental/Dataflow/Source/DataflowEditor/Private/Dataflow/DataflowSimulationScene.cpp
+++ b/Engine/Plugins/Experimental/Dataflow/Source/DataflowEditor/Private/Dataflow/DataflowSimulationScene.cpp
@@ -194,6 +194,11 @@ bool FDataflowSimulationScene::IsSimulationEnabled() const
 	return false;
 }
 
+void FDataflowSimulationScene::SetSimulationEnabled(bool bEnable) const
+{
+	GetWorld()->GetSubsystem<UDataflowSimulationManager>()->SetSimulationEnabled(bEnable);
+}
+
 void FDataflowSimulationScene::RebuildSimulationScene()
 {
 	// Unregister components, cache manager, selection...
@@ -202,9 +207,6 @@ void FDataflowSimulationScene::RebuildSimulationScene()
 	// Register components, cache manager, selection...
 	CreateSimulationScene();
 
-	// Override the simulation enabled flag
-	GetWorld()->GetSubsystem<UDataflowSimulationManager>()->SetSimulationEnabled(true);
-
 	// Reset the simulation dirty flag
 	if(const TObjectPtr<UDataflowBaseContent>& EditorContent = GetEditorContent())
 	{
diff --git a/Engine/Plugins/Experimental/Dataflow/Source/DataflowEditor/Private/Widgets/DataflowSimulationBinding.cpp b/Engine/Plugins/Experimental/Dataflow/Source/DataflowEditor/Private/Widgets/DataflowSimulationBinding.cpp
--- a/Engine/Plugins/Experimental/Dataflow/Source/DataflowEditor/Private/Widgets/DataflowSimulationBinding.cpp
+++ b/Engine/Plugins/Experimental/Dataflow/Source/DataflowEditor/Private/Widgets/DataflowSimulationBinding.cpp
@@ -157,7 +157,9 @@ void FDataflowSimulationBinding::ResetSimulationScene()
 {
 	if (const TSharedPtr<FDataflowSimulationScene> PreviewScene = SimulationScene.Pin())
 	{
-		return PreviewScene->RebuildSimulationScene();
+		const bool bSimulationWasEnabled = PreviewScene->IsSimulationEnabled();
+		PreviewScene->RebuildSimulationScene();
+		PreviewScene->SetSimulationEnabled(bSimulationWasEnabled);
 	}
 }
 
diff --git a/Engine/Plugins/Experimental/Dataflow/Source/DataflowEditor/Public/Dataflow/DataflowSimulationScene.h b/Engine/Plugins/Experimental/Dataflow/Source/DataflowEditor/Public/Dataflow/DataflowSimulationScene.h
--- a/Engine/Plugins/Experimental/Dataflow/Source/DataflowEditor/Public/Dataflow/DataflowSimulationScene.h
+++ b/Engine/Plugins/Experimental/Dataflow/Source/DataflowEditor/Public/Dataflow/DataflowSimulationScene.h
@@ -134,6 +134,9 @@ public:
 	/** Gets whether the simulation is running */
 	UE_API bool IsSimulationEnabled() const;
 
+	/** Gets whether the simulation is running */
+	UE_API void SetSimulationEnabled(bool bEnable) const;
+
 	/** Check if caching is enabled or not */
 	UE_API bool HasCachingEnabled() const;
 
