Dataflow weight paint tool : UI and Performance improvements
* Color ramp now updates based on the selected color mode
* Use ParallelFor for default and smooth brushes
* Applying operations from selection now support smoothing properly and share code with the brush mode

diff --git a/Engine/Plugins/Experimental/Dataflow/Source/DataflowEditor/Private/DataflowEditorTools/DataflowEditorVertexAttributePaintTool.cpp b/Engine/Plugins/Experimental/Dataflow/Source/DataflowEditor/Private/DataflowEditorTools/DataflowEditorVertexAttributePaintTool.cpp
--- a/Engine/Plugins/Experimental/Dataflow/Source/DataflowEditor/Private/DataflowEditorTools/DataflowEditorVertexAttributePaintTool.cpp
+++ b/Engine/Plugins/Experimental/Dataflow/Source/DataflowEditor/Private/DataflowEditorTools/DataflowEditorVertexAttributePaintTool.cpp
@@ -263,6 +263,13 @@ UDataflowEditorVertexAttributePaintToolProperties::UDataflowEditorVertexAttribut
 	{
 		CreateDefaultColorRamp();
 	}
+
+	constexpr bool bOnlyRGB = true;
+	DisplayProperties.GreyScaleColorRamp.SetColorAtTime(0.0f, FLinearColor::Black, bOnlyRGB);
+	DisplayProperties.GreyScaleColorRamp.SetColorAtTime(1.0f, FLinearColor::White, bOnlyRGB);
+
+	DisplayProperties.WhiteColorRamp.SetColorAtTime(0.0f, FLinearColor::White, bOnlyRGB);
+	DisplayProperties.WhiteColorRamp.SetColorAtTime(1.0f, FLinearColor::White, bOnlyRGB);
 }
 
 void UDataflowEditorVertexAttributePaintToolProperties::CreateDefaultColorRamp()
@@ -285,9 +292,7 @@ void UDataflowEditorVertexAttributePaintToolProperties::CreateDefaultColorRamp()
 	for (int32 Index = 0; Index < NumColors; ++Index)
 	{
 		const float Time = (float)Index / ((float)NumColors - 1.f);
-		Ramp.GetCurves()[0].CurveToEdit->AddKey(Time, HeatMapColors[Index].R);
-		Ramp.GetCurves()[1].CurveToEdit->AddKey(Time, HeatMapColors[Index].G);
-		Ramp.GetCurves()[2].CurveToEdit->AddKey(Time, HeatMapColors[Index].B);
+		Ramp.SetColorAtTime(Time, HeatMapColors[Index], /*bOnlyRGB*/true);
 	}
 	Ramp.GetCurves()[3].CurveToEdit->AddKey(0.f, 1.f); // single alpha  value
 }
@@ -441,7 +446,7 @@ void FDataflowEditorVertexAttributePaintToolData::Shutdown()
 
 void FDataflowEditorVertexAttributePaintToolData::BeginChange(UDynamicMeshComponent* Component)
 {
-	check(ActiveWeightEditChangeTracker == nullptr);
+	ensure(ActiveWeightEditChangeTracker == nullptr);
 	if (Component)
 	{
 		ActiveWeightEditChangeTracker = MakeUnique<UE::Geometry::FDynamicMeshChangeTracker>(Component->GetMesh());
@@ -451,9 +456,7 @@ void FDataflowEditorVertexAttributePaintToolData::BeginChange(UDynamicMeshCompon
 
 TUniquePtr<UE::DataflowEditorVertexAttributePaintTool::Private::FMeshChange> FDataflowEditorVertexAttributePaintToolData::EndChange(UDynamicMeshComponent* Component)
 {
-	check(ActiveWeightEditChangeTracker);
-
-	if (ActiveWeightEditChangeTracker && Component)
+	if (ensure(ActiveWeightEditChangeTracker) && Component)
 	{
 		using namespace UE::DataflowEditorVertexAttributePaintTool;
 		TUniquePtr<UE::Geometry::FDynamicMeshChange> EditResult = ActiveWeightEditChangeTracker->EndChange();
@@ -467,9 +470,7 @@ TUniquePtr<UE::DataflowEditorVertexAttributePaintTool::Private::FMeshChange> FDa
 
 void FDataflowEditorVertexAttributePaintToolData::CancelChange()
 {
-	check(ActiveWeightEditChangeTracker);
-
-	if (ActiveWeightEditChangeTracker)
+	if (ensure(ActiveWeightEditChangeTracker))
 	{
 		ActiveWeightEditChangeTracker->EndChange();
 		ActiveWeightEditChangeTracker = nullptr;
@@ -1553,49 +1554,25 @@ void UDataflowEditorVertexAttributePaintTool::ApplyValueToSelection(EDataflowEdi
 		return;
 	}
 
-	TArray<int32> SelectVertices = UDataflowEditorVertexAttributePaintTool::GetSelectedVertices();
-	if (SelectVertices.IsEmpty())
+	VertexROI = UDataflowEditorVertexAttributePaintTool::GetSelectedVertices();
+	if (VertexROI.IsEmpty())
 	{
 		return;
 	}
 
 	BeginChange();
 
-	auto UpdateAttributeValue = 
-		[Operation, InValue](UE::Geometry::FDynamicMeshWeightAttribute& AttributeMap, int32 AttributeIndex) -> double
-		{
-			float CurrentValue;
-			AttributeMap.GetValue(AttributeIndex, &CurrentValue);
-			float NewValue = CurrentValue;
-			switch (Operation)
-			{
-			case EDataflowEditorToolEditOperation::Add:
-				NewValue = (CurrentValue + InValue);
-				break;
-			case EDataflowEditorToolEditOperation::Replace:
-				NewValue = InValue;
-				break;
-			case EDataflowEditorToolEditOperation::Multiply:
-				NewValue = (CurrentValue * InValue);;
-				break;
-			case EDataflowEditorToolEditOperation::Invert:
-				NewValue = (1.0f - CurrentValue);
-				break;
-			case EDataflowEditorToolEditOperation::Relax:
-				// we should not reach there since the relax operation should be handled by FDataflowWeightMapSmoothBrushOp
-				ensure(false);
-				break;
-			}
-			return FMath::Clamp(NewValue, 0.f, 1.f);
-		};
-
-	for (int32 VertexIdx : SelectVertices)
+	if (FDynamicMesh3* Mesh = GetSculptMesh())
 	{
-		VertexData.SetValue(VertexIdx, UpdateAttributeValue);
+		ROIWeightValueBuffer.SetNum(VertexROI.Num(), EAllowShrinking::No);
+		SyncWeightBufferWithMesh(Mesh);
+
+		FDataflowVertexAttributePaintBrushOp::ApplyStampByVerticesStatic(
+			Mesh, CurrentStamp, VertexROI, ROIWeightValueBuffer, 
+			Operation, InValue, /*bApplyRadiusLimit*/false);
+
+		SyncMeshWithWeightBuffer(Mesh);
 	}
-	
-	// TODO(ccaillaud) - make it work with the relax option - brush should be available other than through the general interface
-	//SmoothBrushOpProperties
 
 	// update colors
 	UpdateVertexColorOverlay();
@@ -2286,6 +2263,9 @@ void UDataflowEditorVertexAttributePaintTool::SetColorMode(EDataflowEditorToolCo
 		ToolProperties->SaveConfig();
 		UpdateVertexColorOverlay();
 		GetToolManager()->PostInvalidation();
+
+		FPropertyChangedEvent ColorModePropertyChangedEvent(UBrushBaseProperties::StaticClass()->FindPropertyByName(GET_MEMBER_NAME_CHECKED(FDataflowEditorVertexAttributePaintToolDisplayProperties, ColorMode)));
+		ToolProperties->PostEditChangeProperty(ColorModePropertyChangedEvent);
 	}
 }
 
diff --git a/Engine/Plugins/Experimental/Dataflow/Source/DataflowEditor/Private/DataflowEditorTools/DataflowEditorVertexAttributePaintToolDetailCustomization.cpp b/Engine/Plugins/Experimental/Dataflow/Source/DataflowEditor/Private/DataflowEditorTools/DataflowEditorVertexAttributePaintToolDetailCustomization.cpp
--- a/Engine/Plugins/Experimental/Dataflow/Source/DataflowEditor/Private/DataflowEditorTools/DataflowEditorVertexAttributePaintToolDetailCustomization.cpp
+++ b/Engine/Plugins/Experimental/Dataflow/Source/DataflowEditor/Private/DataflowEditorTools/DataflowEditorVertexAttributePaintToolDetailCustomization.cpp
@@ -136,25 +136,69 @@ namespace UE::Dataflow::Editor
 
 	void FVertexAttributePaintToolDetailCustomization::AddColorRampRow(IDetailCategoryBuilder& EditModeCategory) const
 	{
-		TSharedRef<SColorGradientEditor> GradientEditor = SNew(SColorGradientEditor)
-			.ViewMinInput(0.0f)
-			.ViewMaxInput(1.0f)
-			.ClampStopsToViewRange(true)
-			.DrawColorAndAlphaSeparate(false)
-			.IsEditingEnabled_Lambda([this]()
+		auto MakeColorRampWidget =
+			[](
+				EDataflowEditorToolColorMode ColorMode, 
+				FCurveOwnerInterface* CurveOwner,
+				TWeakObjectPtr<UDataflowEditorVertexAttributePaintToolProperties> ToolProperties
+				)
+			{
+				TSharedRef<SColorGradientEditor> GradientEditor = SNew(SColorGradientEditor)
+					.ViewMinInput(0.0f)
+					.ViewMaxInput(1.0f)
+					.ClampStopsToViewRange(true)
+					.DrawColorAndAlphaSeparate(false)
+					.IsEditingEnabled_Lambda([ToolProperties, ColorMode]()
+						{
+							// Only the ramp one is editable 
+							return ToolProperties.IsValid() && (ToolProperties->DisplayProperties.ColorMode == EDataflowEditorToolColorMode::Ramp);
+						})
+					.Visibility_Lambda([ToolProperties, ColorMode]()
+						{
+							return (ToolProperties.IsValid() && (ToolProperties->DisplayProperties.ColorMode == ColorMode))
+								? EVisibility::Visible
+								: EVisibility::Collapsed;
+						})
+					;
+				if (ToolProperties.IsValid())
 				{
-					return ToolProperties.IsValid() && (ToolProperties->DisplayProperties.ColorMode == EDataflowEditorToolColorMode::Ramp);
-				});
+					GradientEditor->SetCurveOwner(CurveOwner);
+				}
+				return GradientEditor;
+			};
+
 		if (ToolProperties.IsValid())
 		{
-			GradientEditor->SetCurveOwner(&ToolProperties->DisplayProperties.ColorRamp.ColorRamp);
-		}
+			TSharedRef<SColorGradientEditor> GreyScaleRampWidget =
+				MakeColorRampWidget(EDataflowEditorToolColorMode::Greyscale, &ToolProperties->DisplayProperties.GreyScaleColorRamp, ToolProperties);
 
-		EditModeCategory.AddCustomRow(LOCTEXT("ColorRampCategory", "Color Ramp"), false)
-		.WholeRowContent()
-		[
-			GradientEditor
-		];
+			TSharedRef<SColorGradientEditor> CustomRampWidget =
+				MakeColorRampWidget(EDataflowEditorToolColorMode::Ramp, &ToolProperties->DisplayProperties.ColorRamp.ColorRamp, ToolProperties);
+
+			TSharedRef<SColorGradientEditor> MaterialRampWidget =
+				MakeColorRampWidget(EDataflowEditorToolColorMode::FullMaterial, &ToolProperties->DisplayProperties.WhiteColorRamp, ToolProperties);
+
+			EditModeCategory.AddCustomRow(LOCTEXT("ColorRampCategory", "Color Ramp"), false)
+			.WholeRowContent()
+			[
+				SNew(SVerticalBox)
+				+ SVerticalBox::Slot()
+				.AutoHeight()
+				[
+					GreyScaleRampWidget
+				]
+				+ SVerticalBox::Slot()
+				.AutoHeight()
+				[
+					CustomRampWidget
+				]
+				+ SVerticalBox::Slot()
+				.AutoHeight()
+				[
+					MaterialRampWidget
+				]
+			];
+		}
 	}
 
 	///////////////////////////////////////////////////////////////////////////////////////////////////////////////
diff --git a/Engine/Plugins/Experimental/Dataflow/Source/DataflowEditor/Private/DataflowEditorTools/DataflowEditorWeightMapPaintBrushOps.cpp b/Engine/Plugins/Experimental/Dataflow/Source/DataflowEditor/Private/DataflowEditorTools/DataflowEditorWeightMapPaintBrushOps.cpp
--- a/Engine/Plugins/Experimental/Dataflow/Source/DataflowEditor/Private/DataflowEditorTools/DataflowEditorWeightMapPaintBrushOps.cpp
+++ b/Engine/Plugins/Experimental/Dataflow/Source/DataflowEditor/Private/DataflowEditorTools/DataflowEditorWeightMapPaintBrushOps.cpp
@@ -3,7 +3,6 @@
 #include "DataflowEditorTools/DataflowEditorWeightMapPaintBrushOps.h"
 #include "Async/ParallelFor.h"
 #include "DynamicMesh/DynamicMesh3.h"
-
 //
 // Paint Brush
 //
@@ -86,9 +85,19 @@ void FDataflowWeightMapSmoothBrushOp::ApplyStampByVertices(
 	const FSculptBrushStamp& Stamp,
 	const TArray<int32>& Vertices,
 	TArray<double>& VertexWeightValues)
+{
+	ApplyStampByVerticesStatic(Mesh, Stamp, Vertices, VertexWeightValues, bApplyRadiusLimit);
+}
+
+
+void FDataflowWeightMapSmoothBrushOp::ApplyStampByVerticesStatic(
+	const FDynamicMesh3* Mesh,
+	const FSculptBrushStamp& Stamp,
+	const TArray<int32>& Vertices,
+	TArray<double>& VertexWeightValues,
+	bool bApplyRadiusLimit)
 {
 	// Converted from FClothPaintTool_Smooth::SmoothVertices
-	
 	const int32 NumVertices = Vertices.Num();
 
 	TMap<int32, int32> VertexToBufferIndexMap;
@@ -102,50 +111,59 @@ void FDataflowWeightMapSmoothBrushOp::ApplyStampByVertices(
 	TArray<double> OneRingAverages;
 	OneRingAverages.SetNumUninitialized(NumVertices);
 
-	for (int32 BufferIndex = 0; BufferIndex < NumVertices; ++BufferIndex)
-	{
-		const int32 VertexIndex = Vertices[BufferIndex];
-		double Accumulator = 0.0f;
-		int NumNeighbors = 0;
-
-		for (const int32 NeighborIndex : Mesh->VtxVerticesItr(VertexIndex))
+	const TMap<int32, int32>& VertexToBufferIndexMapConst = VertexToBufferIndexMap;
+	const TArray<double>& VertexWeightValuesConst = VertexWeightValues;
+	ParallelFor(NumVertices, 
+		[&Vertices, &Mesh, &VertexToBufferIndexMapConst, &VertexWeightValuesConst, &OneRingAverages]
+		(int32 BufferIndex)
 		{
-			if (VertexToBufferIndexMap.Contains(NeighborIndex))
+			const int32 VertexIndex = Vertices[BufferIndex];
+			double Accumulator = 0.0f;
+			int32 NumNeighbors = 0;
+
+			for (const int32 NeighborIndex : Mesh->VtxVerticesItr(VertexIndex))
 			{
-				const int32 NeighborBufferIndex = VertexToBufferIndexMap[NeighborIndex];
-				Accumulator += VertexWeightValues[NeighborBufferIndex];
-				++NumNeighbors;
+				if (VertexToBufferIndexMapConst.Contains(NeighborIndex))
+				{
+					const int32 NeighborBufferIndex = VertexToBufferIndexMapConst[NeighborIndex];
+					Accumulator += VertexWeightValuesConst[NeighborBufferIndex];
+					++NumNeighbors;
+				}
 			}
-		}
 
-		if (NumNeighbors > 0)
-		{
-			OneRingAverages[BufferIndex] = Accumulator / NumNeighbors;
-		}
-		else
-		{
-			// Don't change the vertex value if it has no neighbors
-			OneRingAverages[BufferIndex] = VertexWeightValues[BufferIndex];
-		}
-	}
+			if (NumNeighbors > 0)
+			{
+				OneRingAverages[BufferIndex] = Accumulator / NumNeighbors;
+			}
+			else
+			{
+				// Don't change the vertex value if it has no neighbors
+				OneRingAverages[BufferIndex] = VertexWeightValuesConst[BufferIndex];
+			}
+		}, 
+		EParallelForFlags::None
+	);
 
 	// Blend vertex value with its average one-ring value
-	for (int32 BufferIndex = 0; BufferIndex < NumVertices; ++BufferIndex)
-	{
-		const double Diff = OneRingAverages[BufferIndex] - VertexWeightValues[BufferIndex];
-		if (bApplyRadiusLimit)
+	ParallelFor(NumVertices,
+		[&OneRingAverages, &VertexWeightValues, bApplyRadiusLimit, &Stamp, &Vertices, &Mesh](int32 BufferIndex)
 		{
-			const FVector3d& StampPos = Stamp.LocalFrame.Origin;
-			const int32 VertexIndex = Vertices[BufferIndex];
-			const FVector3d VertexPos = Mesh->GetVertex(VertexIndex);
-			const double DistanceSquared = (VertexPos - StampPos).SquaredLength();
-			if (DistanceSquared >= Stamp.Radius * Stamp.Radius)
+			const double Diff = OneRingAverages[BufferIndex] - VertexWeightValues[BufferIndex];
+			if (bApplyRadiusLimit)
 			{
-				continue;
+				const FVector3d& StampPos = Stamp.LocalFrame.Origin;
+				const int32 VertexIndex = Vertices[BufferIndex];
+				const FVector3d VertexPos = Mesh->GetVertex(VertexIndex);
+				const double DistanceSquared = (VertexPos - StampPos).SquaredLength();
+				if (DistanceSquared >= Stamp.Radius * Stamp.Radius)
+				{
+					return;
+				}
 			}
-		}
-		VertexWeightValues[BufferIndex] += Stamp.Power * Diff;
-	}
+			VertexWeightValues[BufferIndex] += Stamp.Power * Diff;
+		}, 
+		EParallelForFlags::None
+	);
 }
 
 ////////////////////////////////////////////////////////////////////////////////////////////////////////////////////
@@ -160,33 +178,48 @@ void FDataflowVertexAttributePaintBrushOp::ApplyStampByVertices(
 
 	if (UDataflowVertexAttributePaintBrushOpProps* Props = GetPropertySetAs<UDataflowVertexAttributePaintBrushOpProps>())
 	{
-		// special case for Relax option 
-		if (Props->EditOperation == EDataflowEditorToolEditOperation::Relax)
-		{
-			SmoothBrush.ApplyStampByVertices(Mesh, Stamp, Vertices, OutValues);
-			return;
-		}
+		ApplyStampByVerticesStatic(Mesh, Stamp, Vertices, OutValues, Props->EditOperation, Props->GetAttribute(), bApplyRadiusLimit);
+	}
+}
+
+/*static*/ void FDataflowVertexAttributePaintBrushOp::ApplyStampByVerticesStatic(
+	const FDynamicMesh3* Mesh,
+	const FSculptBrushStamp& Stamp,
+	const TArray<int32>& Vertices,
+	TArray<double>& OutValues,
+	EDataflowEditorToolEditOperation EditOperation,
+	const double ToolValue,
+	bool bApplyRadiusLimit)
+{
+	// special case for Relax option 
+	if (EditOperation == EDataflowEditorToolEditOperation::Relax)
+	{
+		FDataflowWeightMapSmoothBrushOp::ApplyStampByVerticesStatic(Mesh, Stamp, Vertices, OutValues, bApplyRadiusLimit);
+		return;
+	}
 
-		// other operations
-		for (int32 VtxIdx = 0; VtxIdx < Vertices.Num(); ++VtxIdx)
+	// other operations
+	const int32 NumVertices = Vertices.Num();
+	ParallelFor(NumVertices, 
+		[bApplyRadiusLimit, ToolValue, EditOperation , &Stamp, &Mesh, &Vertices, &OutValues]
+		(int32 BufferIndex)
 		{
 			if (bApplyRadiusLimit)
 			{
 				const FVector3d& StampPos = Stamp.LocalFrame.Origin;
-				const int32 MeshVertIdx = Vertices[VtxIdx];
+				const int32 MeshVertIdx = Vertices[BufferIndex];
 				const FVector3d VertexPos = Mesh->GetVertex(MeshVertIdx);
 				const double DistanceSquared = (VertexPos - StampPos).SquaredLength();
 				if (DistanceSquared >= Stamp.Radius * Stamp.Radius)
 				{
-					continue;
+					return;
 				}
 			}
-			
-			const double ToolValue = Props->GetAttribute();
-			const double CurrentValue = OutValues[VtxIdx];
+
+			const double CurrentValue = OutValues[BufferIndex];
 			double NewValue = CurrentValue;
 
-			switch (Props->EditOperation)
+			switch (EditOperation)
 			{
 			case EDataflowEditorToolEditOperation::Add:
 				NewValue = CurrentValue + ToolValue;
@@ -204,7 +237,8 @@ void FDataflowVertexAttributePaintBrushOp::ApplyStampByVertices(
 				ensure(false); // already handled by SmoothBrush.ApplyStampByVertices
 			}
 
-			OutValues[VtxIdx] = FMath::Clamp(NewValue, 0.0, 1.0);
-		}
-	}
+			OutValues[BufferIndex] = FMath::Clamp(NewValue, 0.0, 1.0);
+		},
+		EParallelForFlags::None
+	);
 }
\ No newline at end of file
diff --git a/Engine/Plugins/Experimental/Dataflow/Source/DataflowEditor/Private/DataflowEditorTools/DataflowEditorWeightMapPaintBrushOps.h b/Engine/Plugins/Experimental/Dataflow/Source/DataflowEditor/Private/DataflowEditorTools/DataflowEditorWeightMapPaintBrushOps.h
--- a/Engine/Plugins/Experimental/Dataflow/Source/DataflowEditor/Private/DataflowEditorTools/DataflowEditorWeightMapPaintBrushOps.h
+++ b/Engine/Plugins/Experimental/Dataflow/Source/DataflowEditor/Private/DataflowEditorTools/DataflowEditorWeightMapPaintBrushOps.h
@@ -153,6 +153,12 @@ public:
 		const TArray<int32>& Vertices,
 		TArray<double>& NewAttributesOut) override;
 
+	static void ApplyStampByVerticesStatic(
+		const FDynamicMesh3* Mesh,
+		const FSculptBrushStamp& Stamp,
+		const TArray<int32>& Vertices,
+		TArray<double>& NewAttributesOut,
+		bool bApplyRadiusLimit);
 };
 
 ////////////////////////////////////////////////////////////////////////////////////////////////////////////////////
@@ -194,7 +200,12 @@ public:
 		const TArray<int32>& Vertices,
 		TArray<double>& NewAttributesOut) override;
 
-private:
-	// share the code with SmoothBrush when EditOperation is Relax
-	FDataflowWeightMapSmoothBrushOp SmoothBrush;
+	static void ApplyStampByVerticesStatic(
+		const FDynamicMesh3* Mesh,
+		const FSculptBrushStamp& Stamp,
+		const TArray<int32>& Vertices,
+		TArray<double>& NewAttributesOut,
+		EDataflowEditorToolEditOperation EditOperation,
+		const double ToolValue,
+		bool bApplyRadiusLimit);
 };
diff --git a/Engine/Plugins/Experimental/Dataflow/Source/DataflowEditor/Public/Dataflow/DataflowSimulationScene.h b/Engine/Plugins/Experimental/Dataflow/Source/DataflowEditor/Public/Dataflow/DataflowSimulationScene.h
--- a/Engine/Plugins/Experimental/Dataflow/Source/DataflowEditor/Public/Dataflow/DataflowSimulationScene.h
+++ b/Engine/Plugins/Experimental/Dataflow/Source/DataflowEditor/Public/Dataflow/DataflowSimulationScene.h
@@ -134,7 +134,7 @@ public:
 	/** Gets whether the simulation is running */
 	UE_API bool IsSimulationEnabled() const;
 
-	/** Gets whether the simulation is running */
+	/** Sets whether the simulation should be enabled or not */
 	UE_API void SetSimulationEnabled(bool bEnable) const;
 
 	/** Check if caching is enabled or not */
diff --git a/Engine/Plugins/Experimental/Dataflow/Source/DataflowEditor/Public/DataflowEditorTools/DataflowEditorVertexAttributePaintTool.h b/Engine/Plugins/Experimental/Dataflow/Source/DataflowEditor/Public/DataflowEditorTools/DataflowEditorVertexAttributePaintTool.h
--- a/Engine/Plugins/Experimental/Dataflow/Source/DataflowEditor/Public/DataflowEditorTools/DataflowEditorVertexAttributePaintTool.h
+++ b/Engine/Plugins/Experimental/Dataflow/Source/DataflowEditor/Public/DataflowEditorTools/DataflowEditorVertexAttributePaintTool.h
@@ -183,6 +183,10 @@ struct FDataflowEditorVertexAttributePaintToolDisplayProperties
 	// Color ramp to use when the color mode is set to "Ramp"
 	UPROPERTY()
 	FDataflowColorRamp ColorRamp;
+
+	// Used by UI customization 
+	FDataflowColorCurveOwner GreyScaleColorRamp;
+	FDataflowColorCurveOwner WhiteColorRamp;
 };
 
 ///////////////////////////////////////////////////////////////////////////////////////////////////////////////////////////////////////////////////
diff --git a/Engine/Source/Runtime/Experimental/Dataflow/Engine/Private/Dataflow/DataflowColorRamp.cpp b/Engine/Source/Runtime/Experimental/Dataflow/Engine/Private/Dataflow/DataflowColorRamp.cpp
--- a/Engine/Source/Runtime/Experimental/Dataflow/Engine/Private/Dataflow/DataflowColorRamp.cpp
+++ b/Engine/Source/Runtime/Experimental/Dataflow/Engine/Private/Dataflow/DataflowColorRamp.cpp
@@ -47,6 +47,20 @@ FDataflowColorCurveOwner& FDataflowColorCurveOwner::operator=(const FDataflowCol
 	return *this;
 }
 
+void FDataflowColorCurveOwner::SetColorAtTime(float Time, const FLinearColor& Color, bool bOnlyRBG)
+{
+	if (ensure(Curves.Num() == 4))
+	{
+		Curves[0].CurveToEdit->AddKey(Time, Color.R);
+		Curves[1].CurveToEdit->AddKey(Time, Color.G);
+		Curves[2].CurveToEdit->AddKey(Time, Color.B);
+		if (!bOnlyRBG)
+		{
+			Curves[3].CurveToEdit->AddKey(Time, Color.A);
+		}
+	}
+}
+
 TArray<FRichCurveEditInfoConst> FDataflowColorCurveOwner::GetCurves() const
 {
 	return ConstCurves;
diff --git a/Engine/Source/Runtime/Experimental/Dataflow/Engine/Public/Dataflow/DataflowColorRamp.h b/Engine/Source/Runtime/Experimental/Dataflow/Engine/Public/Dataflow/DataflowColorRamp.h
--- a/Engine/Source/Runtime/Experimental/Dataflow/Engine/Public/Dataflow/DataflowColorRamp.h
+++ b/Engine/Source/Runtime/Experimental/Dataflow/Engine/Public/Dataflow/DataflowColorRamp.h
@@ -20,6 +20,8 @@ public:
 	FDataflowColorCurveOwner(FDataflowColorCurveOwner&& Other) = delete;
 	FDataflowColorCurveOwner& operator=(FDataflowColorCurveOwner&& Other) = delete;
 
+	DATAFLOWENGINE_API void SetColorAtTime(float Time, const FLinearColor& Color, bool bOnlyRBG);
+
 	DATAFLOWENGINE_API bool IsEmpty() const;
 
 	/** implemenmt all the necessary virtual function from FCurveOwnerInterface */
