Dataflow editor : disable camera view types based on the Dataflow view modes
- perspective view mode can chosse any camera view types
- otho view modes are constrained by the single view mode they need

diff --git a/Engine/Plugins/Experimental/Dataflow/Source/DataflowEditor/Private/Dataflow/DataflowConstructionViewport.cpp b/Engine/Plugins/Experimental/Dataflow/Source/DataflowEditor/Private/Dataflow/DataflowConstructionViewport.cpp
--- a/Engine/Plugins/Experimental/Dataflow/Source/DataflowEditor/Private/Dataflow/DataflowConstructionViewport.cpp
+++ b/Engine/Plugins/Experimental/Dataflow/Source/DataflowEditor/Private/Dataflow/DataflowConstructionViewport.cpp
@@ -18,6 +18,7 @@
 #include "Dataflow/DataflowConstructionVisualization.h"
 #include "ViewportToolbar/UnrealEdViewportToolbar.h"
 #include "Widgets/Text/SRichTextBlock.h"
+#include "EditorViewportCommands.h"
 
 #define LOCTEXT_NAMESPACE "SDataflowConstructionViewport"
 
@@ -236,7 +237,50 @@ void SDataflowConstructionViewport::BindCommands()
 				return false;
 			})
 		);
+
+		FUIAction InvalidAction;
+		InvalidAction.IsActionVisibleDelegate.BindLambda([]() { return false; });
 	}
+
+	// Now bind the visibility of the camnera view commands to the selected dataflow view mode perspective flag 
+	auto OverrideCameraViewVisibility = [this](const TSharedPtr<FUICommandInfo>& Command, ELevelViewportType ViewportType)
+		{
+			TWeakPtr<SDataflowConstructionViewport> WeakSelf = StaticCastWeakPtr<SDataflowConstructionViewport>(AsWeak());
+			if (CommandList)
+			{
+				if (const FUIAction* CameraViewAction = CommandList->GetActionForCommand(Command))
+				{
+					FUIAction OverridenCameraViewAction = *CameraViewAction;
+					OverridenCameraViewAction.CanExecuteAction = FCanExecuteAction::CreateLambda(
+						[ViewportType, WeakSelf]() -> bool
+						{
+							if (const TSharedPtr<SDataflowConstructionViewport> Viewport = WeakSelf.Pin())
+							{
+								if (UDataflowEditorMode* EditorMode = Viewport->GetEdMode())
+								{
+									if (const UE::Dataflow::IDataflowConstructionViewMode* CurrentViewMode = EditorMode->GetConstructionViewMode())
+									{
+										const ELevelViewportType ViewModeViewportType = CurrentViewMode->GetViewportType();
+										// perspective view mode can use any camera view type
+										// ortho view modes are constrained by their unique camera veiw type
+										return (ViewModeViewportType == ELevelViewportType::LVT_Perspective || ViewModeViewportType == ViewportType);
+									}
+								}
+							}
+							return true;
+						});
+					CommandList->MapAction(Command, OverridenCameraViewAction);
+				}
+			}
+		};
+	
+	OverrideCameraViewVisibility(FEditorViewportCommands::Get().Perspective, ELevelViewportType::LVT_Perspective);
+	OverrideCameraViewVisibility(FEditorViewportCommands::Get().Front, ELevelViewportType::LVT_OrthoFront);
+	OverrideCameraViewVisibility(FEditorViewportCommands::Get().Left, ELevelViewportType::LVT_OrthoLeft);
+	OverrideCameraViewVisibility(FEditorViewportCommands::Get().Top, ELevelViewportType::LVT_OrthoTop);
+	OverrideCameraViewVisibility(FEditorViewportCommands::Get().Back, ELevelViewportType::LVT_OrthoBack);
+	OverrideCameraViewVisibility(FEditorViewportCommands::Get().Right, ELevelViewportType::LVT_OrthoRight);
+	OverrideCameraViewVisibility(FEditorViewportCommands::Get().Bottom, ELevelViewportType::LVT_OrthoBottom);
 }
 
 bool SDataflowConstructionViewport::IsVisible() const
