Dataflow : fix potential crash when a deleted tool receives a notification from changed skeleton
- make sure the skinweight tool deregister the handler om bone selection on the node
- make the tool responsible for registering the handlers instead of the builder

diff --git a/Engine/Plugins/Experimental/Dataflow/Source/DataflowEditor/Private/DataflowEditorTools/DataflowEditorSkinWeightsPaintTool.cpp b/Engine/Plugins/Experimental/Dataflow/Source/DataflowEditor/Private/DataflowEditorTools/DataflowEditorSkinWeightsPaintTool.cpp
--- a/Engine/Plugins/Experimental/Dataflow/Source/DataflowEditor/Private/DataflowEditorTools/DataflowEditorSkinWeightsPaintTool.cpp
+++ b/Engine/Plugins/Experimental/Dataflow/Source/DataflowEditor/Private/DataflowEditorTools/DataflowEditorSkinWeightsPaintTool.cpp
@@ -79,22 +79,40 @@ UMeshSurfacePointTool* UDataflowEditorSkinWeightsPaintToolBuilder::CreateNewTool
 	{
 		if (FDataflowNode* PrimarySelection = ContextObject->GetSelectedNodeOfType<FDataflowCollectionEditSkinWeightsNode>())
 		{
-			FDataflowCollectionEditSkinWeightsNode* SkinWeightNode =
-				StaticCast<FDataflowCollectionEditSkinWeightsNode*>(PrimarySelection);
-			PaintTool->SkinWeightNode = SkinWeightNode;
-
-			SkinWeightNode->OnBoneSelectionChanged.AddLambda(
-				[PaintTool](const TArray<FName>& BoneNames)
-			{
-				PaintTool->GetNotifier().HandleNotification(BoneNames, ESkeletalMeshNotifyType::BonesSelected);
-			});
-
+			FDataflowCollectionEditSkinWeightsNode* SkinWeightNode = PrimarySelection->AsType<FDataflowCollectionEditSkinWeightsNode>();
+			PaintTool->SetSkinWeightNode(SkinWeightNode);
 			PaintTool->SetDataflowEditorContextObject(ContextObject);
 		}
 	}
 	return PaintTool;
 }
 
+
+void UDataflowEditorSkinWeightsPaintTool::SetSkinWeightNode(FDataflowCollectionEditSkinWeightsNode* Node)
+{
+	if (SkinWeightNode)
+	{
+		SkinWeightNode->OnBoneSelectionChanged.Remove(OnBoneSelectionChangedDelgateHandle);
+	}
+
+	SkinWeightNode = Node;
+
+	if (SkinWeightNode)
+	{
+		OnBoneSelectionChangedDelgateHandle = SkinWeightNode->OnBoneSelectionChanged.AddUObject(this, &UDataflowEditorSkinWeightsPaintTool::HandleOnBoneSelectionChanged);
+	}
+}
+
+void UDataflowEditorSkinWeightsPaintTool::HandleOnBoneSelectionChanged(const TArray<FName>& BoneNames)
+{
+	GetNotifier().HandleNotification(BoneNames, ESkeletalMeshNotifyType::BonesSelected);
+}
+
+void UDataflowEditorSkinWeightsPaintTool::SetDataflowEditorContextObject(TObjectPtr<UDataflowContextObject> InDataflowEditorContextObject)
+{
+	DataflowEditorContextObject = InDataflowEditorContextObject;
+}
+
 FMeshDescription* UDataflowEditorSkinWeightsPaintTool::GetCurrentDescription(const int32 LODIndex) const 
 {
 	if(const IPrimitiveComponentBackedTarget* TargetComponent = Cast<IPrimitiveComponentBackedTarget>(Target))
@@ -266,6 +284,8 @@ void UDataflowEditorSkinWeightsPaintTool::OnShutdown(EToolShutdownType ShutdownT
 			SkinWeightNode->ValidateSkeletalMeshes();
 		}
 	}
+
+	SetSkinWeightNode(nullptr);
 }
 
 #undef LOCTEXT_NAMESPACE
diff --git a/Engine/Plugins/Experimental/Dataflow/Source/DataflowEditor/Public/DataflowEditorTools/DataflowEditorSkinWeightsPaintTool.h b/Engine/Plugins/Experimental/Dataflow/Source/DataflowEditor/Public/DataflowEditorTools/DataflowEditorSkinWeightsPaintTool.h
--- a/Engine/Plugins/Experimental/Dataflow/Source/DataflowEditor/Public/DataflowEditorTools/DataflowEditorSkinWeightsPaintTool.h
+++ b/Engine/Plugins/Experimental/Dataflow/Source/DataflowEditor/Public/DataflowEditorTools/DataflowEditorSkinWeightsPaintTool.h
@@ -13,6 +13,7 @@
 
 class UDataflowContextObject;
 class UDataflowEditorMode;
+struct FDataflowCollectionEditSkinWeightsNode;
 
 /**
  * Dataflow skin weights tool builder
@@ -45,22 +46,21 @@ UCLASS(MinimalAPI)
 class UDataflowEditorSkinWeightsPaintTool : public USkinWeightsPaintTool
 {
 	GENERATED_BODY()
-	
+
+public:
+	void SetSkinWeightNode(FDataflowCollectionEditSkinWeightsNode* Node);
+
+	/** Set the dataflow context object */
+	void SetDataflowEditorContextObject(TObjectPtr<UDataflowContextObject> InDataflowEditorContextObject);
+
+private:
 	//~ Begin USkinWeightsPaintTool interface
 	virtual FEditorViewportClient* GetViewportClient() const override {return nullptr;}
 	UE_API virtual void Setup() override;
 	UE_API virtual void OnShutdown(EToolShutdownType ShutdownType) override;
 	//~ End USkinWeightsPaintTool interface
 
-	friend class UDataflowEditorSkinWeightsPaintToolBuilder;
-
-private:
-
-	/** Set the dataflow context object */
-	void SetDataflowEditorContextObject(TObjectPtr<UDataflowContextObject> InDataflowEditorContextObject)
-	{
-		DataflowEditorContextObject = InDataflowEditorContextObject;
-	}
+	void HandleOnBoneSelectionChanged(const TArray<FName>& BoneNames);
 
 	/** Extract the skin weights from the node */
 	UE_API bool ExtractSkinWeights(TArray<TArray<int32>>& CurrentIndices, TArray<TArray<float>>& CurrentWeights);
@@ -72,7 +72,7 @@ private:
 	UE_API FMeshDescription* GetCurrentDescription(const int32 LODIndex) const;
 	
 	/** Skin weight node associated with that tool */
-	struct FDataflowCollectionEditSkinWeightsNode* SkinWeightNode = nullptr;
+	FDataflowCollectionEditSkinWeightsNode* SkinWeightNode = nullptr;
 
 	/** Setup weights used to store the initial values */
 	TArray<TArray<float>> SetupWeights;
@@ -82,6 +82,8 @@ private:
 
 	/** Dataflow context object to be used to retrieve the node context data*/
 	TObjectPtr<UDataflowContextObject> DataflowEditorContextObject;
+
+	FDelegateHandle OnBoneSelectionChangedDelgateHandle;
 };
 
 
