Add a PropertyListCustomization for getting a list from a ManagedArrayCollection that isn't necessarily a list of groups.

diff --git a/Engine/Plugins/Experimental/Dataflow/Source/DataflowEditor/Private/Dataflow/PropertyGroupCustomization.cpp b/Engine/Plugins/Experimental/Dataflow/Source/DataflowEditor/Private/Dataflow/PropertyGroupCustomization.cpp
--- a/Engine/Plugins/Experimental/Dataflow/Source/DataflowEditor/Private/Dataflow/PropertyGroupCustomization.cpp
+++ b/Engine/Plugins/Experimental/Dataflow/Source/DataflowEditor/Private/Dataflow/PropertyGroupCustomization.cpp
@@ -13,49 +13,6 @@
 
 namespace UE::Dataflow
 {
-	const FManagedArrayCollection& FPropertyGroupCustomization::GetPropertyCollection(
-		const TSharedPtr<UE::Dataflow::FContext>& Context,
-		const TSharedPtr<IPropertyHandle>& ChildPropertyHandle,
-		const FName CollectionPropertyName = FName(TEXT("Collection")))
-	{
-		static const FManagedArrayCollection EmptyCollection;
-
-		TSharedPtr<IPropertyHandle> OwnerHandle = ChildPropertyHandle;
-		while (TSharedPtr<IPropertyHandle> ParentHandle = OwnerHandle->GetParentHandle())
-		{
-			OwnerHandle = MoveTemp(ParentHandle);
-		}
-		if (const TSharedPtr<IPropertyHandleStruct> OwnerHandleStruct = OwnerHandle->AsStruct())
-		{
-			if (const TSharedPtr<FStructOnScope> StructOnScope = OwnerHandleStruct->GetStructData())
-			{
-				if (const UStruct* const Struct = StructOnScope->GetStruct())
-				{
-					if (Struct->IsChildOf<FDataflowNode>())
-					{
-						const FDataflowNode* const DataflowNode = reinterpret_cast<FDataflowNode*>(StructOnScope->GetStructMemory());
-
-						if (const FProperty* const Property = Struct->FindPropertyByName(CollectionPropertyName))
-						{
-							if (const FStructProperty* const StructProperty = CastField<FStructProperty>(Property))
-							{
-								if (StructProperty->GetCPPType(nullptr, CPPF_None) == TEXT("FManagedArrayCollection"))
-								{
-									if (const FDataflowInput* const DataflowInput = DataflowNode->FindInput(StructProperty->ContainerPtrToValuePtr<FManagedArrayCollection*>(DataflowNode)))
-									{
-										UE::Dataflow::FContextThreaded EmptyContext;
-										return DataflowInput->GetValue(Context.IsValid() ? *Context : EmptyContext, EmptyCollection);
-									}
-								}
-							}
-						}
-					}
-				}
-			}
-		}
-		return EmptyCollection;
-	}
-
 	bool FPropertyGroupCustomization::MakeGroupName(FString& InOutString)
 	{
 		const FString SourceString = InOutString;
@@ -70,99 +27,9 @@ namespace UE::Dataflow
 		return MakeShareable(new FPropertyGroupCustomization);
 	}
 
-	void FPropertyGroupCustomization::CustomizeHeader(TSharedRef<IPropertyHandle> PropertyHandle, FDetailWidgetRow& HeaderRow, IPropertyTypeCustomizationUtils& CustomizationUtils)
-	{
-		DataflowGraphEditor = SDataflowGraphEditor::GetSelectedGraphEditor();
-
-		uint32 NumChildren;
-		const FPropertyAccess::Result Result = PropertyHandle->GetNumChildren(NumChildren);
-
-		ChildPropertyHandle = (Result == FPropertyAccess::Success && NumChildren) ? PropertyHandle->GetChildHandle(0) : nullptr;
-
-		GroupNames.Reset();
-
-		HeaderRow
-			.NameContent()
-			[
-				PropertyHandle->CreatePropertyNameWidget(PropertyHandle->GetPropertyDisplayName())
-			]
-			.ValueContent()
-			.MinDesiredWidth(250)
-			.MaxDesiredWidth(350.f)
-			[
-				SNew(SHorizontalBox)
-				+ SHorizontalBox::Slot()
-				.MaxWidth(145.f)
-				[
-					SAssignNew(ComboButton, SComboButton)
-					.ButtonStyle(FAppStyle::Get(), "NoBorder")
-					.ContentPadding(0)
-					.OnGetMenuContent(this, &FPropertyGroupCustomization::OnGetMenuContent)
-					.ButtonContent()
-					[
-						SNew(SEditableTextBox)
-						.Text(this, &FPropertyGroupCustomization::GetText)
-						.Font(FAppStyle::GetFontStyle(TEXT("PropertyWindow.NormalFont")))
-						.OnTextCommitted(this, &FPropertyGroupCustomization::OnTextCommitted)
-						.OnVerifyTextChanged(this, &FPropertyGroupCustomization::OnVerifyTextChanged)
-					]
-				]
-			];
-	}
-
-	FText FPropertyGroupCustomization::GetText() const
+	bool FPropertyGroupCustomization::MakeValidListName(FString& InOutString, FText& OutErrorMessage) const
 	{
-		FText Text;
-		if (ChildPropertyHandle)
-		{
-			ChildPropertyHandle->GetValueAsFormattedText(Text);
-		}
-		return Text;
-	}
-
-	void FPropertyGroupCustomization::OnTextCommitted(const FText& NewText, ETextCommit::Type CommitInfo)
-	{
-		if (ChildPropertyHandle)
-		{
-			FText CurrentText;
-			ChildPropertyHandle->GetValueAsFormattedText(CurrentText);
-
-			if (!NewText.ToString().Equals(CurrentText.ToString(), ESearchCase::CaseSensitive))
-			{
-				FString String = NewText.ToString();
-				MakeGroupName(String);
-				ChildPropertyHandle->SetValueFromFormattedString(String);
-			}
-		}
-	}
-
-	void FPropertyGroupCustomization::OnSelectionChanged(TSharedPtr<FText> ItemSelected, ESelectInfo::Type /*SelectInfo*/)
-	{
-		if (ChildPropertyHandle)
-		{
-			// Set the child property's value
-			if (ItemSelected)
-			{
-				FText CurrentText;
-				ChildPropertyHandle->GetValueAsFormattedText(CurrentText);
-
-				if (!ItemSelected->EqualTo(CurrentText))
-				{
-					ChildPropertyHandle->SetValueFromFormattedString(ItemSelected->ToString());
-				}
-
-				if (TSharedPtr<SComboButton> PinnedComboButton = ComboButton.Pin())
-				{
-					PinnedComboButton->SetIsOpen(false);
-				}
-			}
-		}
-	}
-
-	bool FPropertyGroupCustomization::OnVerifyTextChanged(const FText& Text, FText& OutErrorMessage)
-	{
-		FString TextString = Text.ToString();
-		const bool bIsValidGroupName = MakeGroupName(TextString);
+		const bool bIsValidGroupName = MakeGroupName(InOutString);
 		if (!bIsValidGroupName)
 		{
 			OutErrorMessage =
@@ -173,55 +40,24 @@ namespace UE::Dataflow
 		return bIsValidGroupName;
 	}
 
-	TSharedRef<ITableRow> FPropertyGroupCustomization::MakeCategoryViewWidget(TSharedPtr<FText> Item, const TSharedRef< STableViewBase >& OwnerTable)
-	{
-		if (Item)
-		{
-			return SNew(STableRow<TSharedPtr<FString>>, OwnerTable)
-			[
-				SNew(STextBlock).Text(*Item)
-			];
-		}
-		else
-		{
-			return SNew(STableRow<TSharedPtr<FString>>, OwnerTable);
-		}
-	}
 
-	TSharedRef<SWidget> FPropertyGroupCustomization::OnGetMenuContent()
+	TArray<FName> FPropertyGroupCustomization::GetListNames(const FManagedArrayCollection& Collection) const
 	{
-		const TSharedPtr<const SDataflowGraphEditor> DataflowGraphEditorPtr = DataflowGraphEditor.Pin();
-		const TSharedPtr<UE::Dataflow::FContext> Context = DataflowGraphEditorPtr ? DataflowGraphEditorPtr->GetDataflowContext() : TSharedPtr<UE::Dataflow::FContext>();
-
-		GroupNames.Reset();
-
-		// Retrieve collection
-		const FManagedArrayCollection& Collection = GetPropertyCollection(Context, ChildPropertyHandle, GetCollectionPropertyName());
-
 		// Find all group names in the parent node's collection
-		TArray<FName> CollectionGroupNames;
 		const TArray<FName> TargetGroupNames = GetTargetGroupNames(Collection);
 
 		// Find all group names in the parent selection node's collection
+		TArray<FName> CollectionGroupNames;
+		CollectionGroupNames.Reserve(TargetGroupNames.Num());
 		for (const FName& GroupName : Collection.GroupNames())
 		{
 			if (TargetGroupNames.Contains(GroupName))
 			{
-				GroupNames.Add(MakeShareable(new FText(FText::FromName(GroupName))));
+				CollectionGroupNames.Add(GroupName);
 			}
 		}
 
-		return SNew(SVerticalBox)
-			+ SVerticalBox::Slot()
-			.AutoHeight()
-			.MaxHeight(400.0f)
-			[
-				SNew(SListView<TSharedPtr<FText>>)
-					.ListItemsSource(&GroupNames)
-					.SelectionMode(ESelectionMode::Single)
-					.OnGenerateRow(this, &FPropertyGroupCustomization::MakeCategoryViewWidget)
-					.OnSelectionChanged(this, &FPropertyGroupCustomization::OnSelectionChanged)
-			];
+		return CollectionGroupNames;
 	}
 }  // End namespace UE::Chaos::ClothAsset
 
diff --git a/Engine/Plugins/Experimental/Dataflow/Source/DataflowEditor/Private/Dataflow/PropertyListCustomization.cpp b/Engine/Plugins/Experimental/Dataflow/Source/DataflowEditor/Private/Dataflow/PropertyListCustomization.cpp
new file mode 100644
--- /dev/null
+++ b/Engine/Plugins/Experimental/Dataflow/Source/DataflowEditor/Private/Dataflow/PropertyListCustomization.cpp
@@ -0,0 +1,196 @@
+// Copyright Epic Games, Inc. All Rights Reserved.
+
+#include "Dataflow/PropertyListCustomization.h"
+#include "Dataflow/DataflowGraphEditor.h"
+#include "Dataflow/DataflowNode.h"
+#include "Widgets/Input/SEditableTextBox.h"
+#include "Widgets/Input/SComboButton.h"
+#include "Widgets/Views/SListView.h"
+#include "Widgets/Views/STableRow.h"
+#include "DetailWidgetRow.h"
+
+namespace UE::Dataflow
+{
+	const FManagedArrayCollection& FPropertyListCustomization::GetPropertyCollection(
+		const TSharedPtr<UE::Dataflow::FContext>& Context,
+		const TSharedPtr<IPropertyHandle>& ChildPropertyHandle,
+		const FName CollectionPropertyName = FName(TEXT("Collection")))
+	{
+		static const FManagedArrayCollection EmptyCollection;
+
+		TSharedPtr<IPropertyHandle> OwnerHandle = ChildPropertyHandle;
+		while (TSharedPtr<IPropertyHandle> ParentHandle = OwnerHandle->GetParentHandle())
+		{
+			OwnerHandle = MoveTemp(ParentHandle);
+		}
+		if (const TSharedPtr<IPropertyHandleStruct> OwnerHandleStruct = OwnerHandle->AsStruct())
+		{
+			if (const TSharedPtr<FStructOnScope> StructOnScope = OwnerHandleStruct->GetStructData())
+			{
+				if (const UStruct* const Struct = StructOnScope->GetStruct())
+				{
+					if (Struct->IsChildOf<FDataflowNode>())
+					{
+						const FDataflowNode* const DataflowNode = reinterpret_cast<FDataflowNode*>(StructOnScope->GetStructMemory());
+
+						if (const FProperty* const Property = Struct->FindPropertyByName(CollectionPropertyName))
+						{
+							if (const FStructProperty* const StructProperty = CastField<FStructProperty>(Property))
+							{
+								if (StructProperty->GetCPPType(nullptr, CPPF_None) == TEXT("FManagedArrayCollection"))
+								{
+									if (const FDataflowInput* const DataflowInput = DataflowNode->FindInput(StructProperty->ContainerPtrToValuePtr<FManagedArrayCollection*>(DataflowNode)))
+									{
+										UE::Dataflow::FContextThreaded EmptyContext;
+										return DataflowInput->GetValue(Context.IsValid() ? *Context : EmptyContext, EmptyCollection);
+									}
+								}
+							}
+						}
+					}
+				}
+			}
+		}
+		return EmptyCollection;
+	}
+
+	void FPropertyListCustomization::CustomizeHeader(TSharedRef<IPropertyHandle> PropertyHandle, FDetailWidgetRow& HeaderRow, IPropertyTypeCustomizationUtils& CustomizationUtils)
+	{
+		DataflowGraphEditor = SDataflowGraphEditor::GetSelectedGraphEditor();
+
+		uint32 NumChildren;
+		const FPropertyAccess::Result Result = PropertyHandle->GetNumChildren(NumChildren);
+
+		ChildPropertyHandle = (Result == FPropertyAccess::Success && NumChildren) ? PropertyHandle->GetChildHandle(0) : nullptr;
+
+		ListNames.Reset();
+
+		HeaderRow
+			.NameContent()
+			[
+				PropertyHandle->CreatePropertyNameWidget(PropertyHandle->GetPropertyDisplayName())
+			]
+			.ValueContent()
+			.MinDesiredWidth(250)
+			.MaxDesiredWidth(350.f)
+			[
+				SNew(SHorizontalBox)
+				+ SHorizontalBox::Slot()
+				.MaxWidth(145.f)
+				[
+					SAssignNew(ComboButton, SComboButton)
+					.ButtonStyle(FAppStyle::Get(), "NoBorder")
+					.ContentPadding(0)
+					.OnGetMenuContent(this, &FPropertyListCustomization::OnGetMenuContent)
+					.ButtonContent()
+					[
+						SNew(SEditableTextBox)
+						.Text(this, &FPropertyListCustomization::GetText)
+						.Font(FAppStyle::GetFontStyle(TEXT("PropertyWindow.NormalFont")))
+						.OnTextCommitted(this, &FPropertyListCustomization::OnTextCommitted)
+						.OnVerifyTextChanged(this, &FPropertyListCustomization::OnVerifyTextChanged)
+					]
+				]
+			];
+	}
+
+	FText FPropertyListCustomization::GetText() const
+	{
+		FText Text;
+		if (ChildPropertyHandle)
+		{
+			ChildPropertyHandle->GetValueAsFormattedText(Text);
+		}
+		return Text;
+	}
+
+	void FPropertyListCustomization::OnTextCommitted(const FText& NewText, ETextCommit::Type CommitInfo)
+	{
+		if (ChildPropertyHandle)
+		{
+			FText CurrentText;
+			ChildPropertyHandle->GetValueAsFormattedText(CurrentText);
+
+			if (!NewText.ToString().Equals(CurrentText.ToString(), ESearchCase::CaseSensitive))
+			{
+				FString String = NewText.ToString();
+				FText UnusedErrorString;
+				MakeValidListName(String, UnusedErrorString);
+				ChildPropertyHandle->SetValueFromFormattedString(String);
+			}
+		}
+	}
+
+	void FPropertyListCustomization::OnSelectionChanged(TSharedPtr<FText> ItemSelected, ESelectInfo::Type /*SelectInfo*/)
+	{
+		if (ChildPropertyHandle)
+		{
+			// Set the child property's value
+			if (ItemSelected)
+			{
+				FText CurrentText;
+				ChildPropertyHandle->GetValueAsFormattedText(CurrentText);
+
+				if (!ItemSelected->EqualTo(CurrentText))
+				{
+					ChildPropertyHandle->SetValueFromFormattedString(ItemSelected->ToString());
+				}
+
+				if (TSharedPtr<SComboButton> PinnedComboButton = ComboButton.Pin())
+				{
+					PinnedComboButton->SetIsOpen(false);
+				}
+			}
+		}
+	}
+
+	bool FPropertyListCustomization::OnVerifyTextChanged(const FText& Text, FText& OutErrorMessage)
+	{
+		FString TextString = Text.ToString();
+		return MakeValidListName(TextString, OutErrorMessage);
+	}
+
+	TSharedRef<ITableRow> FPropertyListCustomization::MakeCategoryViewWidget(TSharedPtr<FText> Item, const TSharedRef< STableViewBase >& OwnerTable)
+	{
+		if (Item)
+		{
+			return SNew(STableRow<TSharedPtr<FString>>, OwnerTable)
+			[
+				SNew(STextBlock).Text(*Item)
+			];
+		}
+		else
+		{
+			return SNew(STableRow<TSharedPtr<FString>>, OwnerTable);
+		}
+	}
+
+	TSharedRef<SWidget> FPropertyListCustomization::OnGetMenuContent()
+	{
+		const TSharedPtr<const SDataflowGraphEditor> DataflowGraphEditorPtr = DataflowGraphEditor.Pin();
+		const TSharedPtr<UE::Dataflow::FContext> Context = DataflowGraphEditorPtr ? DataflowGraphEditorPtr->GetDataflowContext() : TSharedPtr<UE::Dataflow::FContext>();
+
+		ListNames.Reset();
+
+		// Retrieve collection
+		const FManagedArrayCollection& Collection = GetPropertyCollection(Context, ChildPropertyHandle, GetCollectionPropertyName());
+
+		const TArray<FName> TargetListNames = OnGetListNames.IsBound() ? OnGetListNames.Execute(Collection) : TArray<FName>();
+		for (const FName& ListName : TargetListNames)
+		{
+			ListNames.Add(MakeShareable(new FText(FText::FromName(ListName))));
+		}
+
+		return SNew(SVerticalBox)
+			+ SVerticalBox::Slot()
+			.AutoHeight()
+			.MaxHeight(400.0f)
+			[
+				SNew(SListView<TSharedPtr<FText>>)
+					.ListItemsSource(&ListNames)
+					.SelectionMode(ESelectionMode::Single)
+					.OnGenerateRow(this, &FPropertyListCustomization::MakeCategoryViewWidget)
+					.OnSelectionChanged(this, &FPropertyListCustomization::OnSelectionChanged)
+			];
+	}
+}  // End namespace UE::Chaos::ClothAsset
diff --git a/Engine/Plugins/Experimental/Dataflow/Source/DataflowEditor/Public/Dataflow/PropertyGroupCustomization.h b/Engine/Plugins/Experimental/Dataflow/Source/DataflowEditor/Public/Dataflow/PropertyGroupCustomization.h
--- a/Engine/Plugins/Experimental/Dataflow/Source/DataflowEditor/Public/Dataflow/PropertyGroupCustomization.h
+++ b/Engine/Plugins/Experimental/Dataflow/Source/DataflowEditor/Public/Dataflow/PropertyGroupCustomization.h
@@ -2,15 +2,9 @@
 
 #pragma once
 
-#include "GeometryCollection/ManagedArrayCollection.h"
-#include "IPropertyTypeCustomization.h"
+#include "Dataflow/PropertyListCustomization.h"
 
-class IPropertyHandle;
-class ITableRow;
-class SComboButton;
-class STableViewBase;
-class SWidget;
-class SDataflowGraphEditor;
+#define UE_API DATAFLOWEDITOR_API
 
 namespace UE::Dataflow
 {
@@ -19,28 +13,16 @@ namespace UE::Dataflow
 	/**
 	 * Property group customization to allow selecting of a group from a list of the groups currently held by the node's collection.
 	 */
-	class FPropertyGroupCustomization : public IPropertyTypeCustomization
+	class FPropertyGroupCustomization : public FPropertyListCustomization
 	{
 	public:
 		static TSharedRef<IPropertyTypeCustomization> MakeInstance();
 
-	private:
-		//~ Begin IPropertyTypeCustomization interface
-		DATAFLOWEDITOR_API virtual void CustomizeHeader(TSharedRef<IPropertyHandle> InPropertyHandle, FDetailWidgetRow& HeaderRow, IPropertyTypeCustomizationUtils& CustomizationUtils) override;
-		virtual void CustomizeChildren(TSharedRef<IPropertyHandle> /*InPropertyHandle*/, IDetailChildrenBuilder& /*ChildBuilder*/, IPropertyTypeCustomizationUtils& /*CustomizationUtils*/) override {}
-		//~ End IPropertyTypeCustomization interface
+		FPropertyGroupCustomization()
+			: FPropertyListCustomization(FPropertyListCustomization::FOnGetListNames::CreateRaw(this, &FPropertyGroupCustomization::GetListNames))
+		{}
 
-		/**
-		 * Return the FManagedArrayCollection with the specified name from the property held by the top level struct owner of ChildPropertyHandle.
-		 * @param Context The current editor Dataflow context.
-		 * @param ChildPropertyHandle The handle of a customized property from which to find the owner struct.
-		 * @param CollectionPropertyName
-		 * @return The collection.
-		 */
-		static DATAFLOWEDITOR_API const FManagedArrayCollection& GetPropertyCollection(
-			const TSharedPtr<UE::Dataflow::FContext>& Context,
-			const TSharedPtr<IPropertyHandle>& ChildPropertyHandle,
-			const FName CollectionPropertyName);
+	private:
 
 		/**
 		 * Turn a string into a valid collection group or attribute name.
@@ -50,7 +32,7 @@ namespace UE::Dataflow
 		 * @param InOutString The string to turn into a valid collection name.
 		 * @return Whether the InOutString was already a valid collection name.
 		 */
-		static DATAFLOWEDITOR_API bool MakeGroupName(FString& InOutString);
+		static bool MakeGroupName(FString& InOutString);
 
 		/** List of valid group names for the drop down list. Override this method to filter to a specific set of group names. */
 		virtual TArray<FName> GetTargetGroupNames(const FManagedArrayCollection& Collection) const
@@ -58,24 +40,9 @@ namespace UE::Dataflow
 			return Collection.GroupNames();  // By default returns all of the collection's group names
 		}
 
-		/** Name of the collection property. Override this method to specify your own collection property name. */
-		virtual FName GetCollectionPropertyName() const
-		{
-			return DefaultCollectionPropertyName;  // By default returns "Collection"
-		}
-
-		FText GetText() const;
-		void OnTextCommitted(const FText& NewText, ETextCommit::Type CommitInfo);
-		void OnSelectionChanged(TSharedPtr<FText> ItemSelected, ESelectInfo::Type SelectInfo);
-		bool OnVerifyTextChanged(const FText& Text, FText& OutErrorMessage);
-		TSharedRef<ITableRow> MakeCategoryViewWidget(TSharedPtr<FText> Item, const TSharedRef<STableViewBase>& OwnerTable);
-		TSharedRef<SWidget> OnGetMenuContent();
-		
-		static inline const FName DefaultCollectionPropertyName = TEXT("Collection");
-
-		TWeakPtr<const SDataflowGraphEditor> DataflowGraphEditor;
-		TSharedPtr<IPropertyHandle> ChildPropertyHandle;
-		TWeakPtr<SComboButton> ComboButton;
-		TArray<TSharedPtr<FText>> GroupNames;
+		UE_API virtual bool MakeValidListName(FString& InOutString, FText& OutErrorMessage) const override;
+		UE_API TArray<FName> GetListNames(const FManagedArrayCollection& Collection) const;
 	};
 }
+
+#undef  UE_API
diff --git a/Engine/Plugins/Experimental/Dataflow/Source/DataflowEditor/Public/Dataflow/PropertyListCustomization.h b/Engine/Plugins/Experimental/Dataflow/Source/DataflowEditor/Public/Dataflow/PropertyListCustomization.h
new file mode 100644
--- /dev/null
+++ b/Engine/Plugins/Experimental/Dataflow/Source/DataflowEditor/Public/Dataflow/PropertyListCustomization.h
@@ -0,0 +1,85 @@
+// Copyright Epic Games, Inc. All Rights Reserved.
+
+#pragma once
+
+#include "Delegates/DelegateCombinations.h"
+#include "GeometryCollection/ManagedArrayCollection.h"
+#include "IPropertyTypeCustomization.h"
+
+class IPropertyHandle;
+class ITableRow;
+class SComboButton;
+class STableViewBase;
+class SWidget;
+class SDataflowGraphEditor;
+
+namespace UE::Dataflow
+{
+	class FContext;
+
+	/**
+	 * Property list customization to allow selecting of a name from a list derived from the node's collection.
+	 */
+	class FPropertyListCustomization : public IPropertyTypeCustomization
+	{
+	public:
+
+		/** Delegate to return the list of valid names for the drop down list. */
+		DECLARE_DELEGATE_RetVal_OneParam(TArray<FName>, FOnGetListNames, const FManagedArrayCollection&);
+
+		FPropertyListCustomization() = default;
+
+		explicit FPropertyListCustomization(FOnGetListNames&& InOnGetListNames)
+			: IPropertyTypeCustomization()
+			, OnGetListNames(MoveTemp(InOnGetListNames))
+		{}
+
+	private:
+
+		//~ Begin IPropertyTypeCustomization interface
+		DATAFLOWEDITOR_API virtual void CustomizeHeader(TSharedRef<IPropertyHandle> InPropertyHandle, FDetailWidgetRow& HeaderRow, IPropertyTypeCustomizationUtils& CustomizationUtils) override;
+		virtual void CustomizeChildren(TSharedRef<IPropertyHandle> /*InPropertyHandle*/, IDetailChildrenBuilder& /*ChildBuilder*/, IPropertyTypeCustomizationUtils& /*CustomizationUtils*/) override {}
+		//~ End IPropertyTypeCustomization interface
+
+		/**
+		 * Return the FManagedArrayCollection with the specified name from the property held by the top level struct owner of ChildPropertyHandle.
+		 * @param Context The current editor Dataflow context.
+		 * @param ChildPropertyHandle The handle of a customized property from which to find the owner struct.
+		 * @param CollectionPropertyName
+		 * @return The collection.
+		 */
+		static DATAFLOWEDITOR_API const FManagedArrayCollection& GetPropertyCollection(
+			const TSharedPtr<UE::Dataflow::FContext>& Context,
+			const TSharedPtr<IPropertyHandle>& ChildPropertyHandle,
+			const FName CollectionPropertyName);
+
+		/** Override to make a list name valid and generate an error message if name is not valid.*/
+		virtual bool MakeValidListName(FString& InOutString, FText& OutErrorMessage) const
+		{
+			OutErrorMessage = FText::GetEmpty();
+			return true;
+		}
+
+		/** Name of the collection property. Override this method to specify your own collection property name. */
+		virtual FName GetCollectionPropertyName() const
+		{
+			return DefaultCollectionPropertyName;  // By default returns "Collection"
+		}
+
+		FText GetText() const;
+		void OnTextCommitted(const FText& NewText, ETextCommit::Type CommitInfo);
+		void OnSelectionChanged(TSharedPtr<FText> ItemSelected, ESelectInfo::Type SelectInfo);
+		bool OnVerifyTextChanged(const FText& Text, FText& OutErrorMessage);
+		TSharedRef<ITableRow> MakeCategoryViewWidget(TSharedPtr<FText> Item, const TSharedRef<STableViewBase>& OwnerTable);
+		TSharedRef<SWidget> OnGetMenuContent();
+		
+		static inline const FName DefaultCollectionPropertyName = TEXT("Collection");
+
+		TWeakPtr<const SDataflowGraphEditor> DataflowGraphEditor;
+		TSharedPtr<IPropertyHandle> ChildPropertyHandle;
+		TWeakPtr<SComboButton> ComboButton;
+		TArray<TSharedPtr<FText>> ListNames;
+
+		FOnGetListNames OnGetListNames;
+	};
+}
