Dataflow Editor : Enable the new rendering system to work along the existing rendering systems
- older rendering systems are still compatible and working
- nodes not implementing any existing rendering techniques will be defaulting to the new system if available for their output types

diff --git a/Engine/Plugins/Experimental/Dataflow/Source/DataflowEditor/Private/Dataflow/DataflowConstructionScene.cpp b/Engine/Plugins/Experimental/Dataflow/Source/DataflowEditor/Private/Dataflow/DataflowConstructionScene.cpp
--- a/Engine/Plugins/Experimental/Dataflow/Source/DataflowEditor/Private/Dataflow/DataflowConstructionScene.cpp
+++ b/Engine/Plugins/Experimental/Dataflow/Source/DataflowEditor/Private/Dataflow/DataflowConstructionScene.cpp
@@ -11,6 +11,7 @@
 #include "Dataflow/DataflowEditorStyle.h"
 #include "Dataflow/DataflowRenderingFactory.h"
 #include "DataflowRendering/DataflowRenderableComponents.h"
+#include "DataflowRendering/DataflowRenderableTypeInstance.h"
 #include "Drawing/MeshElementsVisualizer.h"
 #include "Elements/Framework/EngineElementsLibrary.h"
 #include "Selection.h"
@@ -211,23 +212,6 @@ void FDataflowConstructionScene::UpdateDynamicMeshComponents()
 
 void FDataflowConstructionScene::UpdatePrimitiveComponents()
 {
-	if (const TObjectPtr<UDataflowBaseContent>& EditorContent = GetEditorContent())
-	{
-		if (const TSharedPtr<UE::Dataflow::FEngineContext> DataflowContext = EditorContent->GetDataflowContext())
-		{
-			if (UDataflowEdNode* SelectedNode = EditorContent->GetSelectedNode())
-			{
-				if (TSharedPtr<FDataflowNode> SelectedDataflowNode = SelectedNode->GetDataflowNode())
-				{
-					if(SelectedDataflowNode->IsA(FDataflowPrimitiveNode::StaticType()))
-					{
-						StaticCastSharedPtr<FDataflowPrimitiveNode>(SelectedDataflowNode)->AddPrimitiveComponents(*DataflowContext,  
-							EditorContent->GetRenderCollection(), DataflowContext->Owner, RootSceneActor, PrimitiveComponents);
-					}
-				}
-			}
-		}
-	}
 	for(TObjectPtr<UPrimitiveComponent> PrimitiveComponent : PrimitiveComponents)
 	{
 		PrimitiveComponent->SelectionOverrideDelegate = UPrimitiveComponent::FSelectionOverride::CreateRaw(this, &FDataflowPreviewSceneBase::IsComponentSelected);
@@ -268,6 +252,7 @@ void FDataflowConstructionScene::ResetSceneComponents()
 	}
 	DynamicMeshComponents.Reset();
 	PrimitiveComponents.Reset();
+	RenderableTypeInstances.Reset();
 	bPreviewSceneDirty = true;
 
 	RemoveSceneObject(RootSceneActor);
@@ -394,6 +379,13 @@ bool FDataflowConstructionScene::HasRenderableGeometry()
 			return true;
 		}
 	}
+	for (const UPrimitiveComponent* PrimitiveComponent : PrimitiveComponents)
+	{
+		if (PrimitiveComponent->Bounds.SphereRadius > UE_SMALL_NUMBER)
+		{
+			return true;
+		}
+	}
 	return false;
 }
 
@@ -554,6 +546,26 @@ void FDataflowConstructionScene::UpdateDebugMesh(const GeometryCollection::Facad
 }
 
 
+TArray<UMaterialInterface*> FDataflowConstructionScene::GetMaterialsFromRenderingFacade(const GeometryCollection::Facades::FRenderingFacade& RenderingFacade, int32 MeshIndex)
+{
+	TArray<UMaterialInterface*> Materials;
+
+	if (MeshIndex >= 0 && MeshIndex < RenderingFacade.NumGeometry())
+	{
+		const TManagedArray<FString>& MaterialPaths = RenderingFacade.GetMaterialPaths();
+		const int32 MaterialStart = RenderingFacade.GetMaterialStart()[MeshIndex];
+		const int32 MaterialCount = RenderingFacade.GetMaterialCount()[MeshIndex];
+
+		for (int32 MaterialIndex = MaterialStart; MaterialIndex < MaterialStart + MaterialCount; ++MaterialIndex)
+		{
+			const FString& Path = MaterialPaths[MaterialIndex];
+			UMaterialInterface* const Material = LoadObject<UMaterialInterface>(nullptr, *Path);
+			Materials.Add(Material);
+		}
+	}
+	return Materials;
+}
+
 void FDataflowConstructionScene::UpdateComponentsFromNode(const UDataflowEdNode* EdNode, bool bWireframe)
 {
 	if (EdNode == nullptr)
@@ -566,69 +578,111 @@ void FDataflowConstructionScene::UpdateComponentsFromNode(const UDataflowEdNode*
 	const FDataflowEditorToolkit* EditorToolkit = DataflowEditor ? static_cast<FDataflowEditorToolkit*>(DataflowEditor->GetInstanceInterface()) : nullptr;
 	const bool bEvaluateOutputs = EditorToolkit ? (EditorToolkit->GetEvaluationMode() != EDataflowEditorEvaluationMode::Manual) : true;
 
-	if (const TObjectPtr<UDataflowBaseContent>& EditorContent = GetEditorContent())
+	const TObjectPtr<UDataflowBaseContent>& EditorContent = GetEditorContent();
+	const TSharedPtr<UE::Dataflow::FEngineContext> DataflowContext = EditorContent ? EditorContent->GetDataflowContext() : TSharedPtr<UE::Dataflow::FEngineContext>();
+
+	if (EditorContent && DataflowContext)
 	{
+		const UE::Dataflow::IDataflowConstructionViewMode* ConstructionViewMode = EditorContent->GetConstructionViewMode();
+
 		TSharedPtr<FManagedArrayCollection> RenderCollection(new FManagedArrayCollection);
 		GeometryCollection::Facades::FRenderingFacade RenderingFacade(*RenderCollection);
 		RenderingFacade.DefineSchema();
 
-		const bool bHasRenderCollectionPrimitives = UE::Dataflow::RenderNodeOutput(RenderingFacade, *EdNode, *EditorContent, bEvaluateOutputs);
+		const bool bCanRenderOutput = (ConstructionViewMode)
+			? UE::Dataflow::CanRenderNodeOutput(*EdNode, *EditorContent, *ConstructionViewMode)
+			: false;
 
-		if (EdNode == EditorContent->GetSelectedNode())
+		bool bHasRenderCollectionPrimitives = false;
+		if (bCanRenderOutput)
 		{
-			EditorContent->SetRenderCollection(RenderCollection);
+			bHasRenderCollectionPrimitives = UE::Dataflow::RenderNodeOutput(RenderingFacade, *EdNode, *EditorContent, bEvaluateOutputs);
+			if (EdNode == EditorContent->GetSelectedNode())
+			{
+				EditorContent->SetRenderCollection(RenderCollection);
+			}
+		}
+
+		// Dataflow primitive nodes node explicily provide its own components 
+		// We do not check for bHasRenderCollectionPrimitives because some nodes actualy set HasRenderCollectionPrimitives to false while still implementing AddPrimitiveComponents
+		// and we need to support that for the time being 
+		if (TSharedPtr<const FDataflowNode> DataflowNode = EdNode->GetDataflowNode())
+		{
+			// sadly AddPrimitiveComponents is not const , so we need to const_cast the pointer 
+			// this code path should potentially be deprecated if we adopt the new way of rendering where the rendering callback are making components directly
+			if (FDataflowPrimitiveNode* PrimitiveDataflowNode = const_cast<FDataflowPrimitiveNode*>(DataflowNode->AsType<FDataflowPrimitiveNode>()))
+			{
+				PrimitiveDataflowNode->AddPrimitiveComponents(*DataflowContext, EditorContent->GetRenderCollection(), DataflowContext->Owner, RootSceneActor, PrimitiveComponents);
+			}
 		}
+
 		if (!bHasRenderCollectionPrimitives)
 		{
-			const int32 NumGeometry = RenderingFacade.NumGeometry();
-			for (int32 MeshIndex = 0; MeshIndex < NumGeometry; ++MeshIndex)
+			if (bCanRenderOutput) // Path where the node has defined it's way of rendering the output using a MACRO on the node definition itself 
 			{
-				FDataflowRenderKey DynamicMeshKey{ EdNode, MeshIndex };
-				if (bWireframe && DynamicMeshComponents.Contains(DynamicMeshKey))
-				{
-					UDynamicMeshComponent* const ExistingMeshComponent = DynamicMeshComponents[DynamicMeshKey];
-					MeshComponentsForWireframeRendering.Add(ExistingMeshComponent);
-				}
-				else
+				const int32 NumGeometry = RenderingFacade.NumGeometry();
+				for (int32 MeshIndex = 0; MeshIndex < NumGeometry; ++MeshIndex)
 				{
-					FDynamicMesh3 DynamicMesh;
-					UE::Dataflow::Conversion::RenderingFacadeToDynamicMesh(RenderingFacade, MeshIndex, DynamicMesh);
-
-					if (DynamicMesh.VertexCount())
+					FDataflowRenderKey DynamicMeshKey{ EdNode, MeshIndex };
+					if (bWireframe && DynamicMeshComponents.Contains(DynamicMeshKey))
 					{
-						const FString MeshName = RenderingFacade.GetGeometryName()[MeshIndex];
-						const FTransform& MeshTransform = RenderingFacade.GetGeometryTransform()[MeshIndex];
-						if (bWireframe)
-						{
-							// Add hidden DynamicMeshComponents for any targets that we want to render in wireframe
-							// 
-							// Note: UMeshElementsVisualizers need source meshes to pull from. We add invisible dynamic mesh components to the existing DynamicMeshComponents collection
-							// for this purpose, but could have instead created a separate collection of meshes for wireframe rendering. We are choosing to keep all the scene DynamicMeshComponents 
-							// in one place and using separate structures to dictate how they are used (MeshComponentsForWireframeRendering in this case), in case visualization requirements 
-							// change in the future.
-							//
-							const FString UniqueObjectName = MakeUniqueObjectName(RootSceneActor, UDataflowEditorCollectionComponent::StaticClass(), FName(MeshName)).ToString();
-							UDynamicMeshComponent* const NewDynamicMeshComponent = AddDynamicMeshComponent(DynamicMeshKey, UniqueObjectName, MoveTemp(DynamicMesh), {}, MeshTransform);
-							NewDynamicMeshComponent->SetVisibility(false);
-							MeshComponentsForWireframeRendering.Add(NewDynamicMeshComponent);
-						}
-						else
-						{
-							const TManagedArray<FString>& MaterialPaths = RenderingFacade.GetMaterialPaths();
-							const int32 MaterialStart = RenderingFacade.GetMaterialStart()[MeshIndex];
-							const int32 MaterialCount = RenderingFacade.GetMaterialCount()[MeshIndex];
+						UDynamicMeshComponent* const ExistingMeshComponent = DynamicMeshComponents[DynamicMeshKey];
+						MeshComponentsForWireframeRendering.Add(ExistingMeshComponent);
+					}
+					else
+					{
+						FDynamicMesh3 DynamicMesh;
+						UE::Dataflow::Conversion::RenderingFacadeToDynamicMesh(RenderingFacade, MeshIndex, DynamicMesh);
 
-							TArray<UMaterialInterface*> Materials;
-							for (int32 MaterialIndex = MaterialStart; MaterialIndex < MaterialStart + MaterialCount; ++MaterialIndex)
+						if (DynamicMesh.VertexCount())
+						{
+							const FString MeshName = RenderingFacade.GetGeometryName()[MeshIndex];
+							const FTransform& MeshTransform = RenderingFacade.GetGeometryTransform()[MeshIndex];
+							if (bWireframe)
 							{
-								const FString& Path = MaterialPaths[MaterialIndex];
-								UMaterialInterface* const Material = LoadObject<UMaterialInterface>(nullptr, *Path);
-								Materials.Add(Material);
+								// Add hidden DynamicMeshComponents for any targets that we want to render in wireframe
+								// 
+								// Note: UMeshElementsVisualizers need source meshes to pull from. We add invisible dynamic mesh components to the existing DynamicMeshComponents collection
+								// for this purpose, but could have instead created a separate collection of meshes for wireframe rendering. We are choosing to keep all the scene DynamicMeshComponents 
+								// in one place and using separate structures to dictate how they are used (MeshComponentsForWireframeRendering in this case), in case visualization requirements 
+								// change in the future.
+								//
+								const FString UniqueObjectName = MakeUniqueObjectName(RootSceneActor, UDataflowEditorCollectionComponent::StaticClass(), FName(MeshName)).ToString();
+								UDynamicMeshComponent* const NewDynamicMeshComponent = AddDynamicMeshComponent(DynamicMeshKey, UniqueObjectName, MoveTemp(DynamicMesh), {}, MeshTransform);
+								NewDynamicMeshComponent->SetVisibility(false);
+								MeshComponentsForWireframeRendering.Add(NewDynamicMeshComponent);
+							}
+							else
+							{
+								const TArray<UMaterialInterface*> Materials = GetMaterialsFromRenderingFacade(RenderingFacade, MeshIndex);
+								AddDynamicMeshComponent(DynamicMeshKey, MeshName, MoveTemp(DynamicMesh), Materials, MeshTransform);
 							}
-							AddDynamicMeshComponent(DynamicMeshKey, MeshName, MoveTemp(DynamicMesh), Materials, MeshTransform);
 						}
 					}
 				}
+				UpdateDebugMesh(RenderingFacade);
+			}
+			else // New method where rendering callbacks generate primitive components directly by output types
+			{
+				TArray<UE::Dataflow::FRenderableTypeInstance> NodeRenderableTypeInstances;
+				if (ConstructionViewMode)
+				{
+					if (TSharedPtr<const FDataflowNode> DataflowNode = EdNode->GetDataflowNode())
+					{
+						UE::Dataflow::FRenderableTypeInstance::GetRenderableInstancesForNode(*DataflowContext, *ConstructionViewMode, *DataflowNode, NodeRenderableTypeInstances);
+					}
+				}
+				NodeRenderableTypeInstances.Append(RenderableTypeInstances);
+
+				// extract components from instances
+				UE::Dataflow::FRenderableComponents NodeRenderableComponents(RootSceneActor);
+				TArray<UPrimitiveComponent*> ComponentsFromOutputTypes;
+				for (UE::Dataflow::FRenderableTypeInstance& Instance : NodeRenderableTypeInstances)
+				{
+					// todo : filter by option ? - what option should be enmabled by default ?
+					Instance.GetPrimitiveComponents(NodeRenderableComponents);
+				}
+				PrimitiveComponents.Append(NodeRenderableComponents.GetComponents());
 			}
 		}
 	}
diff --git a/Engine/Plugins/Experimental/Dataflow/Source/DataflowEditor/Private/Dataflow/DataflowEditorModule.cpp b/Engine/Plugins/Experimental/Dataflow/Source/DataflowEditor/Private/Dataflow/DataflowEditorModule.cpp
--- a/Engine/Plugins/Experimental/Dataflow/Source/DataflowEditor/Private/Dataflow/DataflowEditorModule.cpp
+++ b/Engine/Plugins/Experimental/Dataflow/Source/DataflowEditor/Private/Dataflow/DataflowEditorModule.cpp
@@ -35,6 +35,12 @@
 #include "Dataflow/DataflowColorRamp.h"
 #include "Dataflow/DataflowColorRampCustomization.h"
 
+#include "DataflowRendering/DataflowStaticMeshRenderableType.h"
+#include "DataflowRendering/DataflowSkeletalMeshRenderableType.h"
+#include "DataflowRendering/DataflowDynamicMeshRenderableType.h"
+#include "DataflowRendering/DataflowGeometryCollectionRenderableType.h"
+#include "DataflowRendering/DataflowGeometryCollectionProximityRenderableType.h"
+
 #define LOCTEXT_NAMESPACE "DataflowEditor"
 
 namespace UE::Dataflow::Private
@@ -188,6 +194,13 @@ void FDataflowEditorModule::StartupModule()
 	ConstructionVisualizationRegistry.RegisterVisualization(MakeUnique<UE::Dataflow::FDebugDrawText3dVisualization>());
 	
 	FDataflowEditorCommands::Register();
+
+	// register rendering types
+	UE::Dataflow::Private::RegisterStaticMeshRenderableTypes();
+	UE::Dataflow::Private::RegisterSkeletalMeshRenderableTypes();
+	UE::Dataflow::Private::RegisterDynamicMeshRenderableTypes();
+	UE::Dataflow::Private::RegisterGeometryCollectionRenderableTypes();
+	UE::Dataflow::Private::RegisterGeometryCollectionProximityRenderableTypes();
 }
 
 void FDataflowEditorModule::ShutdownModule()
diff --git a/Engine/Plugins/Experimental/Dataflow/Source/DataflowEditor/Public/Dataflow/DataflowConstructionScene.h b/Engine/Plugins/Experimental/Dataflow/Source/DataflowEditor/Public/Dataflow/DataflowConstructionScene.h
--- a/Engine/Plugins/Experimental/Dataflow/Source/DataflowEditor/Public/Dataflow/DataflowConstructionScene.h
+++ b/Engine/Plugins/Experimental/Dataflow/Source/DataflowEditor/Public/Dataflow/DataflowConstructionScene.h
@@ -22,6 +22,11 @@ namespace UE
 	{
 		class FDynamicMesh3;
 	}
+
+	namespace Dataflow
+	{
+		struct FRenderableTypeInstance;
+	}
 }
 
 /**
@@ -97,6 +102,9 @@ private:
 	/** Build the primitive components from the dataflow graph selected node */
     UE_API void UpdatePrimitiveComponents();
 	
+	/** Get the materials stored in the rendering facade for a speicfic mesh */
+	TArray<UMaterialInterface*> GetMaterialsFromRenderingFacade(const GeometryCollection::Facades::FRenderingFacade& RenderingFacade, int32 MeshIndex);
+
 	/** Build the primitive components from the dataflow graph selected node */
 	void UpdateComponentsFromNode(const UDataflowEdNode* Node, bool bWireframe);
 
@@ -134,6 +142,9 @@ private:
 	
 	/** Primitive components generated by the nodes and to be used by the tools */
 	TArray<TObjectPtr<UPrimitiveComponent>> PrimitiveComponents;
+
+	/** list of renderable type instance, there's one per node/output that can be rendered */
+	TArray<UE::Dataflow::FRenderableTypeInstance> RenderableTypeInstances;
 };
 
 #undef UE_API
