Dataflow Vertex attribute tool : add mirroring support

diff --git a/Engine/Plugins/Experimental/Dataflow/Source/DataflowEditor/Private/DataflowEditorTools/DataflowEditorVertexAttributePaintTool.cpp b/Engine/Plugins/Experimental/Dataflow/Source/DataflowEditor/Private/DataflowEditorTools/DataflowEditorVertexAttributePaintTool.cpp
--- a/Engine/Plugins/Experimental/Dataflow/Source/DataflowEditor/Private/DataflowEditorTools/DataflowEditorVertexAttributePaintTool.cpp
+++ b/Engine/Plugins/Experimental/Dataflow/Source/DataflowEditor/Private/DataflowEditorTools/DataflowEditorVertexAttributePaintTool.cpp
@@ -35,11 +35,13 @@
 #include "StaticMeshAttributes.h"
 #include "Selections/MeshConnectedComponents.h"
 #include "Selections/MeshVertexSelection.h"
+#include "Spatial/PointHashGrid3.h"
 #include "ToolSetupUtil.h"
 #include "Util/BufferUtil.h"
 #include "Widgets/Notifications/SNotificationList.h"
 
 
+
 #include UE_INLINE_GENERATED_CPP_BY_NAME(DataflowEditorVertexAttributePaintTool)
 
 DEFINE_LOG_CATEGORY_STATIC(LogDataflowEditorVertexAttributePaintTool, Warning, All);
@@ -1152,13 +1154,16 @@ bool UDataflowEditorVertexAttributePaintTool::SyncMeshWithWeightBuffer(FDynamicM
 		for (int32 VertexROIIndex = 0; VertexROIIndex < NumVertexROI; ++VertexROIIndex)
 		{
 			const int VertIdx = VertexROI[VertexROIIndex];
-			const float NewValue = ROIWeightValueBuffer[VertexROIIndex];
-			VertexData.SetValue(VertIdx, 
-				[NewValue](UE::Geometry::FDynamicMeshWeightAttribute& AttributeMap, int32 AttributeIndex)
-				{
-					return NewValue;
-				});
-			NumModified++;
+			if (VertIdx != INDEX_NONE)
+			{
+				const float NewValue = ROIWeightValueBuffer[VertexROIIndex];
+				VertexData.SetValue(VertIdx,
+					[NewValue](UE::Geometry::FDynamicMeshWeightAttribute& AttributeMap, int32 AttributeIndex)
+					{
+						return NewValue;
+					});
+				NumModified++;
+			}
 		}
 	}
 	return (NumModified > 0);
@@ -2325,6 +2330,49 @@ void UDataflowEditorVertexAttributePaintTool::UpdateVertexColorOverlay(const TSe
 	PreviewMesh->NotifyDeferredEditCompleted(UPreviewMesh::ERenderUpdateMode::FastUpdate, EMeshRenderAttributeFlags::VertexColors, false);
 }
 
+void UDataflowEditorVertexAttributePaintTool::MirrorValues()
+{
+	if (ToolProperties)
+	{
+		if (!VertexData.IsValid())
+		{
+			return;
+		}
+
+		VertexROI = UDataflowEditorVertexAttributePaintTool::GetSelectedVertices();
+		if (VertexROI.IsEmpty())
+		{
+			return;
+		}
+
+		BeginChange();
+
+		if (FDynamicMesh3* Mesh = GetSculptMesh())
+		{
+			MirrorData.EnsureMirrorDataIsUpdated(*Mesh, ToolProperties->MirrorProperties.MirrorAxis, ToolProperties->MirrorProperties.MirrorDirection);
+
+			TArray<int32> VerticesToUpdate;
+			MirrorData.FindMirroredIndices(*Mesh, VertexROI, VerticesToUpdate);
+
+			// get the selected vertices values
+			ROIWeightValueBuffer.SetNum(VertexROI.Num(), EAllowShrinking::No);
+			SyncWeightBufferWithMesh(Mesh);
+
+			// now swap VertexROI with VerticesToUpdate ( ROIWeightValueBuffer map 1:1 to VerticesToUpdate)
+			Swap(VertexROI, VerticesToUpdate);
+
+			// apply the values back to the mesh with the mirrored indices
+			SyncMeshWithWeightBuffer(Mesh);
+		}
+
+		// update colors
+		UpdateVertexColorOverlay();
+		GetToolManager()->PostInvalidation();
+		EndChange();
+
+	}
+}
+
 void UDataflowEditorVertexAttributePaintTool::AddReferencedObjects(UObject* InThis, FReferenceCollector& Collector)
 {
 	UDataflowEditorVertexAttributePaintTool* This = CastChecked<UDataflowEditorVertexAttributePaintTool>(InThis);
@@ -2335,5 +2383,126 @@ void UDataflowEditorVertexAttributePaintTool::AddReferencedObjects(UObject* InTh
 	Super::AddReferencedObjects(InThis, Collector);
 }
 
+/////////////////////////////////////////////////////////////////////////////////////////////////////////////////////////////////////////////////////////////////
+// 
+// MirrorData
+// 
+/////////////////////////////////////////////////////////////////////////////////////////////////////////////////////////////////////////////////////////////////
+
+bool UDataflowEditorVertexAttributePaintTool::FMirrorData::IsPointOnTargetMirrorSide(const FVector& InPoint) const
+{
+	if (Axis == EAxis::None)
+	{
+		return false;
+	}
+	if (Direction == EDataflowEditorToolMirrorDirection::PositiveToNegative && InPoint[Axis - 1] >= 0.f)
+	{
+		return false; // target is negative side, but point is on positive side
+	}
+	if (Direction == EDataflowEditorToolMirrorDirection::NegativeToPositive && InPoint[Axis - 1] <= 0.f)
+	{
+		return false; // target is positive side, but vertex is on negative side
+	}
+	return true;
+}
+
+void UDataflowEditorVertexAttributePaintTool::FMirrorData::EnsureMirrorDataIsUpdated(const FDynamicMesh3& Mesh, EAxis::Type InMirrorAxis, EDataflowEditorToolMirrorDirection InMirrorDirection)
+{
+	if (VertexMap.Num() > 0 && InMirrorAxis == Axis && InMirrorDirection == Direction)
+	{
+		// already initialized, just re-use cached data
+		return;
+	}
+
+	// need to re-initialize
+	Axis = InMirrorAxis;
+	Direction = InMirrorDirection;
+	VertexMap.Reset();
+
+	const int32 NumVertices = Mesh.VertexCount();
+
+	// build a spatial hash grid
+	constexpr float HashGridCellSize = 2.0f; // the length of the cell size in the point hash grid
+	UE::Geometry::TPointHashGrid3f<int32> VertHash(HashGridCellSize, INDEX_NONE);
+	VertHash.Reserve(NumVertices);
+
+	for (int32 VertexID : Mesh.VertexIndicesItr())
+	{
+		const FVector VertexPos = Mesh.GetVertex(VertexID);
+		VertHash.InsertPointUnsafe(VertexID, static_cast<FVector3f>(VertexPos));
+	}
+
+	// generate a map of point IDs on the target side, to their equivalent vertex ID on the source side 
+	for (int32 TargetVertexID : Mesh.VertexIndicesItr())
+	{
+		const FVector SourcePosition = Mesh.GetVertex(TargetVertexID);
+
+		// flip position across the mirror axis
+		FVector3f MirroredPosition = FVector3f(SourcePosition);
+		MirroredPosition[Axis - 1] *= -1.f;
+
+		// Query spatial hash near mirrored position, gradually increasing search radius until at least 1 point is found
+		TPair<int32, double> ClosestMirroredPoint = { INDEX_NONE, TNumericLimits<double>::Max() };
+		float SearchRadius = HashGridCellSize;
+		while (ClosestMirroredPoint.Key == INDEX_NONE)
+		{
+			ClosestMirroredPoint = VertHash.FindNearestInRadius(
+				MirroredPosition,
+				SearchRadius,
+				[&Mesh, MirroredPosition](int32 VID)
+				{
+					return FVector3f::DistSquared(FVector3f(Mesh.GetVertex(VID)), MirroredPosition);
+				});
+
+			SearchRadius += HashGridCellSize;
+
+			// forcibly break out if search radius gets bigger than the maximum search radius
+			static float MaxSearchRadius = 15.f; // TODO we may want to expose this value to the user...
+			if (SearchRadius >= MaxSearchRadius)
+			{
+				break;
+			}
+		}
+
+		if (ClosestMirroredPoint.Key != INDEX_NONE)
+		{
+			VertexMap.FindOrAdd(TargetVertexID, ClosestMirroredPoint.Key);
+		}
+	}
+}
+
+void UDataflowEditorVertexAttributePaintTool::FMirrorData::FindMirroredIndices(const FDynamicMesh3& Mesh, const TArray<int32>& SelectedVertices, TArray<int32>& OutVerticesToUpdate)
+{
+	check(Axis != EAxis::None);
+
+	// results have a 1:1 mapping to the input vertices
+	OutVerticesToUpdate.SetNum(SelectedVertices.Num());
+
+	// we need to convert selection to the equivalent target vertex indices (on the target side of the mirror plane)
+	// if a vertex is already on the target side, great
+	// if the user selected vertices on the source side, we convert them to the mirrored equivalent on the target side
+	TSet<VertexIndex> TargetVertices;
+	TArray<VertexIndex> MissingVertices;
+	for (int32 Index = 0; Index < SelectedVertices.Num(); ++Index)
+	{
+		const int32 SelectedVertexID = SelectedVertices[Index];
+		int32& TargetVertexID = OutVerticesToUpdate[Index];
+
+		// get the mirrored index or the invaliod index if not found
+		OutVerticesToUpdate[Index] = VertexMap.FindRef(SelectedVertexID, INDEX_NONE);
+
+		// if the selected vertex is on the wrong side, reset to an invalid index
+		if (Mesh.IsVertex(TargetVertexID))
+		{
+			const FVector TargetPosition = Mesh.GetVertex(TargetVertexID);
+			if (IsPointOnTargetMirrorSide(TargetPosition))
+			{
+				OutVerticesToUpdate[Index] = INDEX_NONE;
+			}
+		}
+	}
+}
+
+
 #undef LOCTEXT_NAMESPACE
 
diff --git a/Engine/Plugins/Experimental/Dataflow/Source/DataflowEditor/Private/DataflowEditorTools/DataflowEditorVertexAttributePaintToolDetailCustomization.cpp b/Engine/Plugins/Experimental/Dataflow/Source/DataflowEditor/Private/DataflowEditorTools/DataflowEditorVertexAttributePaintToolDetailCustomization.cpp
--- a/Engine/Plugins/Experimental/Dataflow/Source/DataflowEditor/Private/DataflowEditorTools/DataflowEditorVertexAttributePaintToolDetailCustomization.cpp
+++ b/Engine/Plugins/Experimental/Dataflow/Source/DataflowEditor/Private/DataflowEditorTools/DataflowEditorVertexAttributePaintToolDetailCustomization.cpp
@@ -1063,6 +1063,115 @@ namespace UE::Dataflow::Editor
 		];
 	}
 
+	void FVertexAttributePaintToolDetailCustomization::MakeSelectionMirrorOperationRow(IDetailCategoryBuilder& EditValuesCategory) const
+	{
+		TSharedRef<SSegmentedControl<EAxis::Type>> XYZButtonsWidget =
+			SNew(SSegmentedControl<EAxis::Type>)
+			.ToolTipText(LOCTEXT("MirrorAxisTooltip",
+				"X: copies weights across the YZ plane.\n"
+				"Y: copies weights across the XZ plane.\n"
+				"Z: copies weights across the XY plane."))
+			.Value_Lambda([this]()
+				{
+					return ToolProperties.IsValid() ? ToolProperties->MirrorProperties.MirrorAxis : TEnumAsByte<EAxis::Type>(EAxis::X);
+				})
+			.OnValueChanged_Lambda([this](EAxis::Type Mode)
+				{
+					if (ToolProperties.IsValid())
+					{
+						ToolProperties->MirrorProperties.MirrorAxis = Mode;
+					}
+				})
+			+ SSegmentedControl<EAxis::Type>::Slot(EAxis::X)
+			.Text(LOCTEXT("MirrorXLabel", "X"))
+			+ SSegmentedControl<EAxis::Type>::Slot(EAxis::Y)
+			.Text(LOCTEXT("MirrorYLabel", "Y"))
+			+ SSegmentedControl<EAxis::Type>::Slot(EAxis::Z)
+			.Text(LOCTEXT("MirrorZLabel", "Z"));
+
+			TSharedRef<SSegmentedControl<EDataflowEditorToolMirrorDirection>> MirrortSideButtonsWidget =
+				SNew(SSegmentedControl<EDataflowEditorToolMirrorDirection>)
+				.ToolTipText(LOCTEXT("MirrorDirectionTooltip", "The direction that determines what side of the plane to copy weights from."))
+				.Value_Lambda([this]()
+					{
+						return ToolProperties.IsValid() ? ToolProperties->MirrorProperties.MirrorDirection : EDataflowEditorToolMirrorDirection::PositiveToNegative;
+					})
+				.OnValueChanged_Lambda([this](EDataflowEditorToolMirrorDirection Mode)
+					{
+						if (ToolProperties.IsValid())
+						{
+							ToolProperties->MirrorProperties.MirrorDirection = Mode;
+						}
+					})
+				+ SSegmentedControl<EDataflowEditorToolMirrorDirection>::Slot(EDataflowEditorToolMirrorDirection::PositiveToNegative)
+				.Text(LOCTEXT("MirrorPosToNegLabel", "+ to -"))
+				+ SSegmentedControl<EDataflowEditorToolMirrorDirection>::Slot(EDataflowEditorToolMirrorDirection::NegativeToPositive)
+				.Text(LOCTEXT("MirrorNegToPosLabel", "- to +"));
+
+		TSharedRef<SButton> MirrorButtonWidget = 
+			SNew(SButton)
+			.HAlign(HAlign_Center)
+			.Text(LOCTEXT("MirrorWeightsButtonLabel", "Mirror"))
+			.ToolTipText(LOCTEXT("MirrorButtonTooltip",
+				"Values are copied across the given plane in the given direction.\n"
+				"This command operates on the selected vertices."))
+			.OnClicked_Lambda([this]()
+				{
+					if (Tool.IsValid())
+					{
+						Tool->MirrorValues();
+					}
+					return FReply::Handled();
+				});
+
+		// MIRROR WEIGHTS category
+		EditValuesCategory.AddCustomRow(LOCTEXT("MirrorWeightsRow", "Mirror"), false)
+			.WholeRowContent()
+			[
+				SNew(SVerticalBox)
+					.IsEnabled_Lambda([this] 
+						{ 
+							return Tool.IsValid() ? Tool->HasSelection() : false; 
+						})
+					+ SVerticalBox::Slot()
+					.Padding(WeightEditHorizontalPadding, WeightEditVerticalPadding)
+					[
+						SNew(SHorizontalBox)
+						+ SHorizontalBox::Slot()
+						.VAlign(VAlign_Center)
+						.FillWidth(WeightEditingLabelsPercent)
+						[
+							SNew(STextBlock)
+								.Text(LOCTEXT("MirrorPlaneLabel", "Mirror Plane"))
+								.Font(FAppStyle::GetFontStyle(TEXT("PropertyWindow.NormalFont")))
+								.ToolTipText(LOCTEXT("MirrorPlaneTooltip", "The plane to copy weights across."))
+						]
+						+ SHorizontalBox::Slot()
+						.FillWidth(1.f)
+						[
+							SNew(SHorizontalBox)
+							+ SHorizontalBox::Slot()
+							[
+								XYZButtonsWidget
+							]
+
+							+ SHorizontalBox::Slot()
+							[
+								MirrortSideButtonsWidget
+							]
+						]
+					]
+					+ SVerticalBox::Slot()
+					.Padding(0.f, WeightEditVerticalPadding)
+					[
+						SNew(SBox)
+						[
+							MirrorButtonWidget
+						]
+					]
+			];
+	}
+
 	///////////////////////////////////////////////////////////////////////////////////////////////////////////////
 
 	void FVertexAttributePaintToolDetailCustomization::AddSelectionUI(IDetailLayoutBuilder& DetailBuilder) const
@@ -1090,6 +1199,8 @@ namespace UE::Dataflow::Editor
 
 			MakeSelectionRelaxOperationRow(EditValuesCategory);
 
+			MakeSelectionMirrorOperationRow(EditValuesCategory);
+
 			MakeSelectionPruneOperationRow(EditValuesCategory);
 
 			MakeSelectionCopyAndPasteRow(EditValuesCategory);
diff --git a/Engine/Plugins/Experimental/Dataflow/Source/DataflowEditor/Private/DataflowEditorTools/DataflowEditorVertexAttributePaintToolDetailCustomization.h b/Engine/Plugins/Experimental/Dataflow/Source/DataflowEditor/Private/DataflowEditorTools/DataflowEditorVertexAttributePaintToolDetailCustomization.h
--- a/Engine/Plugins/Experimental/Dataflow/Source/DataflowEditor/Private/DataflowEditorTools/DataflowEditorVertexAttributePaintToolDetailCustomization.h
+++ b/Engine/Plugins/Experimental/Dataflow/Source/DataflowEditor/Private/DataflowEditorTools/DataflowEditorVertexAttributePaintToolDetailCustomization.h
@@ -52,6 +52,7 @@ namespace UE::Dataflow::Editor
 		void MakeSelectionReplaceOperationRow(IDetailCategoryBuilder& EditValuesCategory) const;
 		void MakeSelectionInvertOperationRow(IDetailCategoryBuilder& EditValuesCategory) const;
 		void MakeSelectionRelaxOperationRow(IDetailCategoryBuilder& EditValuesCategory) const;
+		void MakeSelectionMirrorOperationRow(IDetailCategoryBuilder& EditValuesCategory) const;
 		void MakeSelectionPruneOperationRow(IDetailCategoryBuilder& EditValuesCategory) const;
 		void MakeSelectionCopyAndPasteRow(IDetailCategoryBuilder& EditValuesCategory) const;
 		
diff --git a/Engine/Plugins/Experimental/Dataflow/Source/DataflowEditor/Public/DataflowEditorTools/DataflowEditorToolEnums.h b/Engine/Plugins/Experimental/Dataflow/Source/DataflowEditor/Public/DataflowEditorTools/DataflowEditorToolEnums.h
--- a/Engine/Plugins/Experimental/Dataflow/Source/DataflowEditor/Public/DataflowEditorTools/DataflowEditorToolEnums.h
+++ b/Engine/Plugins/Experimental/Dataflow/Source/DataflowEditor/Public/DataflowEditorTools/DataflowEditorToolEnums.h
@@ -42,3 +42,11 @@ enum class EDataflowEditorToolVisibilityType : uint8
 	None,
 	Unoccluded,
 };
+
+// mirror direction mode
+UENUM()
+enum class EDataflowEditorToolMirrorDirection : uint8
+{
+	PositiveToNegative,
+	NegativeToPositive,
+};
diff --git a/Engine/Plugins/Experimental/Dataflow/Source/DataflowEditor/Public/DataflowEditorTools/DataflowEditorVertexAttributePaintTool.h b/Engine/Plugins/Experimental/Dataflow/Source/DataflowEditor/Public/DataflowEditorTools/DataflowEditorVertexAttributePaintTool.h
--- a/Engine/Plugins/Experimental/Dataflow/Source/DataflowEditor/Public/DataflowEditorTools/DataflowEditorVertexAttributePaintTool.h
+++ b/Engine/Plugins/Experimental/Dataflow/Source/DataflowEditor/Public/DataflowEditorTools/DataflowEditorVertexAttributePaintTool.h
@@ -189,6 +189,18 @@ struct FDataflowEditorVertexAttributePaintToolDisplayProperties
 	FDataflowColorCurveOwner WhiteColorRamp;
 };
 
+USTRUCT()
+struct FDataflowEditorVertexAttributePaintToolMirrorProperties
+{
+	GENERATED_BODY()
+
+	UPROPERTY()
+	TEnumAsByte<EAxis::Type> MirrorAxis = EAxis::X;
+
+	UPROPERTY()
+	EDataflowEditorToolMirrorDirection MirrorDirection = EDataflowEditorToolMirrorDirection::PositiveToNegative;
+};
+
 ///////////////////////////////////////////////////////////////////////////////////////////////////////////////////////////////////////////////////
 
 UCLASS(MinimalAPI, config = EditorSettings)
@@ -216,6 +228,9 @@ public:
 	UPROPERTY(config)
 	FDataflowEditorVertexAttributePaintToolDisplayProperties DisplayProperties;
 
+	UPROPERTY(config)
+	FDataflowEditorVertexAttributePaintToolMirrorProperties MirrorProperties;
+
 	const FDataflowEditorVertexAttributePaintToolBrushConfig& GetBrushConfig(EDataflowEditorToolEditOperation BrushMode) const;
 	void SetBrushConfig(EDataflowEditorToolEditOperation BrushMode, const FDataflowEditorVertexAttributePaintToolBrushConfig& BrushConfig);
 
@@ -375,6 +390,7 @@ public:
 
 	// selection operations
 	UE_API void ApplyValueToSelection(EDataflowEditorToolEditOperation Operation, float InValue);
+	UE_API void MirrorValues();
 
 	UE_API TArray<int32> GetSelectedVertices() const;
 
@@ -500,9 +516,30 @@ protected:
 
 	static constexpr int32 PaintBrushId = 0;
 	static constexpr int32 SmoothBrushId = 1;
+
+private:
+	// this structure holds information for mirroring 
+	struct FMirrorData
+	{
+	public:
+		// lazily updates the mirror data tables for the current skeleton/mesh/mirror plane
+		void EnsureMirrorDataIsUpdated(const FDynamicMesh3& Mesh, EAxis::Type InMirrorAxis, EDataflowEditorToolMirrorDirection InMirrorDirection);
+
+		void FindMirroredIndices(const FDynamicMesh3& Mesh, const TArray<int32>& SelectedVertices, TArray<int32>& OutVerticesToUpdate);
+
+	private:
+		// return true if the point lies on the TARGET side of the mirror plane
+		bool IsPointOnTargetMirrorSide(const FVector& InPoint) const;
+
+		EAxis::Type Axis = EAxis::X;
+		EDataflowEditorToolMirrorDirection Direction = EDataflowEditorToolMirrorDirection::PositiveToNegative;
+		TMap<int32, int32> VertexMap; // <Target, Source>
+	};
+	FMirrorData MirrorData;
 };
 
 
 
+
 #undef UE_API
 
