Fix crash when clearing DataflowAsset on a DataflowObject's instance.

diff --git a/Engine/Plugins/Experimental/Dataflow/Source/DataflowEditor/Private/Dataflow/DataflowEditorToolkit.cpp b/Engine/Plugins/Experimental/Dataflow/Source/DataflowEditor/Private/Dataflow/DataflowEditorToolkit.cpp
--- a/Engine/Plugins/Experimental/Dataflow/Source/DataflowEditor/Private/Dataflow/DataflowEditorToolkit.cpp
+++ b/Engine/Plugins/Experimental/Dataflow/Source/DataflowEditor/Private/Dataflow/DataflowEditorToolkit.cpp
@@ -184,21 +184,24 @@ namespace UE::Dataflow::Private
 				}
 			}
 
-			for (const UDataflowEdNode* const PinnedNode : EditorContent->GetDataflowAsset()->GetWireframeRenderTargets())
+			if (EditorContent->GetDataflowAsset())
 			{
-				if (const TSharedPtr<const FDataflowNode> PinnedDataflowNode = PinnedNode->GetDataflowNode())
+				for (const UDataflowEdNode* const PinnedNode : EditorContent->GetDataflowAsset()->GetWireframeRenderTargets())
 				{
-					if(PinnedDataflowNode->IsA(NodeType::StaticType()))
+					if (const TSharedPtr<const FDataflowNode> PinnedDataflowNode = PinnedNode->GetDataflowNode())
 					{
-						if (PinnedDataflowNode->CanDebugDrawViewMode(CurrentViewMode))
+						if (PinnedDataflowNode->IsA(NodeType::StaticType()))
 						{
-							DebugDrawObject->ResetAllState();
+							if (PinnedDataflowNode->CanDebugDrawViewMode(CurrentViewMode))
+							{
+								DebugDrawObject->ResetAllState();
 
-							FDataflowNode::FDebugDrawParameters DebugDrawParameters;
-							DebugDrawParameters.bNodeIsSelected = false;
-							DebugDrawParameters.bNodeIsPinned = true;
-							DebugDrawParameters.CurrentViewMode = CurrentViewMode;
-							PinnedDataflowNode->DebugDraw(*DataflowContext, *DebugDrawObject, DebugDrawParameters);
+								FDataflowNode::FDebugDrawParameters DebugDrawParameters;
+								DebugDrawParameters.bNodeIsSelected = false;
+								DebugDrawParameters.bNodeIsPinned = true;
+								DebugDrawParameters.CurrentViewMode = CurrentViewMode;
+								PinnedDataflowNode->DebugDraw(*DataflowContext, *DebugDrawObject, DebugDrawParameters);
+							}
 						}
 					}
 				}
