Dataflow rendering : add caching feature for renderable types
- caching isoptional. It is useful when the generation of the object used by the component may be expensive

diff --git a/Engine/Plugins/Experimental/Dataflow/Source/DataflowEditor/Private/DataflowRendering/DataflowGeometryCollectionRenderableType.cpp b/Engine/Plugins/Experimental/Dataflow/Source/DataflowEditor/Private/DataflowRendering/DataflowGeometryCollectionRenderableType.cpp
--- a/Engine/Plugins/Experimental/Dataflow/Source/DataflowEditor/Private/DataflowRendering/DataflowGeometryCollectionRenderableType.cpp
+++ b/Engine/Plugins/Experimental/Dataflow/Source/DataflowEditor/Private/DataflowRendering/DataflowGeometryCollectionRenderableType.cpp
@@ -47,35 +47,42 @@ namespace UE::Dataflow::Private
 
 			if (TransformAttribute && TransformIndexAttribute)
 			{
-				// TODO(dataflow) : need a version of ConvertGeometryCollectionToDynamicMesh that only takes a managed array collection
-				if (TUniquePtr<const FGeometryCollection> GeomCollection = TUniquePtr<FGeometryCollection>(Collection.NewCopy<FGeometryCollection>()))
+				TArray<FDynamicMesh3> DynamicMeshes;
+				if (Instance.HasUptoDateCachedValue())
 				{
-					// prepare the dynamic meshes
-					TArray<FDynamicMesh3> DynamicMeshes;
-					DynamicMeshes.SetNum(TransformAttribute->Num());
-
-					TArrayView<const FTransform3f> BoneTransforms = MakeArrayView(TransformAttribute->GetConstArray());
+					DynamicMeshes = Instance.GetCachedValue<TArray<FDynamicMesh3>>();
+				}
+				else
+				{
+					// TODO(dataflow) : need a version of ConvertGeometryCollectionToDynamicMesh that only takes a managed array collection
+					if (TUniquePtr<const FGeometryCollection> GeomCollection = TUniquePtr<FGeometryCollection>(Collection.NewCopy<FGeometryCollection>()))
+					{
+						// prepare the dynamic meshes
+						DynamicMeshes.SetNum(TransformAttribute->Num());
 
-					// convert to dynamic meshes in parallel
-					ParallelFor(DynamicMeshes.Num(),
-						[&DynamicMeshes, &BoneTransforms, &GeomCollection](int32 TransformIndex)
-						{
-							FDynamicMesh3& DynamicMesh = DynamicMeshes[TransformIndex];
-							TArrayView<const int32> TransformIndicesToConvert(&TransformIndex, 1);
+						TArrayView<const FTransform3f> BoneTransforms = MakeArrayView(TransformAttribute->GetConstArray());
 
-							FTransform UnusedTransform;
-							::ConvertGeometryCollectionToDynamicMesh(DynamicMesh, UnusedTransform, false, *GeomCollection, true, BoneTransforms, true, TransformIndicesToConvert);
-						});
+						// convert to dynamic meshes in parallel
+						ParallelFor(DynamicMeshes.Num(),
+							[&DynamicMeshes, &BoneTransforms, &GeomCollection](int32 TransformIndex)
+							{
+								FDynamicMesh3& DynamicMesh = DynamicMeshes[TransformIndex];
+								TArrayView<const int32> TransformIndicesToConvert(&TransformIndex, 1);
 
-					// only create components for the one that have vertices
-					for (int32 TransformIndex = 0; TransformIndex < TransformAttribute->Num(); ++TransformIndex)
+								FTransform UnusedTransform;
+								::ConvertGeometryCollectionToDynamicMesh(DynamicMesh, UnusedTransform, false, *GeomCollection, true, BoneTransforms, true, TransformIndicesToConvert);
+							});
+					}
+					Instance.CacheValue(DynamicMeshes);
+				}
+				// only create components for the one that have vertices
+				for (int32 TransformIndex = 0; TransformIndex < TransformAttribute->Num(); ++TransformIndex)
+				{
+					if (DynamicMeshes[TransformIndex].VertexCount() > 0)
 					{
-						if (DynamicMeshes[TransformIndex].VertexCount() > 0)
+						if (UDynamicMeshComponent* Component = OutComponents.AddNewComponent<UDynamicMeshComponent>())
 						{
-							if (UDynamicMeshComponent* Component = OutComponents.AddNewComponent<UDynamicMeshComponent>())
-							{
-								Component->SetMesh(MoveTemp(DynamicMeshes[TransformIndex]));
-							}
+							Component->SetMesh(MoveTemp(DynamicMeshes[TransformIndex]));
 						}
 					}
 				}
diff --git a/Engine/Plugins/Experimental/Dataflow/Source/DataflowEditor/Private/DataflowRendering/DataflowRenderableTypeInstance.cpp b/Engine/Plugins/Experimental/Dataflow/Source/DataflowEditor/Private/DataflowRendering/DataflowRenderableTypeInstance.cpp
--- a/Engine/Plugins/Experimental/Dataflow/Source/DataflowEditor/Private/DataflowRendering/DataflowRenderableTypeInstance.cpp
+++ b/Engine/Plugins/Experimental/Dataflow/Source/DataflowEditor/Private/DataflowRendering/DataflowRenderableTypeInstance.cpp
@@ -58,4 +58,18 @@ namespace UE::Dataflow
 			}
 		}
 	}
+
+	FContextCacheKey  FRenderableTypeInstance::ComputeCacheKey() const
+	{
+		FContextCacheKey OutputCacheKey = 0;
+		if (Node)
+		{
+			if (const FDataflowOutput* Output = Node->FindOutput(OutputName))
+			{
+				OutputCacheKey = Output->CacheKey();
+			}
+		}
+		const uint32 RenderGroupHash = GetTypeHash(RenderableType? RenderableType->GetRenderGroup(): NAME_None);
+		return ::HashCombine(OutputCacheKey, RenderGroupHash);
+	}
 }
diff --git a/Engine/Plugins/Experimental/Dataflow/Source/DataflowEditor/Public/DataflowRendering/DataflowRenderableTypeInstance.h b/Engine/Plugins/Experimental/Dataflow/Source/DataflowEditor/Public/DataflowRendering/DataflowRenderableTypeInstance.h
--- a/Engine/Plugins/Experimental/Dataflow/Source/DataflowEditor/Public/DataflowRendering/DataflowRenderableTypeInstance.h
+++ b/Engine/Plugins/Experimental/Dataflow/Source/DataflowEditor/Public/DataflowRendering/DataflowRenderableTypeInstance.h
@@ -46,6 +46,33 @@ namespace UE::Dataflow
 			return GetOutputValue(OutputName, Default);
 		}
 
+		bool HasUptoDateCachedValue() const
+		{
+			if (Node)
+			{
+				const FContextCacheKey Key = ComputeCacheKey();
+				return EvaluationContext.HasData(Key, Node->GetTimestamp());
+			}
+			return false;
+		}
+
+		template<typename T>
+		const T& GetCachedValue() const
+		{
+			const FContextCacheKey Key = ComputeCacheKey();
+			return EvaluationContext.GetData<T>(Key, nullptr);
+		}
+
+		template<typename T>
+		void CacheValue(T&& InValue) const
+		{
+			if (Node)
+			{
+				const FContextCacheKey Key = ComputeCacheKey();
+				EvaluationContext.SetData(Key, nullptr, Forward<T>(InValue), Node->GetGuid(), /*Node->GetValueHash()*/0, Node->GetTimestamp());
+			}
+		}
+
 	private:
 		template <typename T>
 		const T& GetOutputValue(FName Name, const T& Default = T()) const
@@ -64,6 +91,8 @@ namespace UE::Dataflow
 			return Default;
 		}
 
+		FContextCacheKey ComputeCacheKey() const;
+
 		FRenderableTypeInstance(FContext& InContext, const TSharedPtr<const FDataflowNode>& InNode, const FName InOutputName, const IRenderableType* InRenderableType)
 			: EvaluationContext(InContext)
 			, Node(InNode)
