Remove unused code causing some LOCTEXT check to fail

diff --git a/Engine/Plugins/Experimental/Dataflow/Source/DataflowEditor/Private/DataflowEditorTools/DataflowEditorVertexAttributePaintToolDetailCustomization.cpp b/Engine/Plugins/Experimental/Dataflow/Source/DataflowEditor/Private/DataflowEditorTools/DataflowEditorVertexAttributePaintToolDetailCustomization.cpp
--- a/Engine/Plugins/Experimental/Dataflow/Source/DataflowEditor/Private/DataflowEditorTools/DataflowEditorVertexAttributePaintToolDetailCustomization.cpp
+++ b/Engine/Plugins/Experimental/Dataflow/Source/DataflowEditor/Private/DataflowEditorTools/DataflowEditorVertexAttributePaintToolDetailCustomization.cpp
@@ -1205,111 +1205,6 @@ namespace UE::Dataflow::Editor
 
 			MakeSelectionCopyAndPasteRow(EditValuesCategory);
 		}
-
-#if 0
-		// MIRROR WEIGHTS category
-		EditWeightsCategory.AddCustomRow(LOCTEXT("MirrorWeightsRow", "Mirror"), false)
-		.WholeRowContent()
-		[
-			SNew(SVerticalBox)
-			.IsEnabled_Lambda([this]{ return Tool->HasActiveSelectionOnMainMesh(); })
-		
-			+ SVerticalBox::Slot()
-			.Padding(WeightEditHorizontalPadding, WeightEditVerticalPadding)
-			[
-				SNew(SHorizontalBox)
-			
-				+SHorizontalBox::Slot()
-				.VAlign(VAlign_Center)
-				.FillWidth(WeightEditingLabelsPercent)
-				[
-					SNew(STextBlock)
-					.Text(LOCTEXT("MirrorPlaneLabel", "Mirror Plane"))
-					.Font(FAppStyle::GetFontStyle(TEXT("PropertyWindow.NormalFont")))
-					.ToolTipText(LOCTEXT("MirrorPlaneTooltip", "The plane to copy weights across."))
-				]
-
-				+SHorizontalBox::Slot()
-				.FillWidth(1.f)
-				[
-					SNew(SHorizontalBox)
-					+SHorizontalBox::Slot()
-					[
-						SNew(SSegmentedControl<EAxis::Type>)
-						.ToolTipText(LOCTEXT("MirrorAxisTooltip",
-							"X: copies weights across the YZ plane.\n"
-							"Y: copies weights across the XZ plane.\n"
-							"Z: copies weights across the XY plane."))
-						.Value_Lambda([this]()
-						{
-							return ToolProperties.IsValid() ? ToolProperties->MirrorAxis : TEnumAsByte<EAxis::Type>(EAxis::X);
-						})
-						.OnValueChanged_Lambda([this](EAxis::Type Mode)
-						{
-							if (ToolProperties.IsValid())
-							{
-								ToolProperties->MirrorAxis = Mode;
-							}
-						})
-						+ SSegmentedControl<EAxis::Type>::Slot(EAxis::X)
-						.Text(LOCTEXT("MirrorXLabel", "X"))
-						+ SSegmentedControl<EAxis::Type>::Slot(EAxis::Y)
-						.Text(LOCTEXT("MirrorYLabel", "Y"))
-						+ SSegmentedControl<EAxis::Type>::Slot(EAxis::Z)
-						.Text(LOCTEXT("MirrorZLabel", "Z"))
-					]
-					
-					+SHorizontalBox::Slot()
-					[
-						SNew(SSegmentedControl<EMirrorDirection>)
-						.ToolTipText(LOCTEXT("MirrorDirectionTooltip", "The direction that determines what side of the plane to copy weights from."))
-						.Value_Lambda([this]()
-						{
-							return ToolProperties.IsValid() ? ToolProperties->MirrorDirection : EMirrorDirection::PositiveToNegative;
-						})
-						.OnValueChanged_Lambda([this](EMirrorDirection Mode)
-						{
-							if (ToolProperties.IsValid())
-							{
-								ToolProperties->MirrorDirection = Mode;
-							}
-						})
-						+ SSegmentedControl<EMirrorDirection>::Slot(EMirrorDirection::PositiveToNegative)
-						.Text(LOCTEXT("MirrorPosToNegLabel", "+ to -"))
-						+ SSegmentedControl<EMirrorDirection>::Slot(EMirrorDirection::NegativeToPositive)
-						.Text(LOCTEXT("MirrorNegToPosLabel", "- to +"))
-					]
-				]
-			]
-
-			+SVerticalBox::Slot()
-			.Padding(0.f, WeightEditVerticalPadding)
-			[
-				SNew(SBox)
-				[
-					SNew(SButton)
-					.HAlign(HAlign_Center)
-					.Text(LOCTEXT("MirrorWeightsButtonLabel", "Mirror"))
-					.ToolTipText(LOCTEXT("MirrorButtonTooltip",
-						"Weights are copied across the given plane in the given direction.\n"
-						"This command operates on the selected vertices."))
-					.OnClicked_Lambda([this]()
-					{
-						if (Tool.IsValid() && ToolProperties.IsValid())
-						{
-							Tool->MirrorWeights(ToolProperties->MirrorAxis, ToolProperties->MirrorDirection);
-						}
-						return FReply::Handled();
-					})
-				]
-			]
-		];
-
-		
-		
-
-	
-#endif // 0
 	}
 
 	void FVertexAttributePaintToolDetailCustomization::HideProperty(IDetailLayoutBuilder& DetailBuilder, FName PropertyName) const
