Dataflow rendeirng : add support for specifying a render group for an output in a node

diff --git a/Engine/Plugins/Experimental/Dataflow/Source/DataflowEditor/Private/DataflowRendering/DataflowRenderableTypeInstance.cpp b/Engine/Plugins/Experimental/Dataflow/Source/DataflowEditor/Private/DataflowRendering/DataflowRenderableTypeInstance.cpp
--- a/Engine/Plugins/Experimental/Dataflow/Source/DataflowEditor/Private/DataflowRendering/DataflowRenderableTypeInstance.cpp
+++ b/Engine/Plugins/Experimental/Dataflow/Source/DataflowEditor/Private/DataflowRendering/DataflowRenderableTypeInstance.cpp
@@ -8,6 +8,28 @@
 
 namespace UE::Dataflow
 {
+	namespace Private
+	{
+		bool IsVisibleRenderGroup(FName RenderGroup, const FDataflowOutput& Output)
+		{
+			if (const FProperty* Property = Output.GetProperty())
+			{
+				static const FName RenderGroupsMetaDataName = TEXT("DataflowRenderGroups");
+				const FString RenderGroupsMetaData = Property->GetMetaData(RenderGroupsMetaDataName);
+				if (!RenderGroupsMetaData.IsEmpty())
+				{
+					TArray<FString> RenderGroups;
+					RenderGroupsMetaData.ParseIntoArray(RenderGroups, TEXT(","));
+					if (RenderGroups.Num() > 0)
+					{
+						return RenderGroups.Contains(RenderGroup);
+					}
+				}
+			}
+			return true;
+		}
+	}
+
 	/* static */ void FRenderableTypeInstance::GetRenderableInstancesForNode(FContext& Context, const IDataflowConstructionViewMode& ViewMode, const FDataflowNode& Node, TArray<FRenderableTypeInstance>& OutRenderableTypeInstances)
 	{
 		OutRenderableTypeInstances.Reset();
@@ -25,7 +47,10 @@ namespace UE::Dataflow
 						FRenderableTypeInstance RenderableTypeInstance(Context, Node.AsShared(), Output->GetName(), RenderableType);
 						if (RenderableTypeInstance.CanRender())
 						{
-							OutRenderableTypeInstances.Emplace(RenderableTypeInstance);
+							if (Private::IsVisibleRenderGroup(RenderableType->GetRenderGroup(), *Output))
+							{
+								OutRenderableTypeInstances.Emplace(RenderableTypeInstance);
+							}
 						}
 					}
 				}
