Dataflow : fix issue where marquee select a node does not bring the associated tool up
- remove OnSingleNodeClicked and move the logic to open the tool at the end of the OnNodeSelectionChanged method

diff --git a/Engine/Plugins/Experimental/Dataflow/Source/DataflowEditor/Private/Dataflow/DataflowEditorToolkit.cpp b/Engine/Plugins/Experimental/Dataflow/Source/DataflowEditor/Private/Dataflow/DataflowEditorToolkit.cpp
--- a/Engine/Plugins/Experimental/Dataflow/Source/DataflowEditor/Private/Dataflow/DataflowEditorToolkit.cpp
+++ b/Engine/Plugins/Experimental/Dataflow/Source/DataflowEditor/Private/Dataflow/DataflowEditorToolkit.cpp
@@ -1445,10 +1445,13 @@ void FDataflowEditorToolkit::OnNodeSelectionChanged(const TSet<UObject*>& InNewS
 	// todo : move the tool logic somewhere else: we should not be activelly looking for a collection here, we should have some sort of interface to talk to the tools
 	UDataflowEditorMode* const DataflowMode = Cast<UDataflowEditorMode>(EditorModeManager->GetActiveScriptableMode(UDataflowEditorMode::EM_DataflowEditorModeId));
 
+	bool bSingleSelectionChanged = false;
 	if (const TObjectPtr<UDataflowBaseContent>& EditorContent = GetEditorContent(); EditorContent->GetDataflowAsset())
 	{
 		if (NewSelectedNode != PrevSelectedNode)
 		{
+			bSingleSelectionChanged = true;
+
 			EditorContent->SetSelectedNode(nullptr);
 			EditorContent->SetSelectedCollection(nullptr, /*bCollectionIsInput=*/ false);
 
@@ -1457,7 +1460,7 @@ void FDataflowEditorToolkit::OnNodeSelectionChanged(const TSet<UObject*>& InNewS
 				if (DataflowMode)
 				{
 					// Close any running tool by accepting the changes.
-					// OnNodeSingleClicked() will start a new tool if a new node was clicked.
+					// new tool will start at the end of this method if a single node was selected
 					DataflowMode->ShutdownActiveToolIfNeeded(EToolShutdownType::Accept);
 
 					EditorContent->SetSelectedNode(NewSelectedNode);
@@ -1504,6 +1507,13 @@ void FDataflowEditorToolkit::OnNodeSelectionChanged(const TSet<UObject*>& InNewS
 	{
 		DataflowMode->SetConstructionViewModeForNode(NewSelectedNode);
 		DataflowMode->SetPendingNodeSelectionChanged(false);
+
+		// if we have one single element selected, start the associated tool 
+		if (bSingleSelectionChanged && SelectedDataflowNodes.Num() == 1)
+		{
+			// Start the corresponding tool
+			DataflowMode->StartToolForSelectedNode(NewSelectedNode);
+		}
 	}
 }
 
@@ -1605,20 +1615,6 @@ void FDataflowEditorToolkit::OnNodeInvalidated(UDataflow& DataflowAsset, FDatafl
 	}
 }
 
-void FDataflowEditorToolkit::OnNodeSingleClicked(UObject* ClickedNode) const
-{
-	UDataflowEditorMode* DataflowMode = Cast<UDataflowEditorMode>(EditorModeManager->GetActiveScriptableMode(UDataflowEditorMode::EM_DataflowEditorModeId));
-	if (DataflowMode)
-	{
-		if (GraphEditor && GraphEditor->GetSingleSelectedNode() == ClickedNode)
-		{
-			// Start the corresponding tool
-			DataflowMode->StartToolForSelectedNode(ClickedNode);
-		}
-	}
-}
-
-
 void FDataflowEditorToolkit::OnNodeDeleted(const TSet<UObject*>& NewSelection)
 {
 	const TSet<TObjectPtr<UDataflowEdNode>> DeletedNodes = FilterDataflowEdNodesFromSet(NewSelection);
@@ -2025,7 +2021,6 @@ TSharedRef<SDataflowGraphEditor> FDataflowEditorToolkit::CreateGraphEditorWidget
 	SGraphEditor::FGraphEditorEvents InEvents;
 	InEvents.OnVerifyTextCommit = FOnNodeVerifyTextCommit::CreateSP(this, &FDataflowEditorToolkit::OnNodeVerifyTitleCommit);
 	InEvents.OnTextCommitted = FOnNodeTextCommitted::CreateSP(this, &FDataflowEditorToolkit::OnNodeTitleCommitted);
-	InEvents.OnNodeSingleClicked = SGraphEditor::FOnNodeSingleClicked::CreateSP(this, &FDataflowEditorToolkit::OnNodeSingleClicked);
 	InEvents.OnNodeDoubleClicked = FSingleNodeEvent::CreateSP(this, &FDataflowEditorToolkit::OnNodeDoubleClicked);
 
 
diff --git a/Engine/Plugins/Experimental/Dataflow/Source/DataflowEditor/Public/Dataflow/DataflowEditorToolkit.h b/Engine/Plugins/Experimental/Dataflow/Source/DataflowEditor/Public/Dataflow/DataflowEditorToolkit.h
--- a/Engine/Plugins/Experimental/Dataflow/Source/DataflowEditor/Public/Dataflow/DataflowEditorToolkit.h
+++ b/Engine/Plugins/Experimental/Dataflow/Source/DataflowEditor/Public/Dataflow/DataflowEditorToolkit.h
@@ -128,10 +128,8 @@ protected:
 	UE_API void OnNodeTitleCommitted(const FText& InNewText, ETextCommit::Type InCommitType, UEdGraphNode* GraphNode);
 	UE_API void OnNodeDoubleClicked(UEdGraphNode* ClickedNode);
 	UE_API void OnNodeSelectionChanged(const TSet<UObject*>& NewSelection);
-	UE_API void OnEvaluateSelectedNodes(const TSet<UObject*>& SelectedNodes);
 	UE_API void OnNodeInvalidated(UDataflow& DataflowAsset, FDataflowNode& Node);
 	UE_API void OnNodeDeleted(const TSet<UObject*>& NewSelection);
-	UE_API void OnNodeSingleClicked(UObject* ClickedNode) const;
 	UE_API void OnAssetPropertyValueChanged(const FPropertyChangedEvent& PropertyChangedEvent);
 	UE_API void OnConstructionViewSelectionChanged(const TArray<UPrimitiveComponent*>& SelectedComponents, const TArray<FDataflowBaseElement*>& SelectedElements);
 	UE_API void OnSimulationViewSelectionChanged(const TArray<UPrimitiveComponent*>& SelectedComponents, const TArray<FDataflowBaseElement*>& SelectedElements);
