Dataflow : add undo support for the new vertex attribute ( weightmap nodes) painting tool

diff --git a/Engine/Plugins/ChaosClothAssetEditor/Source/ChaosClothAssetDataflowNodes/Private/ChaosClothAsset/WeightMapNode.cpp b/Engine/Plugins/ChaosClothAssetEditor/Source/ChaosClothAssetDataflowNodes/Private/ChaosClothAsset/WeightMapNode.cpp
--- a/Engine/Plugins/ChaosClothAssetEditor/Source/ChaosClothAssetDataflowNodes/Private/ChaosClothAsset/WeightMapNode.cpp
+++ b/Engine/Plugins/ChaosClothAssetEditor/Source/ChaosClothAssetDataflowNodes/Private/ChaosClothAsset/WeightMapNode.cpp
@@ -294,7 +294,7 @@ void FChaosClothAssetWeightMapNode::GetVertexAttributeValues(UE::Dataflow::FCont
 	}
 }
 
-void FChaosClothAssetWeightMapNode::SetVertexAttributeValues(UE::Dataflow::FContext& Context, const TArray<float>& InValues)
+void FChaosClothAssetWeightMapNode::SetVertexAttributeValues(UE::Dataflow::FContext& Context, const TArray<float>& InValues, const TArray<int32>& InWeightIndices)
 {
 	using namespace UE::Chaos::ClothAsset;
 
@@ -319,6 +319,11 @@ void FChaosClothAssetWeightMapNode::SetVertexAttributeValues(UE::Dataflow::FCont
 	}
 }
 
+void FChaosClothAssetWeightMapNode::SwapStoredAttributeValuesWith(TArray<float>& OtherValues)
+{
+	Swap(VertexWeights, OtherValues);
+}
+
 void FChaosClothAssetWeightMapNode::SetVertexWeights(const TConstArrayView<float> InputMap, const TArray<float>& FinalValues)
 {
 	UE::Chaos::ClothAsset::Private::SetVertexWeights(InputMap, FinalValues, MapOverrideType, GetVertexWeights());
diff --git a/Engine/Plugins/ChaosClothAssetEditor/Source/ChaosClothAssetDataflowNodes/Public/ChaosClothAsset/WeightMapNode.h b/Engine/Plugins/ChaosClothAssetEditor/Source/ChaosClothAssetDataflowNodes/Public/ChaosClothAsset/WeightMapNode.h
--- a/Engine/Plugins/ChaosClothAssetEditor/Source/ChaosClothAssetDataflowNodes/Public/ChaosClothAsset/WeightMapNode.h
+++ b/Engine/Plugins/ChaosClothAssetEditor/Source/ChaosClothAssetDataflowNodes/Public/ChaosClothAsset/WeightMapNode.h
@@ -102,7 +102,9 @@ private:
 	//~ Begin - FDataflowVertexAttributeEditableNode API
 	virtual void GetSupportedViewModes(UE::Dataflow::FContext& Context, TArray<FName>& OutViewModeNames) const override;
 	virtual void GetVertexAttributeValues(UE::Dataflow::FContext& Context, TArray<float>& OutValues) const override;
-	virtual void SetVertexAttributeValues(UE::Dataflow::FContext& Context, const TArray<float>& InValues) override;
+	virtual void SetVertexAttributeValues(UE::Dataflow::FContext& Context, const TArray<float>& InValues, const TArray<int32>& InWeightIndices) override;
+	virtual const TArray<float>& GetStoredAttributeValues() const override { return VertexWeights; }
+	virtual void SwapStoredAttributeValuesWith(TArray<float>& OtherValues) override;
 	//~ End - FDataflowVertexAttributeEditableNode API
 
 	void OnTransfer(UE::Dataflow::FContext& Context);
diff --git a/Engine/Plugins/Experimental/Dataflow/Source/DataflowEditor/Private/DataflowEditorTools/DataflowEditorVertexAttributePaintTool.cpp b/Engine/Plugins/Experimental/Dataflow/Source/DataflowEditor/Private/DataflowEditorTools/DataflowEditorVertexAttributePaintTool.cpp
--- a/Engine/Plugins/Experimental/Dataflow/Source/DataflowEditor/Private/DataflowEditorTools/DataflowEditorVertexAttributePaintTool.cpp
+++ b/Engine/Plugins/Experimental/Dataflow/Source/DataflowEditor/Private/DataflowEditorTools/DataflowEditorVertexAttributePaintTool.cpp
@@ -513,7 +513,7 @@ float FDataflowEditorVertexAttributePaintToolData::GetAverageValue(const TArray<
 	return (float)AverageValue;
 }
 
-void FDataflowEditorVertexAttributePaintToolData::UpdateNode(TObjectPtr<UDataflowContextObject> InDataflowEditorContextObject)
+void FDataflowEditorVertexAttributePaintToolData::UpdateNode(TObjectPtr<UDataflowContextObject> InDataflowEditorContextObject, IToolsContextTransactionsAPI* TransactionAPI)
 {
 	if (ActiveWeightMap && NodeToUpdate)
 	{
@@ -521,21 +521,29 @@ void FDataflowEditorVertexAttributePaintToolData::UpdateNode(TObjectPtr<UDataflo
 		{
 			if (TSharedPtr<UE::Dataflow::FEngineContext> DataflowContext = InDataflowEditorContextObject->GetDataflowContext())
 			{
+				// Save previous state for undo
+				if (UDataflow* const Dataflow = InDataflowEditorContextObject ? InDataflowEditorContextObject->GetDataflowAsset(): nullptr)
+				{
+					if (TransactionAPI)
+					{
+						TransactionAPI->AppendChange(
+							Dataflow,
+							NodeToUpdate->MakeEditNodeToolChange(),
+							LOCTEXT("DataflowEditorVertexAttributePaintTool_ChangeDescription", "Update Weight Map Node")
+						);
+					}
+				}
+
+				// apply the new values to the node
 				TArray<float> WeightsToApply;
 				const int32 NumVertices = Mesh->VertexCount();
 				WeightsToApply.SetNumUninitialized(NumVertices);
 				for (int32 VertexID = 0; VertexID < NumVertices; ++VertexID)
 				{
-					if (bHaveDynamicMeshToWeightConversion)
-					{
-						ActiveWeightMap->GetValue(VertexID, &WeightsToApply[DynamicMeshToWeight[VertexID]]);
-					}
-					else
-					{
-						ActiveWeightMap->GetValue(VertexID, &WeightsToApply[VertexID]);
-					}
+					ActiveWeightMap->GetValue(VertexID, &WeightsToApply[VertexID]);
 				}
-				NodeToUpdate->SetVertexAttributeValues(*DataflowContext, WeightsToApply);
+				// WeightsToApply represents the weights of the selected sub-mesh, DynamicMeshToWeight allow to remap properly to the full collection mesh
+				NodeToUpdate->SetVertexAttributeValues(*DataflowContext, WeightsToApply, DynamicMeshToWeight);
 			}
 		}
 	}
@@ -807,11 +815,14 @@ void UDataflowEditorVertexAttributePaintTool::Shutdown(EToolShutdownType Shutdow
 
 void UDataflowEditorVertexAttributePaintTool::CommitResult(UBaseDynamicMeshComponent* Component, bool bModifiedTopology)
 {
-	GetToolManager()->BeginUndoTransaction(LOCTEXT("DataflowEditorVertexAttributePaintTool_TransactionName", "Paint Weights"));
+	if (UInteractiveToolManager* ToolManager = GetToolManager())
+	{
+		ToolManager->BeginUndoTransaction(LOCTEXT("DataflowEditorVertexAttributePaintTool_TransactionName", "Paint Weights"));
 
-	VertexData.UpdateNode(DataflowEditorContextObject);
+		VertexData.UpdateNode(DataflowEditorContextObject, ToolManager->GetContextTransactionsAPI());
 
-	GetToolManager()->EndUndoTransaction();
+		ToolManager->EndUndoTransaction();
+	}
 }
 
 
@@ -1547,7 +1558,9 @@ TArray<int32> UDataflowEditorVertexAttributePaintTool::GetSelectedVertices() con
 {
 	if (MeshSelector)
 	{
-		return MeshSelector->GetSelectedVertices();
+		// we do not want to remap to original as the tool already take care of doing the remapping when necessary 
+		// we operate in the selected sub-mesh vertex range
+		return MeshSelector->GetSelectedVertices(/*bRemapToOriginal*/false);
 	}
 	return {};
 }
diff --git a/Engine/Plugins/Experimental/Dataflow/Source/DataflowEditor/Public/DataflowEditorTools/DataflowEditorVertexAttributePaintTool.h b/Engine/Plugins/Experimental/Dataflow/Source/DataflowEditor/Public/DataflowEditorTools/DataflowEditorVertexAttributePaintTool.h
--- a/Engine/Plugins/Experimental/Dataflow/Source/DataflowEditor/Public/DataflowEditorTools/DataflowEditorVertexAttributePaintTool.h
+++ b/Engine/Plugins/Experimental/Dataflow/Source/DataflowEditor/Public/DataflowEditorTools/DataflowEditorVertexAttributePaintTool.h
@@ -33,6 +33,7 @@
 
 #define UE_API DATAFLOWEDITOR_API
 
+class IToolsContextTransactionsAPI;
 class UMeshElementsVisualizer;
 class UDataflowWeightMapSmoothBrushOpProps;
 class UDataflowVertexAttributePaintBrushOpProps;
@@ -292,7 +293,7 @@ public:
 		}
 	}
 
-	void UpdateNode(TObjectPtr<UDataflowContextObject> InDataflowEditorContextObject);
+	void UpdateNode(TObjectPtr<UDataflowContextObject> InDataflowEditorContextObject, IToolsContextTransactionsAPI* TransactionAPI);
 
 private:
 	// DynamicMesh might be unwelded mesh, but weights are on the welded mesh.
diff --git a/Engine/Plugins/Experimental/Dataflow/Source/DataflowNodes/Private/Dataflow/DataflowCollectionAddScalarVertexPropertyNode.cpp b/Engine/Plugins/Experimental/Dataflow/Source/DataflowNodes/Private/Dataflow/DataflowCollectionAddScalarVertexPropertyNode.cpp
--- a/Engine/Plugins/Experimental/Dataflow/Source/DataflowNodes/Private/Dataflow/DataflowCollectionAddScalarVertexPropertyNode.cpp
+++ b/Engine/Plugins/Experimental/Dataflow/Source/DataflowNodes/Private/Dataflow/DataflowCollectionAddScalarVertexPropertyNode.cpp
@@ -4,6 +4,7 @@
 
 #include "Dataflow/DataflowInputOutput.h"
 #include "Dataflow/DataflowTools.h"
+#include "InteractiveToolChange.h"
 #include "Misc/LazySingleton.h"
 
 #include UE_INLINE_GENERATED_CPP_BY_NAME(DataflowCollectionAddScalarVertexPropertyNode)
@@ -141,6 +142,65 @@ FDataflowVertexAttributeEditableNode::FDataflowVertexAttributeEditableNode(const
 	: FDataflowNode(InParam, InGuid)
 {}
 
+const TArray<float>& FDataflowVertexAttributeEditableNode::GetStoredAttributeValues() const 
+{ 
+	static const TArray<float> ConstEmptyArray;
+	ensure(false);
+	return ConstEmptyArray;
+}
+
+////////////////////////////////////////////////////////////////////////////////////////////////////////////////////////////////////////////////////////////
+// 
+// Object encapsulating a change to the vertex attribute node's values. Used for Undo/Redo.
+class FDataflowVertexAttributeEditableNode::FEditNodeToolChange final : public FToolCommandChange
+{
+public:
+	FEditNodeToolChange(const FDataflowVertexAttributeEditableNode& Node) :
+		NodeGuid(Node.GetGuid()),
+		SavedAttributeValues(Node.GetStoredAttributeValues())
+	{}
+
+private:
+	virtual FString ToString() const final
+	{
+		return TEXT("FDataflowVertexAttributeEditableNode::FEditNodeToolChange");
+	}
+
+	virtual void Apply(UObject* Object) final
+	{
+		SwapApplyRevert(Object);
+	}
+
+	virtual void Revert(UObject* Object) final
+	{
+		SwapApplyRevert(Object);
+	}
+
+	void SwapApplyRevert(UObject* Object)
+	{
+		if (UDataflow* const Dataflow = Cast<UDataflow>(Object))
+		{
+			if (const TSharedPtr<FDataflowNode> BaseNode = Dataflow->GetDataflow()->FindBaseNode(NodeGuid))
+			{
+				if (FDataflowVertexAttributeEditableNode* const Node = BaseNode->AsType<FDataflowVertexAttributeEditableNode>())
+				{
+					Node->SwapStoredAttributeValuesWith(SavedAttributeValues);
+					Node->Invalidate();
+				}
+			}
+		}
+	}
+
+	FGuid NodeGuid;
+	TArray<float> SavedAttributeValues;
+};
+
+TUniquePtr<FToolCommandChange> FDataflowVertexAttributeEditableNode::MakeEditNodeToolChange()
+{
+	return MakeUnique<FDataflowVertexAttributeEditableNode::FEditNodeToolChange>(*this);
+}
+
+////////////////////////////////////////////////////////////////////////////////////////////////////////////////////////////////////////////////////////////
 //
 // FDataflowCollectionAddScalarVertexPropertyNode
 //
@@ -245,19 +305,23 @@ void FDataflowCollectionAddScalarVertexPropertyNode::GetVertexAttributeValues(UE
 	const FName InAttribName(InKey.Attribute);
 	const FName InGroupName(InKey.Group);
 	const FManagedArrayCollection& InCollection = GetValue(Context, &Collection, Collection);
-	
+
+	const int32 NumVertices = InCollection.NumElements(InGroupName);
+	OutValues.SetNumUninitialized(NumVertices);
+
 	if (const TManagedArray<float>* AttributeArray = InCollection.FindAttributeTyped<float>(InAttribName, InGroupName))
 	{
-		OutValues.SetNumUninitialized(AttributeArray->Num());
 		ExtractVertexWeights(AttributeArray->GetConstArray(), OutValues);
 	}
 	else
 	{
-		OutValues.Init(0.0f, InCollection.NumElements(InGroupName));
+		TArray<float> ZeroSetupWeights;
+		ZeroSetupWeights.Init(0.0f, NumVertices);
+		ExtractVertexWeights(ZeroSetupWeights, OutValues);
 	}
 }
 
-void FDataflowCollectionAddScalarVertexPropertyNode::SetVertexAttributeValues(UE::Dataflow::FContext& Context, const TArray<float>& InValues)
+void FDataflowCollectionAddScalarVertexPropertyNode::SetVertexAttributeValues(UE::Dataflow::FContext& Context, const TArray<float>& InValues, const TArray<int32>& InWeightIndices)
 {
 	const FCollectionAttributeKey InKey = GetWeightAttributeKey(Context);
 	const FName InAttribName(InKey.Attribute);
@@ -266,12 +330,24 @@ void FDataflowCollectionAddScalarVertexPropertyNode::SetVertexAttributeValues(UE
 
 	if (const TManagedArray<float>* AttributeArray = InCollection.FindAttributeTyped<float>(InAttribName, InGroupName))
 	{
-		ReportVertexWeights(AttributeArray->GetConstArray(), InValues, {});
+		ReportVertexWeights(AttributeArray->GetConstArray(), InValues, InWeightIndices);
+	}
+	else
+	{
+		TArray<float> ZeroSetupWeights;
+		const int32 NumVertices = InCollection.NumElements(InGroupName);
+		ZeroSetupWeights.Init(0.0f, NumVertices);
+		ReportVertexWeights(ZeroSetupWeights, InValues, InWeightIndices);
 	}
 
 	Invalidate();
 }
 
+void FDataflowCollectionAddScalarVertexPropertyNode::SwapStoredAttributeValuesWith(TArray<float>& OtherValues)
+{
+	Swap(VertexWeights, OtherValues);
+}
+
 void FDataflowCollectionAddScalarVertexPropertyNode::ReportVertexWeights(const TArray<float>& SetupWeights, const TArray<float>& FinalWeights,
 		const TArray<int32>& WeightIndices)
 {
diff --git a/Engine/Plugins/Experimental/Dataflow/Source/DataflowNodes/Public/Dataflow/DataflowCollectionAddScalarVertexPropertyNode.h b/Engine/Plugins/Experimental/Dataflow/Source/DataflowNodes/Public/Dataflow/DataflowCollectionAddScalarVertexPropertyNode.h
--- a/Engine/Plugins/Experimental/Dataflow/Source/DataflowNodes/Public/Dataflow/DataflowCollectionAddScalarVertexPropertyNode.h
+++ b/Engine/Plugins/Experimental/Dataflow/Source/DataflowNodes/Public/Dataflow/DataflowCollectionAddScalarVertexPropertyNode.h
@@ -90,14 +90,24 @@ struct FDataflowVertexAttributeEditableNode : public FDataflowNode
 	UE_API FDataflowVertexAttributeEditableNode(const UE::Dataflow::FNodeParameters& InParam, FGuid InGuid = FGuid::NewGuid());
 
 public:
+	class FEditNodeToolChange;
+	UE_API TUniquePtr<class FToolCommandChange> MakeEditNodeToolChange();
+
 	/** Get the list of view mode supported by this node when being edited */
 	virtual void GetSupportedViewModes(UE::Dataflow::FContext& Context, TArray<FName>& OutViewModeNames) const { ensure(false); }
 
+	/** Provide a extra indirection vertex mapping */
+	//virtual void GetVertexMapping(UE::Dataflow::FContext& Context, TArray<float>& OutValues) const { ensure(false); }
+
 	/** Get the absolute vertex attribute values */
 	virtual void GetVertexAttributeValues(UE::Dataflow::FContext& Context, TArray<float>& OutValues) const { ensure(false); }
 
-	/** Set the absolute vertex attribute values */
-	virtual void SetVertexAttributeValues(UE::Dataflow::FContext& Context, const TArray<float>& InValues) { ensure(false); }
+	/** Set the absolute vertex attribute values - InWeightIndices can be used to sepcify which indices are to be set */
+	virtual void SetVertexAttributeValues(UE::Dataflow::FContext& Context, const TArray<float>& InValues, const TArray<int32>& InWeightIndices) { ensure(false); }
+
+private:
+	UE_API virtual const TArray<float>& GetStoredAttributeValues() const;
+	virtual void SwapStoredAttributeValuesWith(TArray<float>& OtherValues) { ensure(false); }
 };
 
 /** Scalar vertex properties. */
@@ -156,7 +166,9 @@ private:
 	//~ Begin - FDataflowVertexAttributeEditableNode API
 	virtual void GetSupportedViewModes(UE::Dataflow::FContext& Context, TArray<FName>& OutViewModeNames) const override;
 	virtual void GetVertexAttributeValues(UE::Dataflow::FContext& Context, TArray<float>& OutValues) const override;
-	virtual void SetVertexAttributeValues(UE::Dataflow::FContext& Context, const TArray<float>& InValues) override;
+	virtual void SetVertexAttributeValues(UE::Dataflow::FContext& Context, const TArray<float>& InValues, const TArray<int32>& InWeightIndices) override;
+	virtual const TArray<float>& GetStoredAttributeValues() const override { return VertexWeights; }
+	virtual void SwapStoredAttributeValuesWith(TArray<float>& OtherValues) override;
 	//~ End - FDataflowVertexAttributeEditableNode API
 };
 
