Dataflow : fix code analysis issue related to previous changelist

diff --git a/Engine/Plugins/Experimental/Dataflow/Source/DataflowEditor/Private/DataflowEditorTools/DataflowEditorVertexAttributePaintTool.cpp b/Engine/Plugins/Experimental/Dataflow/Source/DataflowEditor/Private/DataflowEditorTools/DataflowEditorVertexAttributePaintTool.cpp
--- a/Engine/Plugins/Experimental/Dataflow/Source/DataflowEditor/Private/DataflowEditorTools/DataflowEditorVertexAttributePaintTool.cpp
+++ b/Engine/Plugins/Experimental/Dataflow/Source/DataflowEditor/Private/DataflowEditorTools/DataflowEditorVertexAttributePaintTool.cpp
@@ -522,7 +522,7 @@ void FDataflowEditorVertexAttributePaintToolData::UpdateNode(TObjectPtr<UDataflo
 			if (TSharedPtr<UE::Dataflow::FEngineContext> DataflowContext = InDataflowEditorContextObject->GetDataflowContext())
 			{
 				// Save previous state for undo
-				if (UDataflow* const Dataflow = InDataflowEditorContextObject ? InDataflowEditorContextObject->GetDataflowAsset(): nullptr)
+				if (TObjectPtr<UDataflow> Dataflow = InDataflowEditorContextObject ? InDataflowEditorContextObject->GetDataflowAsset(): nullptr)
 				{
 					if (TransactionAPI)
 					{
