Dataflow : Imporve the evaluate button behvior to be less confusing when opening a dataflow graph in automatic mode
- if the graph is dirty, show a question mark icon and make the button enabled
- once the graph ahs evaluated keep the green icon and disable the button

diff --git a/Engine/Plugins/Experimental/Dataflow/Source/DataflowEditor/Private/Dataflow/DataflowEditorToolkit.cpp b/Engine/Plugins/Experimental/Dataflow/Source/DataflowEditor/Private/Dataflow/DataflowEditorToolkit.cpp
--- a/Engine/Plugins/Experimental/Dataflow/Source/DataflowEditor/Private/Dataflow/DataflowEditorToolkit.cpp
+++ b/Engine/Plugins/Experimental/Dataflow/Source/DataflowEditor/Private/Dataflow/DataflowEditorToolkit.cpp
@@ -876,7 +876,7 @@ FSlateIcon FDataflowEditorToolkit::GetEvaluationStatusImage() const
 	}
 
 	FName OverlayIcon;
-	if (EvaluationMode == EDataflowEditorEvaluationMode::Manual && IsGraphDirty())
+	if (IsGraphDirty())
 	{
 		OverlayIcon = CompileStatusUnknown;
 	}
@@ -1009,8 +1009,8 @@ TSharedRef<SWidget> FDataflowEditorToolkit::GenerateEvaluationOptionsMenu()
 
 bool FDataflowEditorToolkit::IsEvaluateButtonEnabled() const
 {
-	// we disable the button in automatic mode as this should not be needed
-	return (EvaluationMode == EDataflowEditorEvaluationMode::Manual);
+	// Button is always enabled in manual mode or in automatic mode if the graph is dirty
+	return (EvaluationMode == EDataflowEditorEvaluationMode::Manual || IsGraphDirty());
 }
 
 void FDataflowEditorToolkit::InitializeEdMode(UBaseCharacterFXEditorMode* EdMode)
