Dataflow : add support for finding variable nodes in the graph when double clicking on the corresponding variable

diff --git a/Engine/Plugins/Experimental/Dataflow/Source/DataflowEditor/Private/Dataflow/DataflowEditorToolkit.cpp b/Engine/Plugins/Experimental/Dataflow/Source/DataflowEditor/Private/Dataflow/DataflowEditorToolkit.cpp
--- a/Engine/Plugins/Experimental/Dataflow/Source/DataflowEditor/Private/Dataflow/DataflowEditorToolkit.cpp
+++ b/Engine/Plugins/Experimental/Dataflow/Source/DataflowEditor/Private/Dataflow/DataflowEditorToolkit.cpp
@@ -40,6 +40,7 @@
 #include "Dataflow/DataflowSimulationViewport.h"
 #include "Dataflow/DataflowSimulationViewportClient.h"
 #include "Dataflow/DataflowSubGraphNodes.h"
+#include "Dataflow/DataflowVariableNodes.h"
 #include "Dataflow/DataflowMembersWidget.h"
 #include "Dataflow/DataflowSelectionView.h"
 #include "Dataflow/DataflowOutputLog.h"
@@ -2735,6 +2736,89 @@ TSharedRef<SGraphEditor> FDataflowEditorToolkit::CreateSubGraphEditorWidget(TSha
 	return SNew(SGraphEditor);
 }
 
+void FDataflowEditorToolkit::FindAllVariableNodeInGraph(UEdGraph* EdGraph, FName VariableName, TArray<UDataflowEdNode*>& OutEdNodes)
+{
+	for (UEdGraphNode* EdNode : EdGraph->Nodes)
+	{
+		if (UDataflowEdNode* DataflowEdNode = Cast<UDataflowEdNode>(EdNode))
+		{
+			if (TSharedPtr<FDataflowNode> DataflowNode = DataflowEdNode->GetDataflowNode())
+			{
+				if (FGetDataflowVariableNode* VariableNode = DataflowNode->AsType<FGetDataflowVariableNode>())
+				{
+					if (VariableNode->GetVariableName() == VariableName)
+					{
+						OutEdNodes.Add(DataflowEdNode);
+					}
+				}
+			}
+		}
+	}
+}
+
+FGuid FDataflowEditorToolkit::FocusOnNextVariableNode(FName VariableName, const FGuid& LastTimeNodeGuid)
+{
+	if (const TObjectPtr<UDataflowBaseContent>& EditorContent = GetEditorContent())
+	{
+		if (UDataflow* DataflowAsset = EditorContent->GetDataflowAsset())
+		{
+			// first collect all the nodes that match the specific variable 
+			TArray<UDataflowEdNode*> EdNodes;
+			FindAllVariableNodeInGraph(DataflowAsset, VariableName, EdNodes);
+			for (TObjectPtr<UDataflowSubGraph> SubGraphs : DataflowAsset->GetSubGraphs())
+			{
+				FindAllVariableNodeInGraph(SubGraphs, VariableName, EdNodes);
+			}
+
+			const int32 LastTimeIndex = EdNodes.IndexOfByPredicate(
+				[&LastTimeNodeGuid](const UDataflowEdNode* EdNode)
+				{
+					return (EdNode && EdNode->GetDataflowNodeGuid() == LastTimeNodeGuid);
+				});
+
+			int32 IndexToUse = 0;
+			if (LastTimeIndex != INDEX_NONE && EdNodes.Num() > 0)
+			{
+				IndexToUse = (LastTimeIndex + 1) % EdNodes.Num();
+			}
+
+			if (EdNodes.IsValidIndex(IndexToUse))
+			{
+				if (UDataflowEdNode* NextVariableNode = EdNodes[IndexToUse])
+				{
+					TSharedPtr<SDataflowGraphEditor> ActiveEditor = GraphEditor;
+					if (UDataflowSubGraph* SubGraph = Cast<UDataflowSubGraph>(NextVariableNode->GetGraph()))
+					{
+						OpenSubGraphTab(SubGraph);
+						ActiveEditor = ActiveSubGraphEditorWeakPtr.Pin();
+					}
+					else
+					{
+						// TODO : try to find a way to focus back on the main graph
+						//			current issue seem to be related to the tabspawner beeing unknown when calling thos method
+						//			maybe a better way would be to leverage the subgraph spawner and show the main graph in the subgraph list 
+						//			so that all graphs are spawned and treated the same way ( but this may be an issue to keep the pointer of the main graph )
+						//if (TSharedPtr<SDockTab> MainGraphTab = FGlobalTabmanager::Get()->TryInvokeTab(GraphCanvasTabId))
+						//{
+						//	MainGraphTab->ActivateInParent(SetDirectly);
+						//}
+					}
+					if (ActiveEditor)
+					{
+						// Clear node selection in the newly active graph
+						ActiveEditor->ClearSelectionSet();
+
+						// now jump to the relevant node 
+						ActiveEditor->JumpToNode(NextVariableNode);
+					}
+					return NextVariableNode->GetDataflowNodeGuid();
+				}
+			}
+		}
+	}
+	return FGuid();
+}
+
 FName FDataflowEditorToolkit::GetToolkitFName() const
 {
 	return FName("DataflowEditor");
diff --git a/Engine/Plugins/Experimental/Dataflow/Source/DataflowEditor/Private/Dataflow/DataflowMembersWidget.cpp b/Engine/Plugins/Experimental/Dataflow/Source/DataflowEditor/Private/Dataflow/DataflowMembersWidget.cpp
--- a/Engine/Plugins/Experimental/Dataflow/Source/DataflowEditor/Private/Dataflow/DataflowMembersWidget.cpp
+++ b/Engine/Plugins/Experimental/Dataflow/Source/DataflowEditor/Private/Dataflow/DataflowMembersWidget.cpp
@@ -200,10 +200,15 @@ namespace UE::Dataflow::Private
 			}
 		}
 
-		virtual void OnDoubleClicked(FEdGraphSchemaAction&, FDataflowEditorToolkit&) const override
+		virtual void OnDoubleClicked(FEdGraphSchemaAction& InAction, FDataflowEditorToolkit& InToolkit) const override
 		{
-			// todo : Jump to the node 
+			if (FEdGraphSchemaAction_DataflowVariable* TypedAction = AsTypedAction(InAction))
+			{
+				LastFoundVariableNodeGuid = InToolkit.FocusOnNextVariableNode(TypedAction->GetFullVariableName(), LastFoundVariableNodeGuid);
+			}
 		}
+
+		mutable FGuid LastFoundVariableNodeGuid;
 	};
 }
 
diff --git a/Engine/Plugins/Experimental/Dataflow/Source/DataflowEditor/Public/Dataflow/DataflowEditorToolkit.h b/Engine/Plugins/Experimental/Dataflow/Source/DataflowEditor/Public/Dataflow/DataflowEditorToolkit.h
--- a/Engine/Plugins/Experimental/Dataflow/Source/DataflowEditor/Public/Dataflow/DataflowEditorToolkit.h
+++ b/Engine/Plugins/Experimental/Dataflow/Source/DataflowEditor/Public/Dataflow/DataflowEditorToolkit.h
@@ -104,6 +104,9 @@ public:
 	UE_API UDataflowSubGraph* GetSubGraph(const FGuid& SubGraphGuid) const;
 	UE_API UDataflowSubGraph* GetSubGraph(FName SubGraphName) const;
 
+	UE_API FGuid FocusOnNextVariableNode(FName VariableName, const FGuid& LastTimeNodeGuid);
+	void FindAllVariableNodeInGraph(UEdGraph* EdGraph, FName VariableName, TArray<UDataflowEdNode*>& OutEdNodes);
+
 	UE_API const FString& GetDebugDrawOverlayString() const;
 	
 	/** Get the toolkit evaluation mode */
diff --git a/Engine/Source/Runtime/Experimental/Dataflow/Engine/Private/Dataflow/DataflowVariableNodes.cpp b/Engine/Source/Runtime/Experimental/Dataflow/Engine/Private/Dataflow/DataflowVariableNodes.cpp
--- a/Engine/Source/Runtime/Experimental/Dataflow/Engine/Private/Dataflow/DataflowVariableNodes.cpp
+++ b/Engine/Source/Runtime/Experimental/Dataflow/Engine/Private/Dataflow/DataflowVariableNodes.cpp
@@ -268,14 +268,14 @@ void FGetDataflowVariableNode::Evaluate(UE::Dataflow::FContext& Context, const F
 						return;
 					}
 				}
-				// no override or engine context , use the default defined in the dataflow asset
-				if (TStrongObjectPtr<UDataflow> DataflowAsset = WeakDataflowPtr.Pin())
+			}
+			// no override or engine context , use the default defined in the dataflow asset
+			if (TStrongObjectPtr<UDataflow> DataflowAsset = WeakDataflowPtr.Pin())
+			{
+				if (DataflowAsset->Variables.FindPropertyDescByName(VariableName) != nullptr)
 				{
-					if (DataflowAsset->Variables.FindPropertyDescByName(VariableName) != nullptr)
-					{
-						EvaluateFunction(*this, DataflowAsset->Variables, Context, *Out);
-						return;
-					}
+					EvaluateFunction(*this, DataflowAsset->Variables, Context, *Out);
+					return;
 				}
 			}
 		}
