Fix CIS for Renderable components

diff --git a/Engine/Plugins/Experimental/Dataflow/Source/DataflowEditor/Public/DataflowRendering/DataflowRenderableComponents.h b/Engine/Plugins/Experimental/Dataflow/Source/DataflowEditor/Public/DataflowRendering/DataflowRenderableComponents.h
--- a/Engine/Plugins/Experimental/Dataflow/Source/DataflowEditor/Public/DataflowRendering/DataflowRenderableComponents.h
+++ b/Engine/Plugins/Experimental/Dataflow/Source/DataflowEditor/Public/DataflowRendering/DataflowRenderableComponents.h
@@ -30,11 +30,11 @@ namespace UE::Dataflow
 		template<typename T UE_REQUIRES(TIsDerivedFrom<T, UPrimitiveComponent>::Value)>
 		T* AddNewComponent(FName ComponentName = NAME_None, UObject* ComponentParent = nullptr)
 		{
-			UObject* ObjectOwner = ComponentParent ? ComponentParent : RootActor;
-			const FName ObjectName = MakeUniqueObjectName(ObjectOwner, T::StaticClass(), FName(ComponentName));
+			// Todo: Will need to fix the name clashing
+			//const FName ObjectName = MakeUniqueObjectName(ObjectOwner, T::StaticClass(), FName(ComponentName));
 			
 			// for now always allocate a new one - in the future we could use a pool 
-			T* NewComponent = NewObject<T>(ObjectOwner/*, ObjectName*/);
+			T* NewComponent = ComponentParent ? NewObject<T>(ComponentParent) : NewObject<T>(RootActor);
 			if (NewComponent)
 			{
 				Components.Add(NewComponent);
