Dataflow: Make skin weight node changes undoable after accepting changes in the tool

diff --git a/Engine/Plugins/Experimental/Dataflow/Source/DataflowEditor/Private/DataflowEditorTools/DataflowEditorSkinWeightsPaintTool.cpp b/Engine/Plugins/Experimental/Dataflow/Source/DataflowEditor/Private/DataflowEditorTools/DataflowEditorSkinWeightsPaintTool.cpp
--- a/Engine/Plugins/Experimental/Dataflow/Source/DataflowEditor/Private/DataflowEditorTools/DataflowEditorSkinWeightsPaintTool.cpp
+++ b/Engine/Plugins/Experimental/Dataflow/Source/DataflowEditor/Private/DataflowEditorTools/DataflowEditorSkinWeightsPaintTool.cpp
@@ -233,6 +233,22 @@ void UDataflowEditorSkinWeightsPaintTool::OnShutdown(EToolShutdownType ShutdownT
 		const int32 VertexOffset = GetVertexOffset();
 		if(VertexOffset != INDEX_NONE)
 		{
+			// Save previous state for undo
+			if (TObjectPtr<UDataflow> Dataflow = DataflowEditorContextObject->GetDataflowAsset())
+			{
+				if (UInteractiveToolManager* ToolManager = GetToolManager())
+				{
+					if (IToolsContextTransactionsAPI* TransactionAPI = ToolManager->GetContextTransactionsAPI())
+					{
+						TransactionAPI->AppendChange(
+							Dataflow,
+							SkinWeightNode->MakeEditNodeToolChange(),
+							LOCTEXT("UDataflowEditorSkinWeightsPaintTool_ChangeDescription", "Update Skin Weight Node")
+						);
+					}
+				}
+			}
+
 			TArray<TArray<int32>> CurrentIndices;
 			TArray<TArray<float>> CurrentWeights;
 
diff --git a/Engine/Plugins/Experimental/Dataflow/Source/DataflowNodes/Private/Dataflow/DataflowCollectionEditSkinWeightsNode.cpp b/Engine/Plugins/Experimental/Dataflow/Source/DataflowNodes/Private/Dataflow/DataflowCollectionEditSkinWeightsNode.cpp
--- a/Engine/Plugins/Experimental/Dataflow/Source/DataflowNodes/Private/Dataflow/DataflowCollectionEditSkinWeightsNode.cpp
+++ b/Engine/Plugins/Experimental/Dataflow/Source/DataflowNodes/Private/Dataflow/DataflowCollectionEditSkinWeightsNode.cpp
@@ -15,6 +15,7 @@
 #include "Engine/SkeletalMesh.h"
 #include "Animation/Skeleton.h"
 #include "Components/SkeletalMeshComponent.h"
+#include "InteractiveToolChange.h"
 
 #if WITH_EDITOR
 #include "StaticToSkeletalMeshConverter.h"
@@ -458,4 +459,56 @@ int32 FDataflowCollectionEditSkinWeightsNode::GetSkeletalMeshOffset(const TObjec
 	return INDEX_NONE;
 }
 
+
+////////////////////////////////////////////////////////////////////////////////////////////////////////////////////////////////////////////////////////////
+// 
+// Object encapsulating a change to the vertex skin weights values. Used for Undo/Redo.
+class FDataflowCollectionEditSkinWeightsNode::FEditNodeToolChange final : public FToolCommandChange
+{
+public:
+	FEditNodeToolChange(const FDataflowCollectionEditSkinWeightsNode& Node) :
+		NodeGuid(Node.GetGuid()),
+		SavedSkinWeightData(Node.SkinWeights)
+	{}
+
+private:
+	virtual FString ToString() const final
+	{
+		return TEXT("FDataflowCollectionEditSkinWeightsNode::FEditNodeToolChange");
+	}
+
+	virtual void Apply(UObject* Object) final
+	{
+		SwapApplyRevert(Object);
+	}
+
+	virtual void Revert(UObject* Object) final
+	{
+		SwapApplyRevert(Object);
+	}
+
+	void SwapApplyRevert(UObject* Object)
+	{
+		if (UDataflow* const Dataflow = Cast<UDataflow>(Object))
+		{
+			if (const TSharedPtr<FDataflowNode> BaseNode = Dataflow->GetDataflow()->FindBaseNode(NodeGuid))
+			{
+				if (FDataflowCollectionEditSkinWeightsNode* const Node = BaseNode->AsType<FDataflowCollectionEditSkinWeightsNode>())
+				{
+					Swap(SavedSkinWeightData, Node->SkinWeights);
+					Node->Invalidate();
+				}
+			}
+		}
+	}
+
+	FGuid NodeGuid;
+	TArray<FDataflowSkinWeightData> SavedSkinWeightData;
+};
+
+TUniquePtr<FToolCommandChange> FDataflowCollectionEditSkinWeightsNode::MakeEditNodeToolChange()
+{
+	return MakeUnique<FDataflowCollectionEditSkinWeightsNode::FEditNodeToolChange>(*this);
+}
+
 #undef LOCTEXT_NAMESPACE
diff --git a/Engine/Plugins/Experimental/Dataflow/Source/DataflowNodes/Public/Dataflow/DataflowCollectionEditSkinWeightsNode.h b/Engine/Plugins/Experimental/Dataflow/Source/DataflowNodes/Public/Dataflow/DataflowCollectionEditSkinWeightsNode.h
--- a/Engine/Plugins/Experimental/Dataflow/Source/DataflowNodes/Public/Dataflow/DataflowCollectionEditSkinWeightsNode.h
+++ b/Engine/Plugins/Experimental/Dataflow/Source/DataflowNodes/Public/Dataflow/DataflowCollectionEditSkinWeightsNode.h
@@ -115,6 +115,9 @@ public:
 	/** Return the collection offset given a skeletal mesh*/
 	UE_API int32 GetSkeletalMeshOffset(const TObjectPtr<USkeletalMesh>& SkeletalMesh) const;
 	
+	class FEditNodeToolChange;
+	UE_API TUniquePtr<class FToolCommandChange> MakeEditNodeToolChange();
+
 private:
 
 	//~ Begin FDataflowPrimitiveNode interface
