Dataflow : new weight map paint tool : add support for custom vertex remapping per view in order to support 2D sim view for cloth

diff --git a/Engine/Plugins/ChaosClothAssetEditor/Source/ChaosClothAssetDataflowNodes/Private/ChaosClothAsset/WeightMapNode.cpp b/Engine/Plugins/ChaosClothAssetEditor/Source/ChaosClothAssetDataflowNodes/Private/ChaosClothAsset/WeightMapNode.cpp
--- a/Engine/Plugins/ChaosClothAssetEditor/Source/ChaosClothAssetDataflowNodes/Private/ChaosClothAsset/WeightMapNode.cpp
+++ b/Engine/Plugins/ChaosClothAssetEditor/Source/ChaosClothAssetDataflowNodes/Private/ChaosClothAsset/WeightMapNode.cpp
@@ -7,6 +7,8 @@
 #include "ChaosClothAsset/ClothDataflowViewModes.h"
 #include "ChaosClothAsset/ClothGeometryTools.h"
 #include "ChaosClothAsset/CollectionClothFacade.h"
+#include "ChaosClothAsset/ClothDataflowViewModes.h"
+#include "Dataflow/DataflowRenderingViewMode.h"
 #include "ChaosClothAsset/WeightedValue.h"
 #include "Dataflow/DataflowInputOutput.h"
 #include "Dataflow/DataflowObject.h"
@@ -315,7 +317,42 @@ void FChaosClothAssetWeightMapNode::SetVertexAttributeValues(UE::Dataflow::FCont
 			InputAttributeValues = ClothFacade.GetUserDefinedAttribute<float>(InInputName, ClothCollectionGroup::RenderVertices);
 		}
 
-		UE::Chaos::ClothAsset::Private::SetVertexWeights(InputAttributeValues, InValues, MapOverrideType, GetVertexWeights());
+		if (InWeightIndices.Num() == InValues.Num())
+		{
+			TArray<float> ValuesToApply;
+			ValuesToApply.Init(0.f, InWeightIndices.Num());
+			for (int32 Index = 0; Index < ValuesToApply.Num(); ++Index)
+			{
+				ValuesToApply[InWeightIndices[Index]] = InValues[Index];
+			}
+			UE::Chaos::ClothAsset::Private::SetVertexWeights(InputAttributeValues, ValuesToApply, MapOverrideType, GetVertexWeights());
+		}
+		else
+		{
+			UE::Chaos::ClothAsset::Private::SetVertexWeights(InputAttributeValues, InValues, MapOverrideType, GetVertexWeights());
+		}
+	}
+}
+
+void FChaosClothAssetWeightMapNode::GetExtraVertexMapping(UE::Dataflow::FContext& Context, FName SelectedViewMode, TArray<int32>& OutMappingToWeight, TArray<TArray<int32>>& OutMappingFromWeight) const
+{
+	OutMappingToWeight.Reset();
+	OutMappingFromWeight.Reset();
+
+	using namespace UE::Chaos::ClothAsset;
+	if (MeshTarget == EChaosClothAssetWeightMapMeshTarget::Simulation)
+	{
+		if (SelectedViewMode == FCloth2DSimViewMode::Name || SelectedViewMode == UE::Dataflow::FDataflowConstruction2DViewMode::Name)
+		{
+			FManagedArrayCollection InCollection = GetValue<FManagedArrayCollection>(Context, &Collection);
+			const TSharedRef<FManagedArrayCollection> ClothCollection = MakeShared<FManagedArrayCollection>(MoveTemp(InCollection));
+			FCollectionClothConstFacade ClothFacade(ClothCollection);
+			if (ClothFacade.IsValid())  // Can only act on the collection if it is a valid cloth collection
+			{
+				OutMappingToWeight = ClothFacade.GetSimVertex3DLookup();
+				OutMappingFromWeight = ClothFacade.GetSimVertex2DLookup();
+			}
+		}
 	}
 }
 
diff --git a/Engine/Plugins/ChaosClothAssetEditor/Source/ChaosClothAssetDataflowNodes/Public/ChaosClothAsset/WeightMapNode.h b/Engine/Plugins/ChaosClothAssetEditor/Source/ChaosClothAssetDataflowNodes/Public/ChaosClothAsset/WeightMapNode.h
--- a/Engine/Plugins/ChaosClothAssetEditor/Source/ChaosClothAssetDataflowNodes/Public/ChaosClothAsset/WeightMapNode.h
+++ b/Engine/Plugins/ChaosClothAssetEditor/Source/ChaosClothAssetDataflowNodes/Public/ChaosClothAsset/WeightMapNode.h
@@ -103,6 +103,7 @@ private:
 	virtual void GetSupportedViewModes(UE::Dataflow::FContext& Context, TArray<FName>& OutViewModeNames) const override;
 	virtual void GetVertexAttributeValues(UE::Dataflow::FContext& Context, TArray<float>& OutValues) const override;
 	virtual void SetVertexAttributeValues(UE::Dataflow::FContext& Context, const TArray<float>& InValues, const TArray<int32>& InWeightIndices) override;
+	virtual void GetExtraVertexMapping(UE::Dataflow::FContext& Context, FName SelectedViewMode, TArray<int32>& OutMappingToWeight, TArray<TArray<int32>>& OutMappingFromWeight) const override;
 	virtual const TArray<float>& GetStoredAttributeValues() const override { return VertexWeights; }
 	virtual void SwapStoredAttributeValuesWith(TArray<float>& OtherValues) override;
 	//~ End - FDataflowVertexAttributeEditableNode API
diff --git a/Engine/Plugins/Experimental/Dataflow/Source/DataflowEditor/Private/DataflowEditorTools/DataflowEditorVertexAttributePaintTool.cpp b/Engine/Plugins/Experimental/Dataflow/Source/DataflowEditor/Private/DataflowEditorTools/DataflowEditorVertexAttributePaintTool.cpp
--- a/Engine/Plugins/Experimental/Dataflow/Source/DataflowEditor/Private/DataflowEditorTools/DataflowEditorVertexAttributePaintTool.cpp
+++ b/Engine/Plugins/Experimental/Dataflow/Source/DataflowEditor/Private/DataflowEditorTools/DataflowEditorVertexAttributePaintTool.cpp
@@ -363,14 +363,31 @@ void FDataflowEditorVertexAttributePaintToolData::Setup(FDynamicMesh3& InMesh, F
 
 	if (InDataflowEditorContextObject)
 	{
+		TArray<int32> ExtraMappingToWeight;
+		TArray<TArray<int32>> ExtraMappingFromWeight;
+		if (NodeToUpdate && InDataflowEditorContextObject)
+		{
+			if (const UE::Dataflow::IDataflowConstructionViewMode* ViewMode = InDataflowEditorContextObject->GetConstructionViewMode())
+			{
+				if (TSharedPtr<UE::Dataflow::FEngineContext> DataflowContext = InDataflowEditorContextObject->GetDataflowContext())
+				{
+					NodeToUpdate->GetExtraVertexMapping(*DataflowContext, ViewMode->GetName(), ExtraMappingToWeight, ExtraMappingFromWeight);
+				}
+			}
+		}
+		
+
 		// Generate mapping array from and to non-manifiold and manifold mesh 
 		if (const TSharedPtr<const FManagedArrayCollection> Collection = InDataflowEditorContextObject->GetRenderCollection())
 		{
 			using namespace UE::Dataflow;
 			const UE::Geometry::FNonManifoldMappingSupport NonManifoldMapping(InMesh);
+			const bool bHasNonManifoldMapping = NonManifoldMapping.IsNonManifoldVertexInSource();
+			const bool bHasExtraMapping = (ExtraMappingToWeight.Num() != 0 && ExtraMappingFromWeight.Num() != 0);
 
-			bHaveDynamicMeshToWeightConversion = NonManifoldMapping.IsNonManifoldVertexInSource();
-			if (bHaveDynamicMeshToWeightConversion)
+			bHaveDynamicMeshToWeightConversion = bHasNonManifoldMapping || bHasExtraMapping;
+
+			if (bHasNonManifoldMapping)
 			{
 				DynamicMeshToWeight.SetNumUninitialized(InMesh.VertexCount());
 				WeightToDynamicMesh.Reset();
@@ -378,18 +395,27 @@ void FDataflowEditorVertexAttributePaintToolData::Setup(FDynamicMesh3& InMesh, F
 				for (int32 DynamicMeshVert = 0; DynamicMeshVert < InMesh.VertexCount(); ++DynamicMeshVert)
 				{
 					DynamicMeshToWeight[DynamicMeshVert] = NonManifoldMapping.GetOriginalNonManifoldVertexID(DynamicMeshVert);
-					if (0 <= DynamicMeshToWeight[DynamicMeshVert] && DynamicMeshToWeight[DynamicMeshVert] < WeightToDynamicMesh.Num())
+					if (bHasExtraMapping)
+					{
+						DynamicMeshToWeight[DynamicMeshVert] = ExtraMappingToWeight[DynamicMeshToWeight[DynamicMeshVert]];
+					}
+					if (WeightToDynamicMesh.IsValidIndex(DynamicMeshToWeight[DynamicMeshVert]))
 					{
 						WeightToDynamicMesh[DynamicMeshToWeight[DynamicMeshVert]].Add(DynamicMeshVert);
 					}
 					else
 					{
 						bHaveDynamicMeshToWeightConversion = false;
-						UE_LOG(LogTemp, Warning, TEXT("Weight map misalignment."));
+						UE_LOG(LogTemp, Warning, TEXT("Weight map misalignment, disabling remapping"));
 						break;
 					}
 				}
 			}
+			else if (bHasExtraMapping)
+			{
+				DynamicMeshToWeight = ExtraMappingToWeight;
+				WeightToDynamicMesh = ExtraMappingFromWeight;
+			}
 		}
 
 		// Create an attribute layer to temporarily paint into
diff --git a/Engine/Plugins/Experimental/Dataflow/Source/DataflowNodes/Private/Dataflow/DataflowCollectionAddScalarVertexPropertyNode.cpp b/Engine/Plugins/Experimental/Dataflow/Source/DataflowNodes/Private/Dataflow/DataflowCollectionAddScalarVertexPropertyNode.cpp
--- a/Engine/Plugins/Experimental/Dataflow/Source/DataflowNodes/Private/Dataflow/DataflowCollectionAddScalarVertexPropertyNode.cpp
+++ b/Engine/Plugins/Experimental/Dataflow/Source/DataflowNodes/Private/Dataflow/DataflowCollectionAddScalarVertexPropertyNode.cpp
@@ -137,6 +137,7 @@ TArray<UE::Dataflow::FRenderingParameter> DataflowAddScalarVertexPropertyCallbac
 }
 PRAGMA_ENABLE_DEPRECATION_WARNINGS
 
+////////////////////////////////////////////////////////////////////////////////////////////////////////////////////////////////////////////////////////////
 
 FDataflowVertexAttributeEditableNode::FDataflowVertexAttributeEditableNode(const UE::Dataflow::FNodeParameters& InParam, FGuid InGuid)
 	: FDataflowNode(InParam, InGuid)
@@ -149,6 +150,12 @@ const TArray<float>& FDataflowVertexAttributeEditableNode::GetStoredAttributeVal
 	return ConstEmptyArray;
 }
 
+void FDataflowVertexAttributeEditableNode::GetExtraVertexMapping(UE::Dataflow::FContext& Context, FName SelectedViewMode, TArray<int32>& OutMappingToWeight, TArray<TArray<int32>>& OutMappingFromWeight) const
+{
+	OutMappingToWeight.Reset();
+	OutMappingFromWeight.Reset();
+}
+
 ////////////////////////////////////////////////////////////////////////////////////////////////////////////////////////////////////////////////////////////
 // 
 // Object encapsulating a change to the vertex attribute node's values. Used for Undo/Redo.
diff --git a/Engine/Plugins/Experimental/Dataflow/Source/DataflowNodes/Public/Dataflow/DataflowCollectionAddScalarVertexPropertyNode.h b/Engine/Plugins/Experimental/Dataflow/Source/DataflowNodes/Public/Dataflow/DataflowCollectionAddScalarVertexPropertyNode.h
--- a/Engine/Plugins/Experimental/Dataflow/Source/DataflowNodes/Public/Dataflow/DataflowCollectionAddScalarVertexPropertyNode.h
+++ b/Engine/Plugins/Experimental/Dataflow/Source/DataflowNodes/Public/Dataflow/DataflowCollectionAddScalarVertexPropertyNode.h
@@ -96,15 +96,15 @@ public:
 	/** Get the list of view mode supported by this node when being edited */
 	virtual void GetSupportedViewModes(UE::Dataflow::FContext& Context, TArray<FName>& OutViewModeNames) const { ensure(false); }
 
-	/** Provide a extra indirection vertex mapping */
-	//virtual void GetVertexMapping(UE::Dataflow::FContext& Context, TArray<float>& OutValues) const { ensure(false); }
-
 	/** Get the absolute vertex attribute values */
 	virtual void GetVertexAttributeValues(UE::Dataflow::FContext& Context, TArray<float>& OutValues) const { ensure(false); }
 
 	/** Set the absolute vertex attribute values - InWeightIndices can be used to sepcify which indices are to be set */
 	virtual void SetVertexAttributeValues(UE::Dataflow::FContext& Context, const TArray<float>& InValues, const TArray<int32>& InWeightIndices) { ensure(false); }
 
+	/** Provide a extra indirection vertex mapping - optional will return empty by default */
+	UE_API virtual void GetExtraVertexMapping(UE::Dataflow::FContext& Context, FName SelectedViewMode, TArray<int32>& OutMappingToWeight, TArray<TArray<int32>>& OutMappingFromWeight) const;
+
 private:
 	UE_API virtual const TArray<float>& GetStoredAttributeValues() const;
 	virtual void SwapStoredAttributeValuesWith(TArray<float>& OtherValues) { ensure(false); }
