Dataflow : Improved copy of cached dynamic in FGeometryCollectionSurfaceRenderableType by parallelizing the copy which can be expensive because of numerous memcopy and allocation

diff --git a/Engine/Plugins/Experimental/Dataflow/Source/DataflowEditor/Private/DataflowRendering/DataflowGeometryCollectionRenderableType.cpp b/Engine/Plugins/Experimental/Dataflow/Source/DataflowEditor/Private/DataflowRendering/DataflowGeometryCollectionRenderableType.cpp
--- a/Engine/Plugins/Experimental/Dataflow/Source/DataflowEditor/Private/DataflowRendering/DataflowGeometryCollectionRenderableType.cpp
+++ b/Engine/Plugins/Experimental/Dataflow/Source/DataflowEditor/Private/DataflowRendering/DataflowGeometryCollectionRenderableType.cpp
@@ -52,7 +52,15 @@ namespace UE::Dataflow::Private
 				TArray<FDynamicMesh3> DynamicMeshes;
 				if (Instance.HasUptoDateCachedValue())
 				{
-					DynamicMeshes = Instance.GetCachedValue<TArray<FDynamicMesh3>>();
+					const TArray<FDynamicMesh3>& CachedDynamicMeshes = Instance.GetCachedValue<TArray<FDynamicMesh3>>();
+					DynamicMeshes.SetNum(CachedDynamicMeshes.Num());
+
+					// copying dynamic meshes can be expensive let's parallelize it 
+					ParallelFor(DynamicMeshes.Num(),
+						[&CachedDynamicMeshes, &DynamicMeshes](int32 Index)
+						{
+							DynamicMeshes[Index] = CachedDynamicMeshes[Index];
+						});
 				}
 				else
 				{
