Dataflow : expose parameters as inputs for correct skin weight node

diff --git a/Engine/Plugins/Experimental/Dataflow/Source/DataflowEditor/Private/DataflowEditorTools/DataflowEditorCorrectSkinWeightsNode.cpp b/Engine/Plugins/Experimental/Dataflow/Source/DataflowEditor/Private/DataflowEditorTools/DataflowEditorCorrectSkinWeightsNode.cpp
--- a/Engine/Plugins/Experimental/Dataflow/Source/DataflowEditor/Private/DataflowEditorTools/DataflowEditorCorrectSkinWeightsNode.cpp
+++ b/Engine/Plugins/Experimental/Dataflow/Source/DataflowEditor/Private/DataflowEditorTools/DataflowEditorCorrectSkinWeightsNode.cpp
@@ -667,6 +667,20 @@ FDataflowCorrectSkinWeightsNode::FDataflowCorrectSkinWeightsNode(const UE::Dataf
 	RegisterInputConnection(&BoneIndicesKey);
 	RegisterInputConnection(&BoneWeightsKey);
 	RegisterInputConnection(&SelectionMapKey);
+
+	RegisterInputConnection(&SmoothingIterations)
+		.SetCanHidePin(true).SetPinIsHidden(true);
+	RegisterInputConnection(&SmoothingFactor)
+		.SetCanHidePin(true).SetPinIsHidden(true);
+	RegisterInputConnection(&bUseSelectionAsPerVertexFactor)
+		.SetCanHidePin(true).SetPinIsHidden(true);
+	RegisterInputConnection(&PruningThreshold)
+		.SetCanHidePin(true).SetPinIsHidden(true);
+	RegisterInputConnection(&ClampingNumber)
+		.SetCanHidePin(true).SetPinIsHidden(true);
+	RegisterInputConnection(&SelectionThreshold)
+		.SetCanHidePin(true).SetPinIsHidden(true);
+
 	RegisterOutputConnection(&Collection, &Collection);
 	RegisterOutputConnection(&BoneIndicesKey, &BoneIndicesKey);
 	RegisterOutputConnection(&BoneWeightsKey, &BoneWeightsKey);
diff --git a/Engine/Plugins/Experimental/Dataflow/Source/DataflowEditor/Public/DataflowEditorTools/DataflowEditorCorrectSkinWeightsNode.h b/Engine/Plugins/Experimental/Dataflow/Source/DataflowEditor/Public/DataflowEditorTools/DataflowEditorCorrectSkinWeightsNode.h
--- a/Engine/Plugins/Experimental/Dataflow/Source/DataflowEditor/Public/DataflowEditorTools/DataflowEditorCorrectSkinWeightsNode.h
+++ b/Engine/Plugins/Experimental/Dataflow/Source/DataflowEditor/Public/DataflowEditorTools/DataflowEditorCorrectSkinWeightsNode.h
@@ -76,27 +76,27 @@ public:
 	ESkinWeightsCorrectionType CorrectionType = ESkinWeightsCorrectionType::Relax;
 
 	/** Number of iteration required for the smoothing */
-	UPROPERTY(EditAnywhere, Category = "Skin Weights", meta = (DisplayName = "Smoothing Iterations", EditCondition = "CorrectionType == ESkinWeightsCorrectionType::Relax", EditConditionHides))
+	UPROPERTY(EditAnywhere, Category = "Skin Weights", meta = (DisplayName = "Smoothing Iterations", DataflowInput, EditCondition = "CorrectionType == ESkinWeightsCorrectionType::Relax", EditConditionHides))
 	int32 SmoothingIterations = 5;
 
 	/** Lerp value in between the current and the average weight values */
-	UPROPERTY(EditAnywhere, Category = "Skin Weights", meta = (DisplayName = "Smooting Factor", EditCondition = "CorrectionType == ESkinWeightsCorrectionType::Relax", EditConditionHides))
+	UPROPERTY(EditAnywhere, Category = "Skin Weights", meta = (DisplayName = "Smooting Factor", DataflowInput, EditCondition = "CorrectionType == ESkinWeightsCorrectionType::Relax", EditConditionHides))
 	float SmoothingFactor = 0.5;
 
 	/** When true, use the decimal values of the selection as a per vertex factor for the selected operation - Currently available for Relax operation only - false by default */
-	UPROPERTY(EditAnywhere, Category = "Skin Weights", meta = (EditCondition = "CorrectionType == ESkinWeightsCorrectionType::Relax", EditConditionHides))
+	UPROPERTY(EditAnywhere, Category = "Skin Weights", meta = (DataflowInput, EditCondition = "CorrectionType == ESkinWeightsCorrectionType::Relax", EditConditionHides))
 	bool bUseSelectionAsPerVertexFactor = false;
 
 	/** All weights below this threshold will be pruned */
-	UPROPERTY(EditAnywhere, Category = "Skin Weights", meta = (DisplayName = "Pruning Threshold", EditCondition = "CorrectionType == ESkinWeightsCorrectionType::Prune", EditConditionHides))
+	UPROPERTY(EditAnywhere, Category = "Skin Weights", meta = (DisplayName = "Pruning Threshold", DataflowInput, EditCondition = "CorrectionType == ESkinWeightsCorrectionType::Prune", EditConditionHides))
 	float PruningThreshold = 0.01;
 	
 	/** Max number of bones to consider for the skin weights */
-    UPROPERTY(EditAnywhere, Category = "Skin Weights", meta = (DisplayName = "Clamping Number", EditCondition = "CorrectionType == ESkinWeightsCorrectionType::Clamp", EditConditionHides))
+    UPROPERTY(EditAnywhere, Category = "Skin Weights", meta = (DisplayName = "Clamping Number", DataflowInput, EditCondition = "CorrectionType == ESkinWeightsCorrectionType::Clamp", EditConditionHides))
     int32 ClampingNumber = 8;
 
 	/** Selection threshold to consider a neighbor skin weight */
-    UPROPERTY(EditAnywhere, Category = "Skin Weights", meta = (DisplayName = "Selection Threshold", EditCondition = "CorrectionType == ESkinWeightsCorrectionType::Hammer", EditConditionHides, ClampMin = "0.0", ClampMax = "1.0", UIMin = "0.0", UIMax = "1.0"))
+    UPROPERTY(EditAnywhere, Category = "Skin Weights", meta = (DisplayName = "Selection Threshold", DataflowInput, EditCondition = "CorrectionType == ESkinWeightsCorrectionType::Hammer", EditConditionHides, ClampMin = "0.0", ClampMax = "1.0", UIMin = "0.0", UIMax = "1.0"))
     float SelectionThreshold = 0.5f;
 
 	/** Get the weights attribute key to retrieve/set the bone indices values*/
