Dataflow : Correct skin weight node : make sure to read the values from the context

diff --git a/Engine/Plugins/Experimental/Dataflow/Source/DataflowEditor/Private/DataflowEditorTools/DataflowEditorCorrectSkinWeightsNode.cpp b/Engine/Plugins/Experimental/Dataflow/Source/DataflowEditor/Private/DataflowEditorTools/DataflowEditorCorrectSkinWeightsNode.cpp
--- a/Engine/Plugins/Experimental/Dataflow/Source/DataflowEditor/Private/DataflowEditorTools/DataflowEditorCorrectSkinWeightsNode.cpp
+++ b/Engine/Plugins/Experimental/Dataflow/Source/DataflowEditor/Private/DataflowEditorTools/DataflowEditorCorrectSkinWeightsNode.cpp
@@ -727,16 +727,18 @@ void FDataflowCorrectSkinWeightsNode::Evaluate(UE::Dataflow::FContext& Context,
 						.Setup = { SetupWeights, SetupIndices },
 						.SelectionMap = SelectionMap,
 						.Final = { FinalWeights, FinalIndices },
-						.bUseSelectionAsPerVertexFactor = bUseSelectionAsPerVertexFactor,
+						.bUseSelectionAsPerVertexFactor = GetValue(Context, &bUseSelectionAsPerVertexFactor),
 					};
 
 					if(CorrectionType == ESkinWeightsCorrectionType::Prune)
 					{
-						UE::Dataflow::Private::PruneSkinWeights(PruningThreshold, Params);
+						const float InPruningThreshold = GetValue(Context, &PruningThreshold);
+						UE::Dataflow::Private::PruneSkinWeights(InPruningThreshold, Params);
 					}
 					else if(CorrectionType == ESkinWeightsCorrectionType::Clamp)
 					{
-						UE::Dataflow::Private::ClampSkinWeights(ClampingNumber, Params);
+						const int32 InClampingNumber = GetValue(Context, &ClampingNumber);
+						UE::Dataflow::Private::ClampSkinWeights(InClampingNumber, Params);
 					}
 					else if(CorrectionType == ESkinWeightsCorrectionType::Normalize)
 					{
@@ -756,11 +758,14 @@ void FDataflowCorrectSkinWeightsNode::Evaluate(UE::Dataflow::FContext& Context,
 							
 							if(CorrectionType == ESkinWeightsCorrectionType::Relax)
 							{
-								UE::Dataflow::Private::SmoothSkinWeights(DynamicMesh, CollocatedVertices, VertexOffset, SmoothingFactor, SmoothingIterations, Params);
+								const float InSmoothingFactor = GetValue(Context, &SmoothingFactor);
+								const int32 InSmoothingIterations = GetValue(Context, &SmoothingIterations);
+								UE::Dataflow::Private::SmoothSkinWeights(DynamicMesh, CollocatedVertices, VertexOffset, InSmoothingFactor, InSmoothingIterations, Params);
 							}
 							else if (CorrectionType == ESkinWeightsCorrectionType::Hammer)
 							{
-								UE::Dataflow::Private::HammerSkinWeights(DynamicMesh, VertexOffset, SelectionThreshold, Params);
+								const float InSelectionThreshold = GetValue(Context, &SelectionThreshold);
+								UE::Dataflow::Private::HammerSkinWeights(DynamicMesh, VertexOffset, InSelectionThreshold, Params);
 							}
 
 							UE::Dataflow::Private::MergeSkinWeights(DynamicMesh, CollocatedVertices, VertexOffset, SelectionMap, FinalIndices, FinalWeights);
