Dataflow : update the generic paint weight tool to support 3D render and sim view modes from Cloth

diff --git a/Engine/Plugins/Experimental/Dataflow/Source/DataflowNodes/Private/Dataflow/DataflowCollectionAddScalarVertexPropertyNode.cpp b/Engine/Plugins/Experimental/Dataflow/Source/DataflowNodes/Private/Dataflow/DataflowCollectionAddScalarVertexPropertyNode.cpp
--- a/Engine/Plugins/Experimental/Dataflow/Source/DataflowNodes/Private/Dataflow/DataflowCollectionAddScalarVertexPropertyNode.cpp
+++ b/Engine/Plugins/Experimental/Dataflow/Source/DataflowNodes/Private/Dataflow/DataflowCollectionAddScalarVertexPropertyNode.cpp
@@ -303,7 +303,11 @@ void FDataflowCollectionAddScalarVertexPropertyNode::GetSupportedViewModes(UE::D
 {
 	// matches UE::Dataflow::FDataflowConstruction3DViewMode::Name - but we cannot include the module to get it directly 
 	static const FName DataflowConstruction3DViewModeName = TEXT("3DView");
+	static const FName Cloth3DSimView = TEXT("Cloth3DSimView");
+	static const FName ClothRenderView = TEXT("ClothRenderView");
 	OutViewModeNames.Add(DataflowConstruction3DViewModeName);
+	OutViewModeNames.Add(Cloth3DSimView);
+	OutViewModeNames.Add(ClothRenderView);
 }
 
 void FDataflowCollectionAddScalarVertexPropertyNode::GetVertexAttributeValues(UE::Dataflow::FContext& Context, TArray<float>& OutValues) const
