Dataflow : fix issue where all Dataflow editors would open a tab for a newly created subgraph from another asset
- this also fixed similar issues related to changing or deleting a subgraph

diff --git a/Engine/Plugins/Experimental/Dataflow/Source/DataflowEditor/Private/Dataflow/DataflowMembersWidget.cpp b/Engine/Plugins/Experimental/Dataflow/Source/DataflowEditor/Private/Dataflow/DataflowMembersWidget.cpp
--- a/Engine/Plugins/Experimental/Dataflow/Source/DataflowEditor/Private/Dataflow/DataflowMembersWidget.cpp
+++ b/Engine/Plugins/Experimental/Dataflow/Source/DataflowEditor/Private/Dataflow/DataflowMembersWidget.cpp
@@ -615,10 +615,13 @@ void SDataflowMembersWidget::OnObjectPropertyChanged(UObject* InObject, FPropert
 
 void SDataflowMembersWidget::OnSubGraphsChanged(const UDataflow* InDataflowAsset, const FGuid& InSubGraphGuid, UE::Dataflow::ESubGraphChangedReason InReason)
 {
-	if (InDataflowAsset == DataflowAssetWeakPtr)
+	// we only need to do something for our own asset
+	if (InDataflowAsset != DataflowAssetWeakPtr)
 	{
-		Refresh();
+		return;
 	}
+
+	Refresh();
 	
 	// handle SubGraphs tabs
 	if (TSharedPtr<FDataflowEditorToolkit> Toolkit = EditorToolkitWeakPtr.Pin())
