Add simulation controls to dataflow + fix bughawk issues

diff --git a/Engine/Plugins/Experimental/Dataflow/Source/DataflowEditor/Private/Dataflow/DataflowEditorCommands.cpp b/Engine/Plugins/Experimental/Dataflow/Source/DataflowEditor/Private/Dataflow/DataflowEditorCommands.cpp
--- a/Engine/Plugins/Experimental/Dataflow/Source/DataflowEditor/Private/Dataflow/DataflowEditorCommands.cpp
+++ b/Engine/Plugins/Experimental/Dataflow/Source/DataflowEditor/Private/Dataflow/DataflowEditorCommands.cpp
@@ -52,6 +52,11 @@ void FDataflowEditorCommandsImpl::RegisterCommands()
 	UI_COMMAND(EvaluateNode, "Evaluate", "Trigger an evaluation of the selected node.", EUserInterfaceActionType::Button, FInputChord());
 	UI_COMMAND(EvaluateGraph, "Evaluate Dataflow Graph", "Trigger an evaluation of the graph Terminal node.", EUserInterfaceActionType::Button, FInputChord(EKeys::F5));
 	UI_COMMAND(ToggleSimulation, "Simulate Dataflow Scene", "Toggle the simulation of the registered components in the simulation scene.", EUserInterfaceActionType::RadioButton, FInputChord(EKeys::S, EModifierKey::Alt));
+	UI_COMMAND(StartSimulation, "Start Dataflow Simulation", "Start the dataflow simulation", EUserInterfaceActionType::Button, FInputChord());
+	UI_COMMAND(StopSimulation, "Stop Dataflow Simulation", "Stop the dataflow simulation.", EUserInterfaceActionType::Button, FInputChord());
+	UI_COMMAND(PauseSimulation, "Pause Dataflow Simulation", "Pause the dataflow simulation.", EUserInterfaceActionType::Button, FInputChord());
+	UI_COMMAND(StepSimulation, "Step Dataflow Simulation", "Advance one frame in the dataflow simulation.", EUserInterfaceActionType::Button, FInputChord());
+	UI_COMMAND(ResetSimulation, "Reset Dataflow Simulation", "Reset the dataflow simulation.", EUserInterfaceActionType::Button, FInputChord());
 	UI_COMMAND(EvaluateGraphAutomatic, "Automatic Graph Evaluation", "Set the evaluation mode of the graph to Automatic.", EUserInterfaceActionType::RadioButton, FInputChord());
 	UI_COMMAND(EvaluateGraphManual, "Manual Graph Evaluation", "Set the evaluation mode of the graph to Manual.", EUserInterfaceActionType::RadioButton, FInputChord());
 	UI_COMMAND(ClearGraphCache, "Force Graph Re-evaluation", "Force the entire graph to re-evaluate by deleting the evaluation cache and re-evaluating the Terminal node.", EUserInterfaceActionType::Button, FInputChord());
diff --git a/Engine/Plugins/Experimental/Dataflow/Source/DataflowEditor/Private/Dataflow/DataflowEditorToolkit.cpp b/Engine/Plugins/Experimental/Dataflow/Source/DataflowEditor/Private/Dataflow/DataflowEditorToolkit.cpp
--- a/Engine/Plugins/Experimental/Dataflow/Source/DataflowEditor/Private/Dataflow/DataflowEditorToolkit.cpp
+++ b/Engine/Plugins/Experimental/Dataflow/Source/DataflowEditor/Private/Dataflow/DataflowEditorToolkit.cpp
@@ -6,6 +6,7 @@
 #include "AdvancedPreviewSceneModule.h"
 #include "AssetEditorModeManager.h"
 #include "AssetViewerSettings.h"
+#include "DataflowEditorStyle.h"
 #include "DetailCategoryBuilder.h"
 #include "Dataflow/DataflowAdvancedPreviewDetailsTab.h"
 #include "Dataflow/DataflowAssetViewerSettingsCustomization.h"
@@ -934,8 +935,26 @@ void FDataflowEditorToolkit::AddEvaluationWidget(FToolMenuSection& InSection)
 		FCanExecuteAction::CreateSP(this, &FDataflowEditorToolkit::HasSimulationManager),
 		FIsActionChecked::CreateSP(this, &FDataflowEditorToolkit::IsSimulationEnabled)
 	);
-	
-	
+	ToolkitCommands->MapAction(FDataflowEditorCommands::Get().ResetSimulation,
+		FExecuteAction::CreateSP(this, &FDataflowEditorToolkit::ResetDataflowSimulation),
+		FCanExecuteAction::CreateSP(this, &FDataflowEditorToolkit::HasSimulationManager),
+		FIsActionChecked()
+	);
+	ToolkitCommands->MapAction(FDataflowEditorCommands::Get().StartSimulation,
+		FExecuteAction::CreateSP(this, &FDataflowEditorToolkit::StartDataflowSimulation),
+		FCanExecuteAction::CreateSP(this, &FDataflowEditorToolkit::IsSimulationDisabled),
+		FIsActionChecked()
+	);
+	ToolkitCommands->MapAction(FDataflowEditorCommands::Get().StopSimulation,
+		FExecuteAction::CreateSP(this, &FDataflowEditorToolkit::StopDataflowSimulation),
+		FCanExecuteAction::CreateSP(this, &FDataflowEditorToolkit::IsSimulationEnabled),
+		FIsActionChecked()
+	);
+	ToolkitCommands->MapAction(FDataflowEditorCommands::Get().StepSimulation,
+		FExecuteAction::CreateSP(this, &FDataflowEditorToolkit::StepDataflowSimulation),
+		FCanExecuteAction::CreateSP(this, &FDataflowEditorToolkit::IsSimulationDisabled),
+		FIsActionChecked()
+	);
 	
 	InSection.AddDynamicEntry("DataflowEvaluateGraphDynamic", FNewToolMenuSectionDelegate::CreateLambda([](FToolMenuSection& DynamicSection)
 		{
@@ -970,21 +989,45 @@ void FDataflowEditorToolkit::AddEvaluationWidget(FToolMenuSection& InSection)
 			}
 		}));
 
-	InSection.AddDynamicEntry("DataflowToggleSimulation", FNewToolMenuSectionDelegate::CreateLambda([](FToolMenuSection& DynamicSection)
-		{
-			static const FSlateIcon EnableSimulationIconIcon(FAppStyle::GetAppStyleSetName(), "PhysicsAssetEditor.SimulationAll");
-			
-			using namespace UE::Dataflow::Private;
-			if (const TSharedPtr<FDataflowEditorToolkit> Toolkit = GetAssetEditorToolkitFromMenuContext(DynamicSection.FindContext<UAssetEditorToolkitMenuContext>()))
-			{
-				DynamicSection.AddEntry(FToolMenuEntry::InitToolBarButton(
-						FDataflowEditorCommands::Get().ToggleSimulation,
-						TAttribute<FText>(),
-						TAttribute<FText>(),
-						TAttribute<FSlateIcon>(EnableSimulationIconIcon)
-					));
-			}
-		}));
+	FToolMenuEntry PlayEntry = FToolMenuEntry::InitToolBarButton(
+					FDataflowEditorCommands::Get().StartSimulation,
+					TAttribute<FText>(),
+					TAttribute<FText>(),
+				FSlateIcon(FAppStyle::GetAppStyleSetName(), "PlayWorld.PlayInViewport"));
+
+	PlayEntry.StyleNameOverride = FName("Toolbar.BackplateLeftPlay");
+
+	InSection.AddEntry(PlayEntry);
+
+	FToolMenuEntry StepEntry = FToolMenuEntry::InitToolBarButton(
+	FDataflowEditorCommands::Get().StepSimulation,
+					TAttribute<FText>(),
+					TAttribute<FText>(),
+		FSlateIcon(FAppStyle::GetAppStyleSetName(), "PlayWorld.SingleFrameAdvance.Small"));
+
+	StepEntry.StyleNameOverride = FName("Toolbar.BackplateCenter");
+				
+	InSection.AddEntry(StepEntry);
+
+	FToolMenuEntry StopEntry = FToolMenuEntry::InitToolBarButton(
+	FDataflowEditorCommands::Get().StopSimulation,
+					TAttribute<FText>(),
+					TAttribute<FText>(),
+		FSlateIcon(FAppStyle::GetAppStyleSetName(), "PlayWorld.StopPlaySession.Small"));
+
+	StopEntry.StyleNameOverride = FName("Toolbar.BackplateCenterStop");
+
+	InSection.AddEntry(StopEntry);
+
+	FToolMenuEntry ResetEntry = FToolMenuEntry::InitToolBarButton(
+	FDataflowEditorCommands::Get().ResetSimulation,
+				TAttribute<FText>(),
+				TAttribute<FText>(),
+	FSlateIcon(FDataflowEditorStyle::Get().GetStyleSetName(),"Dataflow.ResetSimulation"));
+
+	ResetEntry.StyleNameOverride = FName("Toolbar.BackplateRight");
+
+	InSection.AddEntry(ResetEntry);
 
 	// load options
 	if (UDataflowEditorOptions* const Options = UDataflowEditorOptions::StaticClass()->GetDefaultObject<UDataflowEditorOptions>())
@@ -1552,7 +1595,17 @@ bool FDataflowEditorToolkit::IsGraphDirty() const
 	return false;
 }
 
-void FDataflowEditorToolkit::ToggleDataflowSimulation()
+void FDataflowEditorToolkit::ResetDataflowSimulation() const
+{
+	if(SimulationScene)
+	{
+		const bool bSimulationWasEnabled = SimulationScene->IsSimulationEnabled();
+		SimulationScene->RebuildSimulationScene();
+		SimulationScene->SetSimulationEnabled(bSimulationWasEnabled);
+	}
+}
+
+void FDataflowEditorToolkit::ToggleDataflowSimulation() const
 {
 	if(HasSimulationManager())
 	{
@@ -1563,6 +1616,47 @@ void FDataflowEditorToolkit::ToggleDataflowSimulation()
 	}
 }
 
+void FDataflowEditorToolkit::StartDataflowSimulation() const
+{
+	if(HasSimulationManager())
+	{
+		UDataflowSimulationManager* SimulationManager = SimulationScene->GetWorld()->GetSubsystem<UDataflowSimulationManager>();
+		
+		SimulationManager->SetSimulationEnabled(true);
+	}
+}
+
+void FDataflowEditorToolkit::StepDataflowSimulation() const
+{
+	if(HasSimulationManager())
+	{
+		UDataflowSimulationManager* SimulationManager = SimulationScene->GetWorld()->GetSubsystem<UDataflowSimulationManager>();
+
+		SimulationManager->SetSimulationEnabled(true);
+		SimulationManager->SetSimulationStepping(true);
+	}
+}
+
+void FDataflowEditorToolkit::PauseDataflowSimulation() const
+{
+	if(HasSimulationManager())
+	{
+		UDataflowSimulationManager* SimulationManager = SimulationScene->GetWorld()->GetSubsystem<UDataflowSimulationManager>();
+		
+		SimulationManager->SetSimulationEnabled(false);
+	}
+}
+
+void FDataflowEditorToolkit::StopDataflowSimulation() const
+{
+	if(HasSimulationManager())
+	{
+		UDataflowSimulationManager* SimulationManager = SimulationScene->GetWorld()->GetSubsystem<UDataflowSimulationManager>();
+		
+		SimulationManager->SetSimulationEnabled(false);
+	}
+}
+
 bool FDataflowEditorToolkit::IsSimulationEnabled() const
 {
 	if(HasSimulationManager())
@@ -1573,6 +1667,11 @@ bool FDataflowEditorToolkit::IsSimulationEnabled() const
 	return false;
 }
 
+bool FDataflowEditorToolkit::IsSimulationDisabled() const
+{
+	return !IsSimulationEnabled();
+}
+
 bool FDataflowEditorToolkit::HasSimulationManager() const
 {
 	return SimulationScene && SimulationScene->GetWorld() && SimulationScene->GetWorld()->GetSubsystem<UDataflowSimulationManager>();
diff --git a/Engine/Plugins/Experimental/Dataflow/Source/DataflowEditor/Private/Widgets/SDataflowTransportControl.cpp b/Engine/Plugins/Experimental/Dataflow/Source/DataflowEditor/Private/Widgets/SDataflowTransportControl.cpp
--- a/Engine/Plugins/Experimental/Dataflow/Source/DataflowEditor/Private/Widgets/SDataflowTransportControl.cpp
+++ b/Engine/Plugins/Experimental/Dataflow/Source/DataflowEditor/Private/Widgets/SDataflowTransportControl.cpp
@@ -33,9 +33,6 @@ void SDataflowTransportControl::Construct(const FArguments& InArgs, const TShare
 	const FTransportControlWidget PreviewModeWidget(FOnMakeTransportWidget::CreateSP(this, &SDataflowTransportControl::OnCreateModeButton));
 	TransportControlWidgets.Add(PreviewModeWidget);
 
-	const FTransportControlWidget ResetSimulationWidget(FOnMakeTransportWidget::CreateSP(this, &SDataflowTransportControl::OnCreateResetButton));
-	TransportControlWidgets.Add(ResetSimulationWidget);
-
 	FTransportControlArgs TransportControlArgs;
 	TransportControlArgs.OnForwardPlay = FOnClicked::CreateSP(this, &SDataflowTransportControl::OnForwardPlay);
 	TransportControlArgs.OnBackwardPlay = FOnClicked::CreateSP(this, &SDataflowTransportControl::OnBackwardPlay);
diff --git a/Engine/Plugins/Experimental/Dataflow/Source/DataflowEditor/Public/Dataflow/DataflowEditorCommands.h b/Engine/Plugins/Experimental/Dataflow/Source/DataflowEditor/Public/Dataflow/DataflowEditorCommands.h
--- a/Engine/Plugins/Experimental/Dataflow/Source/DataflowEditor/Public/Dataflow/DataflowEditorCommands.h
+++ b/Engine/Plugins/Experimental/Dataflow/Source/DataflowEditor/Public/Dataflow/DataflowEditorCommands.h
@@ -60,6 +60,11 @@ public:
 	TSharedPtr<FUICommandInfo> EvaluateNode;
 	TSharedPtr<FUICommandInfo> EvaluateGraph;
 	TSharedPtr<FUICommandInfo> ToggleSimulation;
+	TSharedPtr<FUICommandInfo> StartSimulation;
+	TSharedPtr<FUICommandInfo> PauseSimulation;
+	TSharedPtr<FUICommandInfo> StepSimulation;
+	TSharedPtr<FUICommandInfo> ResetSimulation;
+	TSharedPtr<FUICommandInfo> StopSimulation;
 	TSharedPtr<FUICommandInfo> EvaluateGraphAutomatic;
 	TSharedPtr<FUICommandInfo> EvaluateGraphManual;
 	TSharedPtr<FUICommandInfo> ClearGraphCache;
diff --git a/Engine/Plugins/Experimental/Dataflow/Source/DataflowEditor/Public/Dataflow/DataflowEditorToolkit.h b/Engine/Plugins/Experimental/Dataflow/Source/DataflowEditor/Public/Dataflow/DataflowEditorToolkit.h
--- a/Engine/Plugins/Experimental/Dataflow/Source/DataflowEditor/Public/Dataflow/DataflowEditorToolkit.h
+++ b/Engine/Plugins/Experimental/Dataflow/Source/DataflowEditor/Public/Dataflow/DataflowEditorToolkit.h
@@ -112,8 +112,23 @@ public:
 	/** Get the toolkit evaluation mode */
 	const EDataflowEditorEvaluationMode& GetEvaluationMode() const { return EvaluationMode; }
 
+	/** Reset dataflow simulation */
+	void ResetDataflowSimulation() const;
+	
 	/** Toggle dataflow simulation */
-	void ToggleDataflowSimulation();
+	void ToggleDataflowSimulation() const;
+	
+	/** Start dataflow simulation */
+	void StartDataflowSimulation() const;
+
+	/** Stop dataflow simulation */
+	void StopDataflowSimulation() const;
+	
+	/** Step dataflow simulation */
+	void StepDataflowSimulation() const;
+
+	/** Pause dataflow simulation */
+	void PauseDataflowSimulation() const;
 
 	/** Check if the simulation can be toggled or not */
 	bool HasSimulationManager() const;
@@ -121,6 +136,9 @@ public:
 	/** Check if the simulation is enabled or not */
 	bool IsSimulationEnabled() const;
 
+	/** Check if the simulation is disabled or not */
+	bool IsSimulationDisabled() const;
+
 protected:
 
 	UDataflowEditor* DataflowEditor = nullptr;
diff --git a/Engine/Plugins/Experimental/PlanarCutPlugin/Source/PlanarCut/Private/GeometryMeshConversion.cpp b/Engine/Plugins/Experimental/PlanarCutPlugin/Source/PlanarCut/Private/GeometryMeshConversion.cpp
--- a/Engine/Plugins/Experimental/PlanarCutPlugin/Source/PlanarCut/Private/GeometryMeshConversion.cpp
+++ b/Engine/Plugins/Experimental/PlanarCutPlugin/Source/PlanarCut/Private/GeometryMeshConversion.cpp
@@ -466,6 +466,10 @@ namespace AugmentedDynamicMesh
 	
 	void InitializeOverlayToPerVertexColors(FDynamicMesh3& Mesh)
 	{
+		if (!ensure(Mesh.HasAttributes()))
+		{
+			return;
+		}
 		Mesh.Attributes()->EnablePrimaryColors();
 		Mesh.Attributes()->PrimaryColors()->CreatePerVertex(0.f);
 		Mesh.EnableVertexColors(FVector3f::Zero());
diff --git a/Engine/Source/Runtime/Experimental/Chaos/Private/GeometryCollection/Facades/CollectionCurveFacade.cpp b/Engine/Source/Runtime/Experimental/Chaos/Private/GeometryCollection/Facades/CollectionCurveFacade.cpp
--- a/Engine/Source/Runtime/Experimental/Chaos/Private/GeometryCollection/Facades/CollectionCurveFacade.cpp
+++ b/Engine/Source/Runtime/Experimental/Chaos/Private/GeometryCollection/Facades/CollectionCurveFacade.cpp
@@ -205,7 +205,7 @@ namespace GeometryCollection::Facades
 						VertexColors[VertexIndex] = GeometryColor;
 						VertexColors[VertexIndex+1] = GeometryColor;
 						
-						// Build the vertex colors
+						// Build the vertex transform indices
 						VertexTransforms[VertexIndex] = GeometryIndex;
 						VertexTransforms[VertexIndex+1] = GeometryIndex;
 
diff --git a/Engine/Source/Runtime/Experimental/Chaos/Private/GeometryCollection/Facades/CollectionVertexBoneWeightsFacade.cpp b/Engine/Source/Runtime/Experimental/Chaos/Private/GeometryCollection/Facades/CollectionVertexBoneWeightsFacade.cpp
--- a/Engine/Source/Runtime/Experimental/Chaos/Private/GeometryCollection/Facades/CollectionVertexBoneWeightsFacade.cpp
+++ b/Engine/Source/Runtime/Experimental/Chaos/Private/GeometryCollection/Facades/CollectionVertexBoneWeightsFacade.cpp
@@ -133,30 +133,33 @@ namespace GeometryCollection::Facades
 
 	void FVertexBoneWeightsFacade::ModifyBoneWeight(int32 VertexIndex, const TArray<int32>& VertexBoneIndices, const TArray<float>& VertexBoneWeights)
 	{
-		TManagedArray< TArray<int32> >& IndicesArray = BoneIndicesAttribute.Modify();
-		TManagedArray< TArray<float> >& WeightsArray = BoneWeightsAttribute.Modify();
-		
-		if (IndicesArray.IsValidIndex(VertexIndex) && WeightsArray.IsValidIndex(VertexIndex))
-		{
-			IndicesArray[VertexIndex].Empty();
-			WeightsArray[VertexIndex].Empty();
-			float TotalWeight = 0.f;
-			int32 LocalIndex = 0;
-			
-			for (const int32& BoneIndex : VertexBoneIndices)
+		if(VertexBoneIndices.Num() == VertexBoneWeights.Num())
+		{ 
+			TManagedArray< TArray<int32> >& IndicesArray = BoneIndicesAttribute.Modify();
+			TManagedArray< TArray<float> >& WeightsArray = BoneWeightsAttribute.Modify();
+
+			if (IndicesArray.IsValidIndex(VertexIndex) && WeightsArray.IsValidIndex(VertexIndex))
 			{
-				if (!bInternalWeights || (0 <= BoneIndex && BoneIndex < NumBones()))
+				IndicesArray[VertexIndex].Empty();
+				WeightsArray[VertexIndex].Empty();
+				float TotalWeight = 0.f;
+				int32 LocalIndex = 0;
+
+				for (const int32& BoneIndex : VertexBoneIndices)
 				{
-					const float BoneWeight = VertexBoneWeights[LocalIndex++];
-					
-					IndicesArray[VertexIndex].Add(BoneIndex);
-					WeightsArray[VertexIndex].Add(BoneWeight);
-					TotalWeight += BoneWeight;
+					if (!bInternalWeights || (0 <= BoneIndex && BoneIndex < NumBones()))
+					{
+						const float BoneWeight = VertexBoneWeights[LocalIndex++];
+
+						IndicesArray[VertexIndex].Add(BoneIndex);
+						WeightsArray[VertexIndex].Add(BoneWeight);
+						TotalWeight += BoneWeight;
+					}
+				}
+				if (TotalWeight < 1.f - UE_KINDA_SMALL_NUMBER || TotalWeight > 1.f + UE_KINDA_SMALL_NUMBER)
+				{
+					UE_LOG(LogChaos, Warning, TEXT("FVertexBoneWeightsFacade::ModifyBoneWeight: Bone weight sum %f is not 1 on vertex %d"), TotalWeight, VertexIndex);
 				}
-			}
-			if (TotalWeight < 1.f - UE_KINDA_SMALL_NUMBER || TotalWeight > 1.f + UE_KINDA_SMALL_NUMBER)
-			{
-				UE_LOG(LogChaos, Warning, TEXT("FVertexBoneWeightsFacade::ModifyBoneWeight: Bone weight sum %f is not 1 on vertex %d"), TotalWeight, VertexIndex);
 			}
 		}
 	}
@@ -193,7 +196,7 @@ namespace GeometryCollection::Facades
     	}
 		else if (KinematicWeightAttribute.IsValid())
 		{
-			return KinematicWeightAttribute.IsValidIndex(VertexIndex) && (KinematicWeightAttribute.Get()[VertexIndex] != 1.0);
+			return KinematicWeightAttribute.IsValidIndex(VertexIndex) && (KinematicWeightAttribute.Get()[VertexIndex] == 1.0);
 		}
 		else //backward compatibility for KinematicAttribute added in 5.5
 		{
@@ -266,7 +269,7 @@ namespace GeometryCollection::Facades
 				{
 					TotalWeight += BoneWeight;
 				}
-				if(TotalWeight <= 1.f + UE_KINDA_SMALL_NUMBER)
+				if (TotalWeight > UE_KINDA_SMALL_NUMBER)
 				{
 					for(float& BoneWeight : BoneWeights)
 					{
