Dataflow : enable by default the new weight map painting tool for both CLoth and the generic weight map paint nodes

diff --git a/Engine/Plugins/Experimental/Dataflow/Source/DataflowEditor/Private/DataflowEditorTools/DataflowEditorVertexAttributePaintTool.cpp b/Engine/Plugins/Experimental/Dataflow/Source/DataflowEditor/Private/DataflowEditorTools/DataflowEditorVertexAttributePaintTool.cpp
--- a/Engine/Plugins/Experimental/Dataflow/Source/DataflowEditor/Private/DataflowEditorTools/DataflowEditorVertexAttributePaintTool.cpp
+++ b/Engine/Plugins/Experimental/Dataflow/Source/DataflowEditor/Private/DataflowEditorTools/DataflowEditorVertexAttributePaintTool.cpp
@@ -53,7 +53,7 @@ DEFINE_LOG_CATEGORY_STATIC(LogDataflowEditorVertexAttributePaintTool, Warning, A
 
 namespace UE::DataflowEditorVertexAttributePaintTool::CVars
 {
-	bool DataflowEditorUseNewWeightMapTool = false;
+	bool DataflowEditorUseNewWeightMapTool = true;
 	FAutoConsoleVariableRef CVarDataflowEditorUseNewWeightMapTool(
 		TEXT("p.DataflowEditor.UseNewWeightMapTool"),
 		DataflowEditorUseNewWeightMapTool,
