Dataflow : Add long task dialogs to inform the user of what's going on when UI is  slow to refresh if the data is quite complex
* DataflowSkinWeightPaint tool initialization
* Component generation for construction view
* Collection Spreadsheet update

diff --git a/Engine/Plugins/Experimental/Dataflow/Source/DataflowEditor/Private/Dataflow/CollectionSpreadSheetWidget.cpp b/Engine/Plugins/Experimental/Dataflow/Source/DataflowEditor/Private/Dataflow/CollectionSpreadSheetWidget.cpp
--- a/Engine/Plugins/Experimental/Dataflow/Source/DataflowEditor/Private/Dataflow/CollectionSpreadSheetWidget.cpp
+++ b/Engine/Plugins/Experimental/Dataflow/Source/DataflowEditor/Private/Dataflow/CollectionSpreadSheetWidget.cpp
@@ -8,6 +8,7 @@
 #include "GeometryCollection/GeometryCollection.h"
 #include "Dataflow/DataflowCollectionSpreadSheetHelpers.h"
 #include "GeometryCollection/ManagedArrayCollection.h"
+#include "Misc/ScopedSlowTask.h"
 #include "Widgets/Views/STableRow.h"
 #include "Widgets/Text/STextBlock.h"
 #include "Widgets/Input/SComboBox.h"
@@ -238,6 +239,10 @@ TSharedRef<ITableRow> SCollectionSpreadSheet::GenerateRow(TSharedPtr<const FColl
 
 void SCollectionSpreadSheetWidget::NodeOutputsComboBoxSelectionChanged(FName InSelectedOutput, ESelectInfo::Type InSelectInfo)
 {
+	const float AmountOfWork = 3;
+	FScopedSlowTask SlowTask(AmountOfWork, LOCTEXT("SCollectionSpreadSheetWidget_SelectionChangedProgress", "Updating Collection Spreadsheet"));
+
+	SlowTask.MakeDialog(/*bShowCancelButton*/false);
 	if (CollectionTable)
 	{
 		if (CollectionTable->GetSelectedOutput() != InSelectedOutput)
@@ -258,6 +263,7 @@ void SCollectionSpreadSheetWidget::NodeOutputsComboBoxSelectionChanged(FName InS
 		}
 	}
 
+	SlowTask.EnterProgressFrame();
 	if (TransformOutliner)
 	{
 		const FString OutputStr = InSelectedOutput.ToString();
@@ -271,6 +277,7 @@ void SCollectionSpreadSheetWidget::NodeOutputsComboBoxSelectionChanged(FName InS
 		}
 	}
 
+	SlowTask.EnterProgressFrame();
 	if (VerticesOutliner)
 	{
 		const FString OutputStr = InSelectedOutput.ToString();
@@ -284,6 +291,7 @@ void SCollectionSpreadSheetWidget::NodeOutputsComboBoxSelectionChanged(FName InS
 		}
 	}
 
+	SlowTask.EnterProgressFrame();
 	if (FacesOutliner)
 	{
 		const FString OutputStr = InSelectedOutput.ToString();
diff --git a/Engine/Plugins/Experimental/Dataflow/Source/DataflowEditor/Private/Dataflow/DataflowConstructionScene.cpp b/Engine/Plugins/Experimental/Dataflow/Source/DataflowEditor/Private/Dataflow/DataflowConstructionScene.cpp
--- a/Engine/Plugins/Experimental/Dataflow/Source/DataflowEditor/Private/Dataflow/DataflowConstructionScene.cpp
+++ b/Engine/Plugins/Experimental/Dataflow/Source/DataflowEditor/Private/Dataflow/DataflowConstructionScene.cpp
@@ -573,6 +573,10 @@ void FDataflowConstructionScene::UpdateComponentsFromNode(const UDataflowEdNode*
 		return;
 	}
 
+	const float AmountOfWork = 1;
+	FScopedSlowTask SlowTask(AmountOfWork, LOCTEXT("FDataflowConstructionScene_UpdateCOmponentFromNode_Progress", "Generating construction view components..."));
+	SlowTask.MakeDialog(/*bShowCancelButton*/false);
+
 	using namespace UE::Geometry;//FDynamicMesh3
 
 	const FDataflowEditorToolkit* EditorToolkit = DataflowEditor ? static_cast<FDataflowEditorToolkit*>(DataflowEditor->GetInstanceInterface()) : nullptr;
diff --git a/Engine/Plugins/Experimental/Dataflow/Source/DataflowEditor/Private/DataflowEditorTools/DataflowEditorSkinWeightsPaintTool.cpp b/Engine/Plugins/Experimental/Dataflow/Source/DataflowEditor/Private/DataflowEditorTools/DataflowEditorSkinWeightsPaintTool.cpp
--- a/Engine/Plugins/Experimental/Dataflow/Source/DataflowEditor/Private/DataflowEditorTools/DataflowEditorSkinWeightsPaintTool.cpp
+++ b/Engine/Plugins/Experimental/Dataflow/Source/DataflowEditor/Private/DataflowEditorTools/DataflowEditorSkinWeightsPaintTool.cpp
@@ -13,6 +13,7 @@
 #include "Components/SkeletalMeshComponent.h"
 #include "DynamicMesh/DynamicVertexSkinWeightsAttribute.h"
 #include "DynamicMesh/NonManifoldMappingSupport.h"
+#include "Misc/ScopedSlowTask.h"
 #include "TargetInterfaces/MeshDescriptionCommitter.h"
 #include "TargetInterfaces/MeshDescriptionProvider.h"
 #include "TargetInterfaces/PrimitiveComponentBackedTarget.h"
@@ -175,10 +176,13 @@ bool UDataflowEditorSkinWeightsPaintTool::ExtractSkinWeights(TArray<TArray<int32
 
 void UDataflowEditorSkinWeightsPaintTool::Setup()
 {
+	FScopedSlowTask SlowTask(2, LOCTEXT("UDataflowEditorSkinWeightsPaintTool", "Starting Dataflow Skin Weight Paint Tool..."));
+	SlowTask.MakeDialog(/*bShowCancelButton*/false);
+
 	const int32 VertexOffset = GetVertexOffset();
 	if(VertexOffset != INDEX_NONE)
 	{
-		if(FMeshDescription* MeshDescription = GetCurrentDescription(0))
+		if (FMeshDescription* MeshDescription = GetCurrentDescription(0))
 		{
 			const int32 NumVertices = MeshDescription->Vertices().Num();
 			
@@ -214,6 +218,8 @@ void UDataflowEditorSkinWeightsPaintTool::Setup()
 			}
 		}
 	}
+
+	SlowTask.EnterProgressFrame();
 	USkinWeightsPaintTool::Setup();
 
 	// enable back face option 
