Optimize Edit Skin Weight node serialization for large arrays + make the skeletal meshes double sided

diff --git a/Engine/Plugins/Experimental/Dataflow/Source/DataflowNodes/Private/Dataflow/DataflowCollectionEditSkinWeightsNode.cpp b/Engine/Plugins/Experimental/Dataflow/Source/DataflowNodes/Private/Dataflow/DataflowCollectionEditSkinWeightsNode.cpp
--- a/Engine/Plugins/Experimental/Dataflow/Source/DataflowNodes/Private/Dataflow/DataflowCollectionEditSkinWeightsNode.cpp
+++ b/Engine/Plugins/Experimental/Dataflow/Source/DataflowNodes/Private/Dataflow/DataflowCollectionEditSkinWeightsNode.cpp
@@ -165,6 +165,74 @@ void CorrectSkinWeights(TArray<TArray<int32>>& BoneIndices, TArray<TArray<float>
 
 }
 
+void FDataflowVertexSkinWeightData::FromArrays(int32 VertexIndex, const TArray<int32>& InIndices, const TArray<float>& InWeights)
+{
+	if (ensure(VertexIndex < ArrayNum))
+	{
+		ensure(InIndices.Num() == InWeights.Num());
+		const int32 Offset = VertexIndex * MaxNumInflences;
+		for (int32 Index = 0; Index < MaxNumInflences; ++Index)
+		{
+			Bones[Offset + Index] = (InIndices.IsValidIndex(Index)) ? InIndices[Index] : INDEX_NONE;
+			Weights[Offset + Index] = (InWeights.IsValidIndex(Index)) ? InWeights[Index] : 0.0f;
+		}
+	}
+}
+
+void FDataflowVertexSkinWeightData::ToArrays(int32 VertexIndex, TArray<int32>& OutIndices, TArray<float>& OutWeights) const
+{
+	if (ensure(VertexIndex < ArrayNum))
+	{
+		const int32 Offset = VertexIndex * MaxNumInflences;
+		for (int32 Index = 0; Index < MaxNumInflences; ++Index)
+		{
+			if (Bones[Index] != INDEX_NONE)
+			{
+				OutIndices.Add(Bones[Offset + Index]);
+				OutWeights.Add(Weights[Offset + Index]);
+			}
+		}
+	}
+}
+
+int32 FDataflowVertexSkinWeightData::Num() const
+{
+	return ArrayNum;
+}
+
+void FDataflowVertexSkinWeightData::SetNum(int32 VertexNum)
+{
+	if (ArrayNum != VertexNum)
+	{
+		Bones.Init(INDEX_NONE, VertexNum * MaxNumInflences);
+		Weights.Init(0.0f, VertexNum * MaxNumInflences);
+		ArrayNum = VertexNum;
+	}
+	
+}
+
+void FDataflowVertexSkinWeightData::Reset()
+{
+	Bones.Reset();
+	Weights.Reset();
+	ArrayNum = 0;
+}
+
+bool FDataflowVertexSkinWeightData::Serialize(FArchive& Ar)
+{
+	Ar << ArrayNum;
+	Bones.BulkSerialize(Ar);
+	Weights.BulkSerialize(Ar);
+	return true;
+}
+
+FArchive& operator<<(FArchive& Ar, FDataflowVertexSkinWeightData& Data)
+{
+	Data.Serialize(Ar);
+	return Ar;
+}
+
+
 //
 // FDataflowCollectionEditSkinWeightsNode
 //
@@ -181,6 +249,25 @@ FDataflowCollectionEditSkinWeightsNode::FDataflowCollectionEditSkinWeightsNode(c
 	RegisterOutputConnection(&BoneWeightsKey, &BoneWeightsKey);
 }
 
+void FDataflowCollectionEditSkinWeightsNode::PostSerialize(const FArchive& Ar)
+{
+	FDataflowPrimitiveNode::PostSerialize(Ar);
+
+	if (Ar.IsLoading())
+	{
+		if (SkinWeights_DEPRECATED.Num() > 0)
+		{
+			VertexSkinWeights.SetNum(SkinWeights_DEPRECATED.Num());
+			for (int32 Index = 0; Index < SkinWeights_DEPRECATED.Num(); ++Index)
+			{
+				const FDataflowSkinWeightData& OldSkinWeights = SkinWeights_DEPRECATED[Index];
+				VertexSkinWeights.FromArrays(Index, OldSkinWeights.BoneIndices, OldSkinWeights.BoneWeights);
+			}
+		}
+		SkinWeights_DEPRECATED.Empty();
+	}
+}
+
 TArray<UE::Dataflow::FRenderingParameter> FDataflowCollectionEditSkinWeightsNode::GetRenderParametersImpl() const
 {
 	return FDataflowAddScalarVertexPropertyCallbackRegistry::Get().GetRenderingParameters(VertexGroup.Name);
@@ -224,6 +311,8 @@ void FDataflowCollectionEditSkinWeightsNode::AddPrimitiveComponents(UE::Dataflow
 				FName SkeletalMeshName(FString("Dataflow_SkeletalMesh") + FString::FromInt(GeometryIndex));
 				USkeletalMeshComponent* SkeletalMeshComponent = NewObject<USkeletalMeshComponent>(RootActor, SkeletalMeshName);
 				SkeletalMeshComponent->SetSkeletalMesh(SkeletalMeshes[GeometryIndex]);
+				UMaterial* DefaultTwoSidedMaterial = Cast<UMaterial>(StaticLoadObject(UMaterial::StaticClass(), NULL, TEXT("/Engine/EditorMaterials/Dataflow/DataflowTwoSidedVertexMaterial")));
+				SkeletalMeshComponent->SetMaterial(0, DefaultTwoSidedMaterial);
 				PrimitiveComponents.Add(SkeletalMeshComponent);
 			}
 		}
@@ -234,7 +323,7 @@ void FDataflowCollectionEditSkinWeightsNode::OnPropertyChanged(UE::Dataflow::FCo
 {
 	if (InPropertyChangedEvent.GetPropertyName() == GET_MEMBER_NAME_CHECKED(FDataflowCollectionEditSkinWeightsNode, SkeletalMesh))
 	{
-		SkinWeights.Reset();
+		VertexSkinWeights.Reset();
 	}
 }
 
@@ -300,12 +389,11 @@ void FDataflowCollectionEditSkinWeightsNode::ReportVertexWeights(UE::Dataflow::F
 	SkinWeightsRefSkeletonHash = ComputeInputRefSkeletonHashFromInput(Context);
 
 	// update the weights themselves
-	SkinWeights.SetNumZeroed(FinalWeights.Num());
+	VertexSkinWeights.SetNum(FinalWeights.Num());
 	
 	for (int32 VertexIndex = 0, NumVertices = FinalWeights.Num(); VertexIndex < NumVertices; ++VertexIndex)
 	{
-		SkinWeights[VertexIndex].BoneWeights = FinalWeights[VertexIndex];
-		SkinWeights[VertexIndex].BoneIndices = FinalIndices[VertexIndex];
+		VertexSkinWeights.FromArrays(VertexIndex, FinalIndices[VertexIndex], FinalWeights[VertexIndex]);
 	}
 }
 
@@ -319,16 +407,12 @@ void FDataflowCollectionEditSkinWeightsNode::ExtractVertexWeights(UE::Dataflow::
 	// we only returned the stored weight if they match the input refskeleton, other wise, returned the original weights
 	const int32 RefSkeletonHashFromInput = ComputeInputRefSkeletonHashFromInput(Context);
 
-	if(SkinWeights.Num() == FinalWeights.Num() && RefSkeletonHashFromInput == SkinWeightsRefSkeletonHash)
+	if(VertexSkinWeights.Num() == FinalWeights.Num() && RefSkeletonHashFromInput == SkinWeightsRefSkeletonHash)
 	{
 		for (int32 VertexIndex = 0, NumVertices = FinalWeights.Num(); VertexIndex < NumVertices; ++VertexIndex)
 		{
-			if(!SkinWeights[VertexIndex].BoneWeights.IsEmpty() && !SkinWeights[VertexIndex].BoneIndices.IsEmpty())
-			{
-				FinalWeights[VertexIndex] = SkinWeights[VertexIndex].BoneWeights;
-				FinalIndices[VertexIndex] = SkinWeights[VertexIndex].BoneIndices;
-			}
-			else
+			VertexSkinWeights.ToArrays(VertexIndex, FinalIndices[VertexIndex], FinalWeights[VertexIndex]);
+			if (FinalIndices.IsEmpty())
 			{
 				FinalWeights[VertexIndex] = SetupWeights[VertexIndex];
 				FinalIndices[VertexIndex] = SetupIndices[VertexIndex];
@@ -468,7 +552,7 @@ class FDataflowCollectionEditSkinWeightsNode::FEditNodeToolChange final : public
 public:
 	FEditNodeToolChange(const FDataflowCollectionEditSkinWeightsNode& Node) :
 		NodeGuid(Node.GetGuid()),
-		SavedSkinWeightData(Node.SkinWeights)
+		SavedSkinWeightData(Node.VertexSkinWeights)
 	{}
 
 private:
@@ -495,7 +579,7 @@ private:
 			{
 				if (FDataflowCollectionEditSkinWeightsNode* const Node = BaseNode->AsType<FDataflowCollectionEditSkinWeightsNode>())
 				{
-					Swap(SavedSkinWeightData, Node->SkinWeights);
+					Swap(SavedSkinWeightData, Node->VertexSkinWeights);
 					Node->Invalidate();
 				}
 			}
@@ -503,7 +587,7 @@ private:
 	}
 
 	FGuid NodeGuid;
-	TArray<FDataflowSkinWeightData> SavedSkinWeightData;
+	FDataflowVertexSkinWeightData SavedSkinWeightData;
 };
 
 TUniquePtr<FToolCommandChange> FDataflowCollectionEditSkinWeightsNode::MakeEditNodeToolChange()
diff --git a/Engine/Plugins/Experimental/Dataflow/Source/DataflowNodes/Public/Dataflow/DataflowCollectionEditSkinWeightsNode.h b/Engine/Plugins/Experimental/Dataflow/Source/DataflowNodes/Public/Dataflow/DataflowCollectionEditSkinWeightsNode.h
--- a/Engine/Plugins/Experimental/Dataflow/Source/DataflowNodes/Public/Dataflow/DataflowCollectionEditSkinWeightsNode.h
+++ b/Engine/Plugins/Experimental/Dataflow/Source/DataflowNodes/Public/Dataflow/DataflowCollectionEditSkinWeightsNode.h
@@ -3,12 +3,14 @@
 #pragma once 
 
 #include "DataflowCollectionAddScalarVertexPropertyNode.h"
+
 #include "DataflowPrimitiveNode.h"
 #include "DataflowSkeletalMeshNodes.h"
+#include "Dataflow/DataflowCollectionAttributeKeyNodes.h"
 #include "Dataflow/DataflowNode.h"
 #include "GeometryCollection/GeometryCollection.h"
 #include "GeometryCollection/ManagedArrayCollection.h"
-#include "Dataflow/DataflowCollectionAttributeKeyNodes.h"
+#include "GPUSkinPublicDefs.h"
 
 #include "DataflowCollectionEditSkinWeightsNode.generated.h"
 
@@ -27,6 +29,49 @@ struct FDataflowSkinWeightData
 	TArray<int32> BoneIndices;
 };
 
+/**
+* FDataflowSkinWeightData is a replacement for FDataflowSkinWeightData in FDataflowCollectionEditSkinWeightsNode
+* This allow storing all the data linearly and enable bulk serialization for speed
+* This solve the problem when we have million vertices meshes that was causing the serialization of the node to take a very long time
+* this is now an order of magnitude faster that using FDataflowSkinWeightData
+*/
+USTRUCT()
+struct FDataflowVertexSkinWeightData 
+{
+	GENERATED_BODY()
+
+public:
+	static constexpr const int32 MaxNumInflences = MAX_TOTAL_INFLUENCES;
+
+	int32 Num() const;
+	void SetNum(int32 VertexNum);
+	void Reset();
+
+	void FromArrays(int32 VertexIndex, const TArray<int32>& InIndices, const TArray<float>& InWeights);
+	void ToArrays(int32 VertexIndex, TArray<int32>& OutIndices, TArray<float>& OutWeights) const;
+
+	bool Serialize(FArchive& Ar);
+
+	friend FArchive& operator<<(FArchive& Ar, FDataflowVertexSkinWeightData& Data);
+
+private:
+	TArray<int16> Bones;
+
+	TArray<float> Weights;
+
+	int32 ArrayNum = 0;
+};
+
+template<>
+struct TStructOpsTypeTraits<FDataflowVertexSkinWeightData> : public TStructOpsTypeTraitsBase2<FDataflowVertexSkinWeightData>
+{
+	enum
+	{
+		WithSerializer = true,
+	};
+	static constexpr EPropertyObjectReferenceType WithSerializerObjectReferences = EPropertyObjectReferenceType::None;
+};
+
 /** Edit skin weights vertex properties. */
 USTRUCT(Meta = (Experimental,DataflowCollection))
 struct FDataflowCollectionEditSkinWeightsNode : public FDataflowPrimitiveNode
@@ -70,8 +115,11 @@ public:
 	bool bCompressSkinWeights = false;
 
 	/** List of skin weights */
+	UPROPERTY(meta = (DeprecatedProperty, DeprecationMessage = "Use VertexSkinWeights instead"))
+	TArray<FDataflowSkinWeightData> SkinWeights_DEPRECATED;
+
 	UPROPERTY()
-	TArray<FDataflowSkinWeightData> SkinWeights;
+	FDataflowVertexSkinWeightData VertexSkinWeights;
 	
 	/** 
 	* Hash of the ref skeleton the stored skin weight are based on 
@@ -136,6 +184,8 @@ private:
 #endif
 	//~ End FDataflowPrimitiveNode interface
 
+	virtual void PostSerialize(const FArchive& Ar) override;
+
 	int32 ComputeInputRefSkeletonHashFromInput(UE::Dataflow::FContext& Context) const;
 
 	/** Transient skeletal mesh built from dataflow render collection */
