Fix some code analysis issues for the edit skin weight node recent changes

diff --git a/Engine/Plugins/Experimental/Dataflow/Source/DataflowNodes/Private/Dataflow/DataflowCollectionEditSkinWeightsNode.cpp b/Engine/Plugins/Experimental/Dataflow/Source/DataflowNodes/Private/Dataflow/DataflowCollectionEditSkinWeightsNode.cpp
--- a/Engine/Plugins/Experimental/Dataflow/Source/DataflowNodes/Private/Dataflow/DataflowCollectionEditSkinWeightsNode.cpp
+++ b/Engine/Plugins/Experimental/Dataflow/Source/DataflowNodes/Private/Dataflow/DataflowCollectionEditSkinWeightsNode.cpp
@@ -186,7 +186,7 @@ void FDataflowVertexSkinWeightData::ToArrays(int32 VertexIndex, TArray<int32>& O
 		const int32 Offset = VertexIndex * MaxNumInflences;
 		for (int32 Index = 0; Index < MaxNumInflences; ++Index)
 		{
-			if (Bones[Index] != INDEX_NONE)
+			if (Bones[Offset + Index] != INDEX_NONE)
 			{
 				OutIndices.Add(Bones[Offset + Index]);
 				OutWeights.Add(Weights[Offset + Index]);
@@ -412,7 +412,7 @@ void FDataflowCollectionEditSkinWeightsNode::ExtractVertexWeights(UE::Dataflow::
 		for (int32 VertexIndex = 0, NumVertices = FinalWeights.Num(); VertexIndex < NumVertices; ++VertexIndex)
 		{
 			VertexSkinWeights.ToArrays(VertexIndex, FinalIndices[VertexIndex], FinalWeights[VertexIndex]);
-			if (FinalIndices.IsEmpty())
+			if (FinalIndices[VertexIndex].IsEmpty())
 			{
 				FinalWeights[VertexIndex] = SetupWeights[VertexIndex];
 				FinalIndices[VertexIndex] = SetupIndices[VertexIndex];
diff --git a/Engine/Plugins/Experimental/Dataflow/Source/DataflowNodes/Public/Dataflow/DataflowCollectionEditSkinWeightsNode.h b/Engine/Plugins/Experimental/Dataflow/Source/DataflowNodes/Public/Dataflow/DataflowCollectionEditSkinWeightsNode.h
--- a/Engine/Plugins/Experimental/Dataflow/Source/DataflowNodes/Public/Dataflow/DataflowCollectionEditSkinWeightsNode.h
+++ b/Engine/Plugins/Experimental/Dataflow/Source/DataflowNodes/Public/Dataflow/DataflowCollectionEditSkinWeightsNode.h
@@ -30,7 +30,7 @@ struct FDataflowSkinWeightData
 };
 
 /**
-* FDataflowSkinWeightData is a replacement for FDataflowSkinWeightData in FDataflowCollectionEditSkinWeightsNode
+* FDataflowVertexSkinWeightData is a replacement for FDataflowSkinWeightData in FDataflowCollectionEditSkinWeightsNode
 * This allow storing all the data linearly and enable bulk serialization for speed
 * This solve the problem when we have million vertices meshes that was causing the serialization of the node to take a very long time
 * this is now an order of magnitude faster that using FDataflowSkinWeightData
