Fixed assetion failure in VisualizeTetrahedrons node render callback

diff --git a/Engine/Plugins/Experimental/Dataflow/Source/DataflowEditor/Private/Dataflow/DataflowEngineRendering.cpp b/Engine/Plugins/Experimental/Dataflow/Source/DataflowEditor/Private/Dataflow/DataflowEngineRendering.cpp
--- a/Engine/Plugins/Experimental/Dataflow/Source/DataflowEditor/Private/Dataflow/DataflowEngineRendering.cpp
+++ b/Engine/Plugins/Experimental/Dataflow/Source/DataflowEditor/Private/Dataflow/DataflowEngineRendering.cpp
@@ -1042,15 +1042,17 @@ namespace UE::Dataflow
 				const bool bFoundVertices = Collection.FindAttributeTyped<FVector3f>("Vertex", FGeometryCollection::VerticesGroup) != nullptr;
 				const bool bFoundTransforms = Collection.FindAttributeTyped<FTransform3f>(FTransformCollection::TransformAttribute, FTransformCollection::TransformGroup) != nullptr;
 				const bool bFoundBoneMap = Collection.FindAttributeTyped<int32>("BoneMap", FGeometryCollection::VerticesGroup) != nullptr;
+				const bool bFoundBoneNames = Collection.FindAttributeTyped<FString>("BoneName", FGeometryCollection::TransformGroup) != nullptr;
 				const bool bFoundParents = Collection.FindAttributeTyped<int32>(FTransformCollection::ParentAttribute, FTransformCollection::TransformGroup) != nullptr;
-				bool bFoundRenderData = bFoundIndices && bFoundVertices && bFoundTransforms && bFoundBoneMap && bFoundParents
+				bool bFoundRenderData = bFoundIndices && bFoundVertices && bFoundTransforms && bFoundBoneMap && bFoundBoneNames && bFoundParents
 					&& Collection.NumElements(FTransformCollection::TransformGroup) > 0;
 
 				const bool bFoundVertexStart = Collection.FindAttributeTyped<int32>("VertexStart", FGeometryCollection::GeometryGroup) != nullptr;
 				const bool bFoundVertexCount = Collection.FindAttributeTyped<int32>("VertexCount", FGeometryCollection::GeometryGroup) != nullptr;
-				const bool bFoundFaceStart = Collection.FindAttributeTyped<int32>("FaceStart", FGeometryCollection::GeometryGroup) != nullptr;
-				const bool bFoundFaceCount = Collection.FindAttributeTyped<int32>("FaceCount", FGeometryCollection::GeometryGroup) != nullptr;
-				bool bFoundGeometryAttributes = bFoundVertexStart && bFoundVertexCount && bFoundFaceStart && bFoundFaceCount
+				const bool bFoundTetrahedronStart = Collection.FindAttributeTyped<int32>("TetrahedronStart", FGeometryCollection::GeometryGroup) != nullptr;
+				const bool bFoundTetrahedronCount = Collection.FindAttributeTyped<int32>("TetrahedronCount", FGeometryCollection::GeometryGroup) != nullptr;
+				const bool bFoundTetrahedrons = Collection.FindAttributeTyped<FIntVector4>("Tetrahedron", "Tetrahedral") != nullptr;
+				bool bFoundGeometryAttributes = bFoundVertexStart && bFoundVertexCount && bFoundTetrahedronStart && bFoundTetrahedronCount && bFoundTetrahedrons
 					&& Collection.NumElements(FGeometryCollection::GeometryGroup) > 0;
 
 				if (bFoundRenderData && bFoundGeometryAttributes)
