Dataflow : fix code analysis issue regarding getting values from rendering cache
- the current use of the code will never return a temporary object but if used in different way , this may be possible if the code needs to return a default value
- fixed it by making sure we always return a default const value that persist through the entire scope of the caller

diff --git a/Engine/Plugins/Experimental/Dataflow/Source/DataflowEditor/Public/DataflowRendering/DataflowRenderableTypeInstance.h b/Engine/Plugins/Experimental/Dataflow/Source/DataflowEditor/Public/DataflowRendering/DataflowRenderableTypeInstance.h
--- a/Engine/Plugins/Experimental/Dataflow/Source/DataflowEditor/Public/DataflowRendering/DataflowRenderableTypeInstance.h
+++ b/Engine/Plugins/Experimental/Dataflow/Source/DataflowEditor/Public/DataflowRendering/DataflowRenderableTypeInstance.h
@@ -61,8 +61,9 @@ namespace UE::Dataflow
 		template<typename T>
 		const T& GetCachedValue() const
 		{
+			static const T StaticDefaultValue;
 			const FContextCacheKey Key = ComputeCacheKey();
-			return EvaluationContext.GetData<T>(Key, nullptr);
+			return EvaluationContext.GetData<T>(Key, nullptr, StaticDefaultValue);
 		}
 
 		template<typename T>
