Dataflow : optimize the collection spreadsheet
- get an improvement about 25x for very large collections
- only update the sub-outliners that are needed
- avoid copying the collection multiple time
- parallelize generation of items
- avoid expanding outliner by default

diff --git a/Engine/Plugins/Experimental/Dataflow/Source/DataflowEditor/Private/Dataflow/CollectionSpreadSheetWidget.cpp b/Engine/Plugins/Experimental/Dataflow/Source/DataflowEditor/Private/Dataflow/CollectionSpreadSheetWidget.cpp
--- a/Engine/Plugins/Experimental/Dataflow/Source/DataflowEditor/Private/Dataflow/CollectionSpreadSheetWidget.cpp
+++ b/Engine/Plugins/Experimental/Dataflow/Source/DataflowEditor/Private/Dataflow/CollectionSpreadSheetWidget.cpp
@@ -7,8 +7,8 @@
 #include "Dataflow/DataflowSettings.h"
 #include "GeometryCollection/GeometryCollection.h"
 #include "Dataflow/DataflowCollectionSpreadSheetHelpers.h"
+#include "Dataflow/DataflowEditorUtil.h"
 #include "GeometryCollection/ManagedArrayCollection.h"
-#include "Misc/ScopedSlowTask.h"
 #include "Widgets/Views/STableRow.h"
 #include "Widgets/Text/STextBlock.h"
 #include "Widgets/Input/SComboBox.h"
@@ -134,12 +134,14 @@ void SCollectionSpreadSheet::RegenerateHeader()
 	{
 		Header->ColumnNames.Add(FCollectionSpreadSheetHeader::IndexColumnName);
 
-		for (FName Attr : CollectionInfoMap[SelectedOutput.ToString()].Collection.AttributeNames(SelectedGroup))
+		if (TSharedPtr<const FManagedArrayCollection> CollectionPtr = CollectionInfoMap[SelectedOutput.ToString()].Collection)
 		{
-			Header->ColumnNames.Add(Attr);
-			AttrTypes.Add(UE::Dataflow::CollectionSpreadSheetHelpers::GetArrayTypeString(CollectionInfoMap[SelectedOutput.ToString()].Collection.GetAttributeType(Attr, SelectedGroup)).ToString());
+			for (FName Attr : CollectionPtr->AttributeNames(SelectedGroup))
+			{
+				Header->ColumnNames.Add(Attr);
+				AttrTypes.Add(UE::Dataflow::CollectionSpreadSheetHelpers::GetArrayTypeString(CollectionPtr->GetAttributeType(Attr, SelectedGroup)).ToString());
+			}
 		}
-
 	}
 
 	for (int32 IdxAttr = 0; IdxAttr < Header->ColumnNames.Num(); ++IdxAttr)
@@ -198,27 +200,33 @@ void SCollectionSpreadSheet::RepopulateListView()
 		!SelectedOutput.IsNone() && SelectedOutput.ToString() != "" &&
 		!SelectedGroup.IsNone() && SelectedGroup.ToString() != "")
 	{
-		int32 NumElems = CollectionInfoMap[SelectedOutput.ToString()].Collection.NumElements(SelectedGroup);
-
-		for (int32 IdxElem = 0; IdxElem < NumElems; ++IdxElem)
+		TSharedPtr<const FManagedArrayCollection> CollectionPtr = CollectionInfoMap[SelectedOutput.ToString()].Collection;
+		if (CollectionPtr)
 		{
-			const TSharedPtr<FCollectionSpreadSheetItem> NewItem = MakeShared<FCollectionSpreadSheetItem>();
-			NewItem->Values.SetNum(Header->ColumnNames.Num());
+			const int32 NumElems = CollectionPtr->NumElements(SelectedGroup);
+			ListItems.SetNum(NumElems);
 
-			for (int32 IdxColumn = 0; IdxColumn < Header->ColumnNames.Num(); ++IdxColumn)
-			{
-				const FName& ColumnName = Header->ColumnNames[IdxColumn];
-				if (ColumnName == FCollectionSpreadSheetHeader::IndexColumnName)
+			constexpr int32 BatchSize = 50;
+			ParallelFor(TEXT("ScollectionSpreadSheet_GenerateItems"), NumElems, BatchSize,
+				[this, &CollectionPtr](int32 IdxElem)
 				{
-					NewItem->Values[IdxColumn] = FString::FromInt(IdxElem);
-				}
-				else
-				{
-					NewItem->Values[IdxColumn] = UE::Dataflow::CollectionSpreadSheetHelpers::AttributeValueToString(CollectionInfoMap[SelectedOutput.ToString()].Collection, ColumnName, SelectedGroup, IdxElem);
-				}
-			}
-
-			ListItems.Add(NewItem);
+					const TSharedPtr<FCollectionSpreadSheetItem> NewItem = MakeShared<FCollectionSpreadSheetItem>();
+					NewItem->Values.SetNum(Header->ColumnNames.Num());
+
+					for (int32 IdxColumn = 0; IdxColumn < Header->ColumnNames.Num(); ++IdxColumn)
+					{
+						const FName& ColumnName = Header->ColumnNames[IdxColumn];
+						if (ColumnName == FCollectionSpreadSheetHeader::IndexColumnName)
+						{
+							NewItem->Values[IdxColumn] = FString::FromInt(IdxElem);
+						}
+						else
+						{
+							NewItem->Values[IdxColumn] = UE::Dataflow::CollectionSpreadSheetHelpers::AttributeValueToString(*CollectionPtr, ColumnName, SelectedGroup, IdxElem);
+						}
+					}
+					ListItems[IdxElem] = NewItem;
+				});
 		}
 	}
 
@@ -237,12 +245,66 @@ TSharedRef<ITableRow> SCollectionSpreadSheet::GenerateRow(TSharedPtr<const FColl
 // ----------------------------------------------------------------------------
 //
 
+void SCollectionSpreadSheetWidget::UpdateTransformOutliner()
+{
+	if (CollectionTable)
+	{
+		const bool bIsTransformGroupSelected = CollectionTable->GetSelectedGroup() == TEXT("Transform");
+		if (TransformOutliner && bIsTransformGroupSelected)
+		{
+			const FString OutputStr = CollectionTable->GetSelectedOutputStr();
+			if (CollectionTable->GetCollectionInfoMap().Contains(OutputStr))
+			{
+				TSharedPtr<const FManagedArrayCollection> CollectionPtr = CollectionTable->GetCollectionInfoMap()[OutputStr].Collection;
+				TransformOutliner->SetCollection(CollectionPtr, CollectionTable->GetSelectedOutput());
+				TransformOutliner->RegenerateHeader();
+			}
+		}
+	}
+}
+
+void SCollectionSpreadSheetWidget::UpdateVertexOutliner()
+{
+	if (CollectionTable)
+	{
+		const bool bIsVerticesGroupSelected = CollectionTable->GetSelectedGroup() == TEXT("Vertices");
+		if (VerticesOutliner && bIsVerticesGroupSelected)
+		{
+			const FString OutputStr = CollectionTable->GetSelectedOutputStr();
+			if (CollectionTable->GetCollectionInfoMap().Contains(OutputStr))
+			{
+				TSharedPtr<const FManagedArrayCollection> CollectionPtr = CollectionTable->GetCollectionInfoMap()[OutputStr].Collection;
+				VerticesOutliner->SetCollection(CollectionPtr);
+				VerticesOutliner->RegenerateHeader();
+			}
+		}
+	}
+}
+
+void SCollectionSpreadSheetWidget::UpdateFaceOutliner()
+{
+	if (CollectionTable)
+	{
+		const bool bIsFacesGroupSelected = CollectionTable->GetSelectedGroup() == TEXT("Faces");
+		if (FacesOutliner && bIsFacesGroupSelected)
+		{
+			const FString OutputStr = CollectionTable->GetSelectedOutputStr();
+			if (CollectionTable->GetCollectionInfoMap().Contains(OutputStr))
+			{
+				TSharedPtr<const FManagedArrayCollection> CollectionPtr = CollectionTable->GetCollectionInfoMap()[OutputStr].Collection;
+				FacesOutliner->SetCollection(CollectionPtr);
+				FacesOutliner->RegenerateHeader();
+			}
+		}
+	}
+}
+
 void SCollectionSpreadSheetWidget::NodeOutputsComboBoxSelectionChanged(FName InSelectedOutput, ESelectInfo::Type InSelectInfo)
 {
-	const float AmountOfWork = 3;
-	FScopedSlowTask SlowTask(AmountOfWork, LOCTEXT("SCollectionSpreadSheetWidget_SelectionChangedProgress", "Updating Collection Spreadsheet"));
+	const int32 TotalWork = 400;
+	UE::Dataflow::FScopedProgressNotification ProgressNotification(LOCTEXT("SCollectionSpreadSheetWidget_SelectionChangedProgress", "Updating Collection Spreadsheet"), TotalWork);
 
-	SlowTask.MakeDialog(/*bShowCancelButton*/false);
+	ProgressNotification.SetProgress(0);
 	if (CollectionTable)
 	{
 		if (CollectionTable->GetSelectedOutput() != InSelectedOutput)
@@ -262,53 +324,24 @@ void SCollectionSpreadSheetWidget::NodeOutputsComboBoxSelectionChanged(FName InS
 			}
 		}
 	}
+	ProgressNotification.SetProgress(100);
+	UpdateTransformOutliner();
 
-	SlowTask.EnterProgressFrame();
-	if (TransformOutliner)
-	{
-		const FString OutputStr = InSelectedOutput.ToString();
-
-		if (CollectionTable->GetCollectionInfoMap().Contains(OutputStr))
-		{
-			TSharedPtr<FManagedArrayCollection> CollectionPtr = MakeShared<FManagedArrayCollection>(CollectionTable->GetCollectionInfoMap()[OutputStr].Collection);
-			TransformOutliner->SetCollection(CollectionPtr, InSelectedOutput);
-
-			TransformOutliner->RegenerateHeader();
-		}
-	}
+	ProgressNotification.SetProgress(200);
+	UpdateVertexOutliner();
 
-	SlowTask.EnterProgressFrame();
-	if (VerticesOutliner)
-	{
-		const FString OutputStr = InSelectedOutput.ToString();
+	ProgressNotification.SetProgress(300);
+	UpdateFaceOutliner();
 
-		if (CollectionTable->GetCollectionInfoMap().Contains(OutputStr))
-		{
-			TSharedPtr<FManagedArrayCollection> CollectionPtr = MakeShared<FManagedArrayCollection>(CollectionTable->GetCollectionInfoMap()[OutputStr].Collection);
-			VerticesOutliner->SetCollection(CollectionPtr);
-
-			VerticesOutliner->RegenerateHeader();
-		}
-	}
-
-	SlowTask.EnterProgressFrame();
-	if (FacesOutliner)
-	{
-		const FString OutputStr = InSelectedOutput.ToString();
-
-		if (CollectionTable->GetCollectionInfoMap().Contains(OutputStr))
-		{
-			TSharedPtr<FManagedArrayCollection> CollectionPtr = MakeShared<FManagedArrayCollection>(CollectionTable->GetCollectionInfoMap()[OutputStr].Collection);
-			FacesOutliner->SetCollection(CollectionPtr);
-
-			FacesOutliner->RegenerateHeader();
-		}
-	}
+	SetStatusText();
 }
 
 
 void SCollectionSpreadSheetWidget::CollectionGroupsComboBoxSelectionChanged(FName InSelectedGroup, ESelectInfo::Type InSelectInfo)
 {
+	const int32 TotalWork = 400;
+	UE::Dataflow::FScopedProgressNotification ProgressNotification(LOCTEXT("SCollectionSpreadSheetWidget_SelectionChangedProgress", "Updating Collection Spreadsheet"), TotalWork);
+
 	if (CollectionTable)
 	{
 		if (CollectionTable->GetSelectedGroup() != InSelectedGroup)
@@ -319,38 +352,20 @@ void SCollectionSpreadSheetWidget::CollectionGroupsComboBoxSelectionChanged(FNam
 
 			if (!InSelectedGroup.IsNone())
 			{
-				int32 NumElems = CollectionTable->GetCollectionInfoMap()[CollectionTable->GetSelectedOutput().ToString()].Collection.NumElements(InSelectedGroup);
+				int32 NumElems = CollectionTable->GetCollectionInfoMap()[CollectionTable->GetSelectedOutput().ToString()].Collection->NumElements(InSelectedGroup);
 				CollectionTable->SetNumItems(NumElems);
 			}
 		}
 	}
 
-	if (TransformOutliner)
-	{
-		const FString GroupStr = InSelectedGroup.ToString();
-		if (GroupStr == "Transform")
-		{
-			TransformOutliner->RegenerateHeader();
-		}
-	}
+	ProgressNotification.SetProgress(100);
+	UpdateTransformOutliner();
 
-	if (VerticesOutliner)
-	{
-		const FString GroupStr = InSelectedGroup.ToString();
-		if (GroupStr == "Vertices")
-		{
-			VerticesOutliner->RegenerateHeader();
-		}
-	}
+	ProgressNotification.SetProgress(200);
+	UpdateVertexOutliner();
 
-	if (FacesOutliner)
-	{
-		const FString GroupStr = InSelectedGroup.ToString();
-		if (GroupStr == "Faces")
-		{
-			FacesOutliner->RegenerateHeader();
-		}
-	}
+	ProgressNotification.SetProgress(300);
+	UpdateFaceOutliner();
 
 	SetStatusText();
 }
@@ -705,7 +720,7 @@ void SCollectionSpreadSheetWidget::SetData(const FString& InNodeName)
 
 	if (!NodeName.IsEmpty())
 	{
-		if (CollectionTable->GetCollectionInfoMap().Num() > 0)
+		if (CollectionTable && CollectionTable->GetCollectionInfoMap().Num() > 0)
 		{
 			for (auto& Info : CollectionTable->GetCollectionInfoMap())
 			{
@@ -757,7 +772,7 @@ void SCollectionSpreadSheetWidget::SetStatusText()
 {
 	if (!NodeName.IsEmpty())
 	{
-		FString NumberStr = FString::FormatAsNumber(CollectionTable->GetNumItems());
+		FString NumberStr = FString::FormatAsNumber(CollectionTable? CollectionTable->GetNumItems(): 0);
 		FString Str = "Group has " + NumberStr + " elements";
 		StatusTextBlock->SetText(FText::FromString(Str));
 	}
@@ -774,9 +789,15 @@ void SCollectionSpreadSheetWidget::UpdateCollectionGroups(const FName& InOutputN
 	{
 		CollectionGroups.Empty();
 
-		for (FName Group : GetCollectionTable()->GetCollectionInfoMap()[InOutputName.ToString()].Collection.GroupNames())
+		if (CollectionTable)
 		{
-			CollectionGroups.Add(Group);
+			if (TSharedPtr<const FManagedArrayCollection> CollectionPtr = CollectionTable->GetCollectionInfoMap()[InOutputName.ToString()].Collection)
+			{
+				for (FName Group : CollectionPtr->GroupNames())
+				{
+					CollectionGroups.Add(Group);
+				}
+			}
 		}
 	}
 }
diff --git a/Engine/Plugins/Experimental/Dataflow/Source/DataflowEditor/Private/Dataflow/DataflowCollectionSpreadSheet.cpp b/Engine/Plugins/Experimental/Dataflow/Source/DataflowEditor/Private/Dataflow/DataflowCollectionSpreadSheet.cpp
--- a/Engine/Plugins/Experimental/Dataflow/Source/DataflowEditor/Private/Dataflow/DataflowCollectionSpreadSheet.cpp
+++ b/Engine/Plugins/Experimental/Dataflow/Source/DataflowEditor/Private/Dataflow/DataflowCollectionSpreadSheet.cpp
@@ -25,7 +25,7 @@ void FDataflowCollectionSpreadSheet::SetSupportedOutputTypes()
 
 void FDataflowCollectionSpreadSheet::UpdateViewData()
 {
-	if (CollectionSpreadSheet)
+	if (CollectionSpreadSheet && CollectionSpreadSheet->GetCollectionTable())
 	{
 		CollectionSpreadSheet->GetCollectionTable()->GetCollectionInfoMap().Empty();
 
@@ -42,14 +42,13 @@ void FDataflowCollectionSpreadSheet::UpdateViewData()
 						{
 							for (FDataflowOutput* Output : Outputs)
 							{
-								FName Name = Output->GetName();
-								FName Type = Output->GetType();
-
 								if (Output->GetType() == "FManagedArrayCollection")
 								{
+									const FName Name = Output->GetName();
 									const FManagedArrayCollection DefaultCollection;
 									const FManagedArrayCollection& Value = Output->ReadValue(*Context, DefaultCollection);
-									CollectionSpreadSheet->GetCollectionTable()->GetCollectionInfoMap().Add(Name.ToString(), { Value });
+									TSharedPtr<const FManagedArrayCollection> CopiedCollectionPtr = MakeShared<const FManagedArrayCollection>(Value);
+									CollectionSpreadSheet->GetCollectionTable()->GetCollectionInfoMap().Add(Name.ToString(), { CopiedCollectionPtr });
 								}
 							}
 						}
diff --git a/Engine/Plugins/Experimental/Dataflow/Source/DataflowEditor/Private/Dataflow/DataflowCollectionSpreadSheetHelpers.cpp b/Engine/Plugins/Experimental/Dataflow/Source/DataflowEditor/Private/Dataflow/DataflowCollectionSpreadSheetHelpers.cpp
--- a/Engine/Plugins/Experimental/Dataflow/Source/DataflowEditor/Private/Dataflow/DataflowCollectionSpreadSheetHelpers.cpp
+++ b/Engine/Plugins/Experimental/Dataflow/Source/DataflowEditor/Private/Dataflow/DataflowCollectionSpreadSheetHelpers.cpp
@@ -353,7 +353,7 @@ namespace UE::Dataflow::CollectionSpreadSheetHelpers
 		return DataflowSettings->TransformLevelColors.LevelColors[Depth % ColorPerDepthCount].ToFColor(true);
 	}
 
-	FSlateColor UpdateItemColorFromCollection(const TSharedPtr<FManagedArrayCollection> InCollection, const FName InGroup, const int32 InItemIndex)
+	FSlateColor UpdateItemColorFromCollection(const TSharedPtr<const FManagedArrayCollection> InCollection, const FName InGroup, const int32 InItemIndex)
 	{
 		constexpr FLinearColor InvalidColor(0.1f, 0.1f, 0.1f);
 
@@ -422,7 +422,7 @@ namespace UE::Dataflow::CollectionSpreadSheetHelpers
 		return InvalidColor;
 	}
 
-	TSharedRef<SWidget> MakeColumnWidget(const TSharedPtr<FManagedArrayCollection> InCollection,
+	TSharedRef<SWidget> MakeColumnWidget(const TSharedPtr<const FManagedArrayCollection> InCollection,
 		const FName InGroup,
 		const FName InAttr,
 		const int32 InItemIndex,
diff --git a/Engine/Plugins/Experimental/Dataflow/Source/DataflowEditor/Private/Dataflow/DataflowEditorUtil.cpp b/Engine/Plugins/Experimental/Dataflow/Source/DataflowEditor/Private/Dataflow/DataflowEditorUtil.cpp
--- a/Engine/Plugins/Experimental/Dataflow/Source/DataflowEditor/Private/Dataflow/DataflowEditorUtil.cpp
+++ b/Engine/Plugins/Experimental/Dataflow/Source/DataflowEditor/Private/Dataflow/DataflowEditorUtil.cpp
@@ -14,10 +14,10 @@
 #include "DynamicMesh/MeshNormals.h"
 #include "Engine/SkeletalMesh.h"
 #include "Animation/Skeleton.h"
+#include "Misc/SlowTask.h"
 #include "GeometryCollection/Facades/CollectionRenderingFacade.h"
 #include "Materials/Material.h"
 
-
 using namespace UE::Geometry;
 
 namespace UE::Dataflow::Private
@@ -187,4 +187,47 @@ namespace UE::Dataflow
 		}
 	}
 
+	FScopedProgressNotification::FScopedProgressNotification(const FText& Title, int32 InMaxProgress, float DelayInSeconds)
+		: MaxProgress(InMaxProgress)
+	{
+		SlowTask = MakeUnique<FSlowTask>(static_cast<float>(MaxProgress), Title);
+		SlowTask->Initialize();
+		if (DelayInSeconds != 0)
+		{
+			SlowTask->MakeDialogDelayed(DelayInSeconds);
+		}
+		else
+		{
+			SlowTask->MakeDialog(false);
+		}
+	}
+
+	FScopedProgressNotification::~FScopedProgressNotification()
+	{
+		SlowTask->Destroy();
+	}
+
+	void FScopedProgressNotification::SetProgress(float InProgress, const FText& Message)
+	{
+		const float ClampedProgress = FMath::Clamp(InProgress, 0.f, MaxProgress);
+
+		const float DeltaProgressSinceLastUpdate = FMath::Max(0, (ClampedProgress - LastUpdateProgress));
+		if (UpdateSteps <= 0 || DeltaProgressSinceLastUpdate > UpdateSteps)
+		{
+			const float DeltaProgress = FMath::Max(0, InProgress - CurrentProgress);
+			SlowTask->EnterProgressFrame(DeltaProgress, Message);
+			LastUpdateProgress = ClampedProgress;
+		}
+		CurrentProgress = ClampedProgress;
+	}
+
+	void FScopedProgressNotification::AddProgress(float InDeltaProgress, const FText& Message)
+	{
+		SetProgress((CurrentProgress + InDeltaProgress), Message);
+	}
+
+	void FScopedProgressNotification::SetUpdateSteps(float InUpdateSteps)
+	{
+		UpdateSteps = InUpdateSteps;
+	}
 }
\ No newline at end of file
diff --git a/Engine/Plugins/Experimental/Dataflow/Source/DataflowEditor/Private/Dataflow/DataflowEditorUtil.h b/Engine/Plugins/Experimental/Dataflow/Source/DataflowEditor/Private/Dataflow/DataflowEditorUtil.h
--- a/Engine/Plugins/Experimental/Dataflow/Source/DataflowEditor/Private/Dataflow/DataflowEditorUtil.h
+++ b/Engine/Plugins/Experimental/Dataflow/Source/DataflowEditor/Private/Dataflow/DataflowEditorUtil.h
@@ -7,7 +7,9 @@
 #include "DynamicMesh/DynamicMesh3.h"
 #include "Dataflow/DataflowObjectInterface.h"
 #include "Templates/SharedPointer.h"
+#include "Framework/Notifications/NotificationManager.h"
 
+struct FSlowTask;
 class UDataflow;
 class UObject;
 class UDataflowBaseContent;
@@ -16,6 +18,7 @@ class USkeleton;
 class UAnimationAsset;
 class UMaterial;
 class UDataflowEdNode;
+
 namespace UE::Dataflow { class IDataflowConstructionViewMode; }
 namespace GeometryCollection::Facades
 {
@@ -57,4 +60,24 @@ namespace UE::Dataflow
 	 * Returns true if the resulting collection should be directly rendered of false if the collection will be used to generate a primitive component */
 	bool RenderNodeOutput(GeometryCollection::Facades::FRenderingFacade& RenderingFacade, const UDataflowEdNode& Node, const UDataflowBaseContent& EditorContent, const bool bEvaluateOutputs);
 	void GetViewModesForNode(const UDataflowEdNode& EdNode, const UDataflowBaseContent& EditorContent, TArray<FName>& OutValidViewModeNames);
+
+
+	struct FScopedProgressNotification
+	{
+	public:
+		FScopedProgressNotification(const FText& Title, int32 InMaxProgress = 1.0f, float DelayInSeconds = 0.1f);
+		~FScopedProgressNotification();
+
+		void SetUpdateSteps(float InUpdateSteps);
+
+		void SetProgress(float InProgress, const FText& Message = FText::GetEmpty());
+		void AddProgress(float InDeltaProgress, const FText& Message = FText::GetEmpty());
+
+	private:
+		TUniquePtr<FSlowTask> SlowTask;
+		float CurrentProgress = 0.0f;
+		float LastUpdateProgress = 0.0f;
+		float UpdateSteps = 0.0f;
+		float MaxProgress = 1.0f;
+	};
 }
\ No newline at end of file
diff --git a/Engine/Plugins/Experimental/Dataflow/Source/DataflowEditor/Private/Dataflow/DataflowFacesOutlinerWidget.cpp b/Engine/Plugins/Experimental/Dataflow/Source/DataflowEditor/Private/Dataflow/DataflowFacesOutlinerWidget.cpp
--- a/Engine/Plugins/Experimental/Dataflow/Source/DataflowEditor/Private/Dataflow/DataflowFacesOutlinerWidget.cpp
+++ b/Engine/Plugins/Experimental/Dataflow/Source/DataflowEditor/Private/Dataflow/DataflowFacesOutlinerWidget.cpp
@@ -121,6 +121,9 @@ void FGeometryCollectionFacesTreeItemGeometry::RegenerateChildren()
 
 			RootIndex = FGeometryCollection::Invalid;
 
+			ItemsByFaceIndex.Reserve(FaceCount);
+			ChildItems.Reserve(FaceCount);
+
 			// Add a sub item to the outliner tree for each of the bones/chunks in this GeometryCollection
 			for (int32 Index = 0; Index < FaceCount; ++Index)
 			{
@@ -166,7 +169,7 @@ bool FGeometryCollectionFacesTreeItemGeometry::IsValid() const
 // --- FGeometryCollectionFacesTreeItemGeometry ---
 
 
-TSharedPtr<FManagedArrayCollection> FGeometryCollectionFacesTreeItemFace::GetCollection() const
+TSharedPtr<const FManagedArrayCollection> FGeometryCollectionFacesTreeItemFace::GetCollection() const
 {
 	ensure(ParentCollectionItem);
 	return ParentCollectionItem->GetCollection();
@@ -181,14 +184,14 @@ TSharedRef<ITableRow> FGeometryCollectionFacesTreeItemFace::MakeTreeRowWidget(co
 {
 	UpdateItemColorFromCollection();
 
-	TSharedPtr<FManagedArrayCollection> Collection = GetCollection();
+	TSharedPtr<const FManagedArrayCollection> Collection = GetCollection();
 
 	return SNew(SFacesOutlinerRow, InOwnerTable, SharedThis(this));
 }
 
 bool FGeometryCollectionFacesTreeItemFace::IsValidFace() const
 {
-	TSharedPtr<FManagedArrayCollection> Collection = GetCollection();
+	TSharedPtr<const FManagedArrayCollection> Collection = GetCollection();
 
 	if (ParentCollectionItem && ParentCollectionItem->IsValid())
 	{
@@ -235,7 +238,7 @@ bool FGeometryCollectionFacesTreeItemFace::HasChildren() const
 
 TSharedRef<SWidget> SFacesOutlinerRow::GenerateWidgetForColumn(const FName& ColumnName)
 {
-	TSharedPtr<FManagedArrayCollection> Collection = Item->GetCollection();
+	TSharedPtr<const FManagedArrayCollection> Collection = Item->GetCollection();
 	FName AttrType = UE::Dataflow::CollectionSpreadSheetHelpers::GetArrayTypeString(Collection->GetAttributeType(ColumnName, FGeometryCollection::FacesGroup));
 
 	// This can happen because sometimes slate retains old items until the next tick, and keeps calling callbacks on them until then
@@ -397,11 +400,16 @@ TSharedPtr<SWidget> SFacesOutliner::OnOpenContextMenu()
 void SFacesOutliner::UpdateGeometryCollection()
 {
 	TreeView->RequestTreeRefresh();
-	ExpandAll();
+	//ExpandAll();
 }
 
-void SFacesOutliner::SetCollection(const TSharedPtr<FManagedArrayCollection>& InCollection)
+void SFacesOutliner::SetCollection(const TSharedPtr<const FManagedArrayCollection>& InCollection)
 {
+	if (Collection == InCollection)
+	{
+		return;
+	}
+
 	// Clear the cached Tree ItemSelection without affecting the SelectedBones as 
 	// we want to refresh the tree selection using selected bones
 	TGuardValue<bool> ExternalSelectionGuard(bPerformingSelection, true);
@@ -425,7 +433,7 @@ void SFacesOutliner::SetCollection(const TSharedPtr<FManagedArrayCollection>& In
 	}
 
 	TreeView->RequestTreeRefresh();
-	ExpandAll();
+	//ExpandAll();
 
 	Collection = InCollection;
 }
diff --git a/Engine/Plugins/Experimental/Dataflow/Source/DataflowEditor/Private/Dataflow/DataflowTransformOutlinerWidget.cpp b/Engine/Plugins/Experimental/Dataflow/Source/DataflowEditor/Private/Dataflow/DataflowTransformOutlinerWidget.cpp
--- a/Engine/Plugins/Experimental/Dataflow/Source/DataflowEditor/Private/Dataflow/DataflowTransformOutlinerWidget.cpp
+++ b/Engine/Plugins/Experimental/Dataflow/Source/DataflowEditor/Private/Dataflow/DataflowTransformOutlinerWidget.cpp
@@ -240,7 +240,7 @@ bool FDataflowOutlinerTreeItemCollection::IsValid() const
 // --- FDataflowOutlinerTreeItemBone ---
 
 
-TSharedPtr<FManagedArrayCollection> FDataflowOutlinerTreeItemBone::GetCollection() const
+TSharedPtr<const FManagedArrayCollection> FDataflowOutlinerTreeItemBone::GetCollection() const
 {
 	ensure(ParentCollectionItem);
 	return ParentCollectionItem->GetCollection();
@@ -255,7 +255,7 @@ TSharedRef<ITableRow> FDataflowOutlinerTreeItemBone::MakeTreeRowWidget(const TSh
 {
 	UpdateItemColorFromCollection();
 
-	TSharedPtr<FManagedArrayCollection> Collection = GetCollection();
+	TSharedPtr<const FManagedArrayCollection> Collection = GetCollection();
 
 	// This gets displayed in the pinned rows
 	if (bIsPinned)
@@ -288,7 +288,7 @@ TSharedRef<ITableRow> FDataflowOutlinerTreeItemBone::MakeTreeRowWidget(const TSh
 
 bool FDataflowOutlinerTreeItemBone::IsValidBone() const
 {
-	TSharedPtr<FManagedArrayCollection> Collection = GetCollection();
+	TSharedPtr<const FManagedArrayCollection> Collection = GetCollection();
 
 	if (ParentCollectionItem && ParentCollectionItem->IsValid())
 	{
@@ -334,7 +334,7 @@ bool FDataflowOutlinerTreeItemBone::HasChildren() const
 
 TSharedRef<SWidget> STransformOutlinerRow::GenerateWidgetForColumn(const FName& ColumnName)
 {
-	TSharedPtr<FManagedArrayCollection> Collection = Item->GetCollection();
+	TSharedPtr<const FManagedArrayCollection> Collection = Item->GetCollection();
 	FName AttrType = UE::Dataflow::CollectionSpreadSheetHelpers::GetArrayTypeString(Collection->GetAttributeType(ColumnName, DataCollectionGroup));
 
 	// This can happen because sometimes slate retains old items until the next tick, and keeps calling callbacks on them until then
@@ -518,8 +518,13 @@ void STransformOutliner::UpdateGeometryCollection()
 	ExpandAll();
 }
 
-void STransformOutliner::SetCollection(const TSharedPtr<FManagedArrayCollection>& InCollection, const FName InOutputName)
+void STransformOutliner::SetCollection(const TSharedPtr<const FManagedArrayCollection>& InCollection, const FName InOutputName)
 {
+	if (InCollection == Collection)
+	{
+		return;
+	}
+
 	// Clear the cached Tree ItemSelection without affecting the SelectedBones as 
 	// we want to refresh the tree selection using selected bones
 	TGuardValue<bool> ExternalSelectionGuard(bPerformingSelection, true);
diff --git a/Engine/Plugins/Experimental/Dataflow/Source/DataflowEditor/Private/Dataflow/DataflowVerticesOutlinerWidget.cpp b/Engine/Plugins/Experimental/Dataflow/Source/DataflowEditor/Private/Dataflow/DataflowVerticesOutlinerWidget.cpp
--- a/Engine/Plugins/Experimental/Dataflow/Source/DataflowEditor/Private/Dataflow/DataflowVerticesOutlinerWidget.cpp
+++ b/Engine/Plugins/Experimental/Dataflow/Source/DataflowEditor/Private/Dataflow/DataflowVerticesOutlinerWidget.cpp
@@ -117,6 +117,9 @@ void FGeometryCollectionVerticesTreeItemGeometry::RegenerateChildren()
 
 			RootIndex = FGeometryCollection::Invalid;
 
+			ItemsByVertexIndex.Reserve(VertexCount);
+			ChildItems.Reserve(VertexCount);
+
 			// Add a sub item to the outliner tree for each of the bones/chunks in this GeometryCollection
 			for (int32 Index = 0; Index < VertexCount; ++Index)
 			{
@@ -162,7 +165,7 @@ bool FGeometryCollectionVerticesTreeItemGeometry::IsValid() const
 // --- FGeometryCollectionVerticesTreeItemGeometry ---
 
 
-TSharedPtr<FManagedArrayCollection> FGeometryCollectionVerticesTreeItemVertex::GetCollection() const
+TSharedPtr<const FManagedArrayCollection> FGeometryCollectionVerticesTreeItemVertex::GetCollection() const
 {
 	ensure(ParentCollectionItem);
 	return ParentCollectionItem->GetCollection();
@@ -177,14 +180,14 @@ TSharedRef<ITableRow> FGeometryCollectionVerticesTreeItemVertex::MakeTreeRowWidg
 {
 	UpdateItemColorFromCollection();
 
-	TSharedPtr<FManagedArrayCollection> Collection = GetCollection();
+	TSharedPtr<const FManagedArrayCollection> Collection = GetCollection();
 
 	return SNew(SVerticesOutlinerRow, InOwnerTable, SharedThis(this));
 }
 
 bool FGeometryCollectionVerticesTreeItemVertex::IsValidVertex() const
 {
-	TSharedPtr<FManagedArrayCollection> Collection = GetCollection();
+	TSharedPtr<const FManagedArrayCollection> Collection = GetCollection();
 
 	if (ParentCollectionItem && ParentCollectionItem->IsValid())
 	{
@@ -231,7 +234,7 @@ bool FGeometryCollectionVerticesTreeItemVertex::HasChildren() const
 
 TSharedRef<SWidget> SVerticesOutlinerRow::GenerateWidgetForColumn(const FName& ColumnName)
 {
-	TSharedPtr<FManagedArrayCollection> Collection = Item->GetCollection();
+	TSharedPtr<const FManagedArrayCollection> Collection = Item->GetCollection();
 	FName AttrType = UE::Dataflow::CollectionSpreadSheetHelpers::GetArrayTypeString(Collection->GetAttributeType(ColumnName, FGeometryCollection::VerticesGroup));
 
 	// This can happen because sometimes slate retains old items until the next tick, and keeps calling callbacks on them until then
@@ -393,11 +396,16 @@ TSharedPtr<SWidget> SVerticesOutliner::OnOpenContextMenu()
 void SVerticesOutliner::UpdateGeometryCollection()
 {
 	TreeView->RequestTreeRefresh();
-	ExpandAll();
+	//ExpandAll();
 }
 
-void SVerticesOutliner::SetCollection(const TSharedPtr<FManagedArrayCollection>& InCollection)
+void SVerticesOutliner::SetCollection(const TSharedPtr<const FManagedArrayCollection>& InCollection)
 {
+	if (Collection == InCollection)
+	{
+		return;
+	}
+
 	// Clear the cached Tree ItemSelection without affecting the SelectedBones as 
 	// we want to refresh the tree selection using selected bones
 	TGuardValue<bool> ExternalSelectionGuard(bPerformingSelection, true);
@@ -421,7 +429,7 @@ void SVerticesOutliner::SetCollection(const TSharedPtr<FManagedArrayCollection>&
 	}
 
 	TreeView->RequestTreeRefresh();
-	ExpandAll();
+	//ExpandAll();
 
 	Collection = InCollection;
 }
diff --git a/Engine/Plugins/Experimental/Dataflow/Source/DataflowEditor/Private/DataflowEditorTools/DataflowEditorSkinWeightsPaintTool.cpp b/Engine/Plugins/Experimental/Dataflow/Source/DataflowEditor/Private/DataflowEditorTools/DataflowEditorSkinWeightsPaintTool.cpp
--- a/Engine/Plugins/Experimental/Dataflow/Source/DataflowEditor/Private/DataflowEditorTools/DataflowEditorSkinWeightsPaintTool.cpp
+++ b/Engine/Plugins/Experimental/Dataflow/Source/DataflowEditor/Private/DataflowEditorTools/DataflowEditorSkinWeightsPaintTool.cpp
@@ -7,6 +7,7 @@
 #include "Dataflow/DataflowContent.h"
 #include "Dataflow/DataflowObjectInterface.h"
 #include "Dataflow/DataflowSNode.h"
+#include "Dataflow/DataflowEditorUtil.h"
 #include "Engine/World.h"
 #include "InteractiveToolManager.h"
 #include "ModelingToolTargetUtil.h"
@@ -176,8 +177,8 @@ bool UDataflowEditorSkinWeightsPaintTool::ExtractSkinWeights(TArray<TArray<int32
 
 void UDataflowEditorSkinWeightsPaintTool::Setup()
 {
-	FScopedSlowTask SlowTask(2, LOCTEXT("UDataflowEditorSkinWeightsPaintTool", "Starting Dataflow Skin Weight Paint Tool..."));
-	SlowTask.MakeDialog(/*bShowCancelButton*/false);
+	UE::Dataflow::FScopedProgressNotification ProgressNotification(LOCTEXT("SetupProgressMsg", "Starting Dataflow Skin Weight Paint Tool..."), 100, /* delay */ 0);
+	ProgressNotification.SetProgress(0);
 
 	const int32 VertexOffset = GetVertexOffset();
 	if(VertexOffset != INDEX_NONE)
@@ -189,6 +190,8 @@ void UDataflowEditorSkinWeightsPaintTool::Setup()
 			TArray<TArray<float>> CurrentWeights;
 			TArray<TArray<int32>> CurrentIndices;
 	
+			ProgressNotification.AddProgress(0.25f);
+
 			if(ExtractSkinWeights(CurrentIndices, CurrentWeights))
 			{
 				FSkeletalMeshAttributes MeshAttribs(*MeshDescription);
@@ -198,16 +201,19 @@ void UDataflowEditorSkinWeightsPaintTool::Setup()
 				UE::AnimationCore::FBoneWeightsSettings WeightsSettings;
 				WeightsSettings.SetNormalizeType(UE::AnimationCore::EBoneWeightNormalizeType::Always);
 
-				for(int32 VertexIndex = 0; VertexIndex < NumVertices; ++VertexIndex)
+				ProgressNotification.AddProgress(0.25f);
+				
+				// process vertices weights 
+				for (int32 VertexIndex = 0; VertexIndex < NumVertices; ++VertexIndex)
 				{
 					TArray<UE::AnimationCore::FBoneWeight> BoneWeights;
-					const int32 CurrentIndex = VertexOffset+VertexIndex;
-				
-					if(CurrentIndex < CurrentWeights.Num() && CurrentIndex < CurrentIndices.Num() &&
+					const int32 CurrentIndex = VertexOffset + VertexIndex;
+
+					if (CurrentIndex < CurrentWeights.Num() && CurrentIndex < CurrentIndices.Num() &&
 						CurrentWeights[CurrentIndex].Num() == CurrentIndices[CurrentIndex].Num())
 					{
 						BoneWeights.Reserve(CurrentWeights[CurrentIndex].Num());
-						for(int32 WeightIndex = 0, NumWeights = CurrentWeights[CurrentIndex].Num(); WeightIndex < NumWeights; ++WeightIndex)
+						for (int32 WeightIndex = 0, NumWeights = CurrentWeights[CurrentIndex].Num(); WeightIndex < NumWeights; ++WeightIndex)
 						{
 							BoneWeights.Add(UE::AnimationCore::FBoneWeight(CurrentIndices[CurrentIndex][WeightIndex],
 								CurrentWeights[CurrentIndex][WeightIndex]));
@@ -219,7 +225,7 @@ void UDataflowEditorSkinWeightsPaintTool::Setup()
 		}
 	}
 
-	SlowTask.EnterProgressFrame();
+	ProgressNotification.AddProgress(0.25f);
 	USkinWeightsPaintTool::Setup();
 
 	// enable back face option 
diff --git a/Engine/Plugins/Experimental/Dataflow/Source/DataflowEditor/Public/Dataflow/DataflowCollectionSpreadSheetHelpers.h b/Engine/Plugins/Experimental/Dataflow/Source/DataflowEditor/Public/Dataflow/DataflowCollectionSpreadSheetHelpers.h
--- a/Engine/Plugins/Experimental/Dataflow/Source/DataflowEditor/Public/Dataflow/DataflowCollectionSpreadSheetHelpers.h
+++ b/Engine/Plugins/Experimental/Dataflow/Source/DataflowEditor/Public/Dataflow/DataflowCollectionSpreadSheetHelpers.h
@@ -102,9 +102,9 @@ namespace UE::Dataflow::CollectionSpreadSheetHelpers
 
 	FColor GetColorPerDepth(uint32 Depth);
 
-	FSlateColor UpdateItemColorFromCollection(const TSharedPtr<FManagedArrayCollection> InCollection, const FName InGroup, const int32 InItemIndex);
+	FSlateColor UpdateItemColorFromCollection(const TSharedPtr<const FManagedArrayCollection> InCollection, const FName InGroup, const int32 InItemIndex);
 
-	TSharedRef<SWidget> MakeColumnWidget(const TSharedPtr<FManagedArrayCollection> InCollection,
+	TSharedRef<SWidget> MakeColumnWidget(const TSharedPtr<const FManagedArrayCollection> InCollection,
 		const FName InGroup,
 		const FName InAttr,
 		const int32 InItemIndex,
diff --git a/Engine/Plugins/Experimental/Dataflow/Source/DataflowEditor/Public/Dataflow/DataflowCollectionSpreadSheetWidget.h b/Engine/Plugins/Experimental/Dataflow/Source/DataflowEditor/Public/Dataflow/DataflowCollectionSpreadSheetWidget.h
--- a/Engine/Plugins/Experimental/Dataflow/Source/DataflowEditor/Public/Dataflow/DataflowCollectionSpreadSheetWidget.h
+++ b/Engine/Plugins/Experimental/Dataflow/Source/DataflowEditor/Public/Dataflow/DataflowCollectionSpreadSheetWidget.h
@@ -21,7 +21,7 @@ struct FManagedArrayCollection;
 struct SCollectionInfo
 {
 //	FString OutputType;
-	const FManagedArrayCollection Collection;
+	TSharedPtr<const FManagedArrayCollection> Collection;
 };
 
 
@@ -180,6 +180,10 @@ public:
 	FText GetNoGroupText();
 
 private:
+	void UpdateTransformOutliner();
+	void UpdateVertexOutliner();
+	void UpdateFaceOutliner();
+
 	TSharedPtr<STextBlock> NodeNameTextBlock;
 	TSharedPtr<SComboBox<FName>> NodeOutputsComboBox;
 	TSharedPtr<STextBlock> NodeOutputsComboBoxLabel;
diff --git a/Engine/Plugins/Experimental/Dataflow/Source/DataflowEditor/Public/Dataflow/DataflowFacesOutlinerWidget.h b/Engine/Plugins/Experimental/Dataflow/Source/DataflowEditor/Public/Dataflow/DataflowFacesOutlinerWidget.h
--- a/Engine/Plugins/Experimental/Dataflow/Source/DataflowEditor/Public/Dataflow/DataflowFacesOutlinerWidget.h
+++ b/Engine/Plugins/Experimental/Dataflow/Source/DataflowEditor/Public/Dataflow/DataflowFacesOutlinerWidget.h
@@ -35,7 +35,7 @@ public:
 class FGeometryCollectionFacesTreeItemGeometry : public FGeometryCollectionFacesTreeItem
 {
 public:
-	FGeometryCollectionFacesTreeItemGeometry(TSharedPtr<FManagedArrayCollection> InCollection, int32 InGeoemtryIdx, TSharedPtr<STreeView<FGeometryCollectionFacesTreeItemPtr>> InTreeView)
+	FGeometryCollectionFacesTreeItemGeometry(TSharedPtr<const FManagedArrayCollection> InCollection, int32 InGeoemtryIdx, TSharedPtr<STreeView<FGeometryCollectionFacesTreeItemPtr>> InTreeView)
 		: Collection(InCollection)
 		, GeometryIdx(InGeoemtryIdx)
 		, TreeView(InTreeView)
@@ -58,7 +58,7 @@ public:
 	void RegenerateChildren();
 	void RequestTreeRefresh();
 
-	TSharedPtr<FManagedArrayCollection> GetCollection() { return Collection; }
+	TSharedPtr<const FManagedArrayCollection> GetCollection() { return Collection; }
 
 	bool IsValid() const;
 
@@ -71,7 +71,7 @@ public:
 private:
 	bool FilterBoneIndex(int32 BoneIndex) const;
 
-	TSharedPtr<FManagedArrayCollection> Collection;
+	TSharedPtr<const FManagedArrayCollection> Collection;
 	int32 GeometryIdx;
 
 	TSharedPtr<STreeView<FGeometryCollectionFacesTreeItemPtr>> TreeView;
@@ -104,7 +104,7 @@ public:
 	virtual int32 GetFaceIndex() const override { return FaceIndex; }
 	FSlateColor GetItemColor() const { return ItemColor; }
 	bool HasChildren() const;
-	TSharedPtr<FManagedArrayCollection> GetCollection() const;
+	TSharedPtr<const FManagedArrayCollection> GetCollection() const;
 
 protected:
 
@@ -155,7 +155,7 @@ public:
 	TSharedPtr<SWidget> OnOpenContextMenu();
 
 	void UpdateGeometryCollection();
-	void SetCollection(const TSharedPtr<FManagedArrayCollection>& InCollection);
+	void SetCollection(const TSharedPtr<const FManagedArrayCollection>& InCollection);
 	int32 GetBoneSelectionCount() const;
 
 	void ExpandAll();
@@ -168,6 +168,6 @@ private:
 	TSharedPtr<SHeaderRow> HeaderRowWidget;
 	TArray<TSharedPtr<FGeometryCollectionFacesTreeItemGeometry>> RootNodes;
 	bool bPerformingSelection;
-	TSharedPtr<FManagedArrayCollection> Collection;
+	TSharedPtr<const FManagedArrayCollection> Collection;
 };
 
diff --git a/Engine/Plugins/Experimental/Dataflow/Source/DataflowEditor/Public/Dataflow/DataflowTransformOutlinerWidget.h b/Engine/Plugins/Experimental/Dataflow/Source/DataflowEditor/Public/Dataflow/DataflowTransformOutlinerWidget.h
--- a/Engine/Plugins/Experimental/Dataflow/Source/DataflowEditor/Public/Dataflow/DataflowTransformOutlinerWidget.h
+++ b/Engine/Plugins/Experimental/Dataflow/Source/DataflowEditor/Public/Dataflow/DataflowTransformOutlinerWidget.h
@@ -42,7 +42,7 @@ public:
 class FDataflowOutlinerTreeItemCollection : public FDataflowOutlinerTreeItem
 {
 public:
-	FDataflowOutlinerTreeItemCollection(TSharedPtr<FManagedArrayCollection> InCollection, FName InOutputName, TSharedPtr<STreeView<FDataflowOutlinerTreeItemPtr>> InTreeView)
+	FDataflowOutlinerTreeItemCollection(TSharedPtr<const FManagedArrayCollection> InCollection, FName InOutputName, TSharedPtr<STreeView<FDataflowOutlinerTreeItemPtr>> InTreeView)
 		: Collection(InCollection)
 		, OutputName(InOutputName)
 		, TreeView(InTreeView)
@@ -66,7 +66,7 @@ public:
 	void RequestTreeRefresh();
 	void ExpandByLevel(const int32 InLevel);
 
-	TSharedPtr<FManagedArrayCollection> GetCollection() { return Collection; }
+	TSharedPtr<const FManagedArrayCollection> GetCollection() { return Collection; }
 
 	bool IsValid() const;
 
@@ -79,7 +79,7 @@ public:
 private:
 	bool FilterBoneIndex(int32 BoneIndex) const;
 
-	TSharedPtr<FManagedArrayCollection> Collection;
+	TSharedPtr<const FManagedArrayCollection> Collection;
 	FName OutputName;
 
 	TSharedPtr<STreeView<FDataflowOutlinerTreeItemPtr>> TreeView;
@@ -113,7 +113,7 @@ public:
 	virtual int32 GetBoneIndex() const override { return BoneIndex; }
 	FSlateColor GetItemColor() const { return ItemColor; }
 	bool HasChildren() const;
-	TSharedPtr<FManagedArrayCollection> GetCollection() const;
+	TSharedPtr<const FManagedArrayCollection> GetCollection() const;
 
 protected:
 
@@ -164,7 +164,7 @@ public:
 	TSharedPtr<SWidget> OnOpenContextMenu();
 
 	void UpdateGeometryCollection();
-	void SetCollection(const TSharedPtr<FManagedArrayCollection>& InCollection, const FName InOutputName);
+	void SetCollection(const TSharedPtr<const FManagedArrayCollection>& InCollection, const FName InOutputName);
 	int32 GetBoneSelectionCount() const;
 
 	void ExpandAll();
@@ -184,6 +184,6 @@ private:
 	TSharedPtr<SHeaderRow> HeaderRowWidget;
 	TArray<TSharedPtr<FDataflowOutlinerTreeItemCollection>> RootNodes;
 	bool bPerformingSelection;
-	TSharedPtr<FManagedArrayCollection> Collection;
+	TSharedPtr<const FManagedArrayCollection> Collection;
 };
 
diff --git a/Engine/Plugins/Experimental/Dataflow/Source/DataflowEditor/Public/Dataflow/DataflowVerticesOutlinerWidget.h b/Engine/Plugins/Experimental/Dataflow/Source/DataflowEditor/Public/Dataflow/DataflowVerticesOutlinerWidget.h
--- a/Engine/Plugins/Experimental/Dataflow/Source/DataflowEditor/Public/Dataflow/DataflowVerticesOutlinerWidget.h
+++ b/Engine/Plugins/Experimental/Dataflow/Source/DataflowEditor/Public/Dataflow/DataflowVerticesOutlinerWidget.h
@@ -38,7 +38,7 @@ public:
 class FGeometryCollectionVerticesTreeItemGeometry : public FGeometryCollectionVerticesTreeItem
 {
 public:
-	FGeometryCollectionVerticesTreeItemGeometry(TSharedPtr<FManagedArrayCollection> InCollection, int32 InGeoemtryIdx, TSharedPtr<STreeView<FGeometryCollectionVerticesTreeItemPtr>> InTreeView)
+	FGeometryCollectionVerticesTreeItemGeometry(TSharedPtr<const FManagedArrayCollection> InCollection, int32 InGeoemtryIdx, TSharedPtr<STreeView<FGeometryCollectionVerticesTreeItemPtr>> InTreeView)
 		: Collection(InCollection)
 		, GeometryIdx(InGeoemtryIdx)
 		, TreeView(InTreeView)
@@ -61,7 +61,7 @@ public:
 	void RegenerateChildren();
 	void RequestTreeRefresh();
 
-	TSharedPtr<FManagedArrayCollection> GetCollection() { return Collection; }
+	TSharedPtr<const FManagedArrayCollection> GetCollection() { return Collection; }
 
 	bool IsValid() const;
 
@@ -74,7 +74,7 @@ public:
 private:
 	bool FilterBoneIndex(int32 BoneIndex) const;
 
-	TSharedPtr<FManagedArrayCollection> Collection;
+	TSharedPtr<const FManagedArrayCollection> Collection;
 	int32 GeometryIdx;
 
 	TSharedPtr<STreeView<FGeometryCollectionVerticesTreeItemPtr>> TreeView;
@@ -107,7 +107,7 @@ public:
 	virtual int32 GetVertexIndex() const override { return VertexIndex; }
 	FSlateColor GetItemColor() const { return ItemColor; }
 	bool HasChildren() const;
-	TSharedPtr<FManagedArrayCollection> GetCollection() const;
+	TSharedPtr<const FManagedArrayCollection> GetCollection() const;
 
 protected:
 
@@ -158,7 +158,7 @@ public:
 	TSharedPtr<SWidget> OnOpenContextMenu();
 
 	void UpdateGeometryCollection();
-	void SetCollection(const TSharedPtr<FManagedArrayCollection>& InCollection);
+	void SetCollection(const TSharedPtr<const FManagedArrayCollection>& InCollection);
 	int32 GetBoneSelectionCount() const;
 
 	void ExpandAll();
@@ -171,6 +171,6 @@ private:
 	TSharedPtr<SHeaderRow> HeaderRowWidget;
 	TArray<TSharedPtr<FGeometryCollectionVerticesTreeItemGeometry>> RootNodes;
 	bool bPerformingSelection;
-	TSharedPtr<FManagedArrayCollection> Collection;
+	TSharedPtr<const FManagedArrayCollection> Collection;
 };
 
diff --git a/Engine/Plugins/Experimental/ProceduralVegetationEditor/Source/ProceduralVegetationEditor/Private/PVEditor.cpp b/Engine/Plugins/Experimental/ProceduralVegetationEditor/Source/ProceduralVegetationEditor/Private/PVEditor.cpp
--- a/Engine/Plugins/Experimental/ProceduralVegetationEditor/Source/ProceduralVegetationEditor/Private/PVEditor.cpp
+++ b/Engine/Plugins/Experimental/ProceduralVegetationEditor/Source/ProceduralVegetationEditor/Private/PVEditor.cpp
@@ -491,8 +491,9 @@ void FPVEditor::UpdateCollectionSpreadSheet(const UPCGNode* InSelectedNode) cons
 			FName OutputName;
 			GetCollectionFromNode(InSelectedNode, OutputName, OutCollection);
 			
+			TSharedPtr<const FManagedArrayCollection> CopiedCollectionPtr = MakeShared<const FManagedArrayCollection>(OutCollection);
 			CollectionSpreadSheetWidget->GetCollectionTable()->GetCollectionInfoMap().Empty();
-			CollectionSpreadSheetWidget->GetCollectionTable()->GetCollectionInfoMap().Add(OutputName.ToString(), { OutCollection });
+			CollectionSpreadSheetWidget->GetCollectionTable()->GetCollectionInfoMap().Add(OutputName.ToString(), { CopiedCollectionPtr });
 			CollectionSpreadSheetWidget->SetData(InSelectedNode->GetNodeTitle(EPCGNodeTitleType::FullTitle).ToString());
 		}
 		else
