Dataflow : fix Tetrahedron rendering crashing when the source collection was split in multiple geometry

fyi tim.brakensiek

diff --git a/Engine/Plugins/Experimental/Dataflow/Source/DataflowEditor/Private/Dataflow/DataflowEngineRendering.cpp b/Engine/Plugins/Experimental/Dataflow/Source/DataflowEditor/Private/Dataflow/DataflowEngineRendering.cpp
--- a/Engine/Plugins/Experimental/Dataflow/Source/DataflowEditor/Private/Dataflow/DataflowEngineRendering.cpp
+++ b/Engine/Plugins/Experimental/Dataflow/Source/DataflowEditor/Private/Dataflow/DataflowEngineRendering.cpp
@@ -974,34 +974,39 @@ namespace UE::Dataflow
 		const TManagedArray<int32>& TetrahedronStart = Collection.GetAttribute<int32>("TetrahedronStart", FGeometryCollection::GeometryGroup);
 		const TManagedArray<int32>& TetrahedronCount = Collection.GetAttribute<int32>("TetrahedronCount", FGeometryCollection::GeometryGroup);
 		const TManagedArray<FIntVector4>& Tetrahedrons = Collection.GetAttribute<FIntVector4>("Tetrahedron", "Tetrahedral");
+		const TManagedArray<int32>& GeometryToTransformIndex = Collection.GetAttribute<int32>("TransformIndex", FGeometryCollection::GeometryGroup);
 
-		TArray<FTransform> M;
-		GeometryCollectionAlgo::GlobalMatrices(Transforms, Parents, M);
+		TArray<FTransform> RootSpaceTransforms;
+		GeometryCollectionAlgo::GlobalMatrices(Transforms, Parents, RootSpaceTransforms);
 
 		for (int GeometryIdx = 0; GeometryIdx < Collection.NumElements(FGeometryCollection::GeometryGroup); GeometryIdx++)
 		{
 			TArray<FVector3f> VerticesInCollectionSpace; VerticesInCollectionSpace.AddUninitialized(VertexCount[GeometryIdx]);
 			TArray<FVector3f> SplitVertices;
 
+			const int32 TransformIndex = GeometryToTransformIndex[GeometryIdx];
+
 			// Transform vertices to Collection space
-			for (int32 VertexIdx = VertexStart[GeometryIdx]; VertexIdx < VertexStart[GeometryIdx] + VertexCount[GeometryIdx]; ++VertexIdx)
+			const int32 GlobalVertexOffset = VertexStart[GeometryIdx];
+			for (int32 LocalVertexIdx = 0; LocalVertexIdx < VertexCount[GeometryIdx]; ++LocalVertexIdx)
 			{
-				VerticesInCollectionSpace[VertexIdx] = ToF(M[GeometryIdx].TransformPosition(ToD(Vertex[VertexIdx])));
+				VerticesInCollectionSpace[LocalVertexIdx] = ToF(RootSpaceTransforms[TransformIndex].TransformPosition(ToD(Vertex[GlobalVertexOffset + LocalVertexIdx])));
 			}
 
 			TArray<FIntVector4> Tetras; Tetras.AddUninitialized(TetrahedronCount[GeometryIdx]);
 
-			for (int32 TetrahedronIdx = TetrahedronStart[GeometryIdx]; TetrahedronIdx < TetrahedronStart[GeometryIdx] + TetrahedronCount[GeometryIdx]; ++TetrahedronIdx)
+			const int32 GlobalTetrahedronOffset = TetrahedronStart[GeometryIdx];
+			for (int32 LocalTetrahedronIdx = 0; LocalTetrahedronIdx < TetrahedronCount[GeometryIdx]; ++LocalTetrahedronIdx)
 			{
-				const FIntVector4& Tetra = Tetrahedrons[TetrahedronIdx];
+				const FIntVector4& Tetra = Tetrahedrons[GlobalTetrahedronOffset + LocalTetrahedronIdx];
 				const int32 VtxStart = SplitVertices.Num();
 
-				SplitVertices.Add(VerticesInCollectionSpace[Tetra[0]]);
-				SplitVertices.Add(VerticesInCollectionSpace[Tetra[1]]);
-				SplitVertices.Add(VerticesInCollectionSpace[Tetra[2]]);
-				SplitVertices.Add(VerticesInCollectionSpace[Tetra[3]]);
+				SplitVertices.Add(VerticesInCollectionSpace[Tetra[0] - GlobalVertexOffset]);
+				SplitVertices.Add(VerticesInCollectionSpace[Tetra[1] - GlobalVertexOffset]);
+				SplitVertices.Add(VerticesInCollectionSpace[Tetra[2] - GlobalVertexOffset]);
+				SplitVertices.Add(VerticesInCollectionSpace[Tetra[3] - GlobalVertexOffset]);
 
-				Tetras[TetrahedronIdx] = { VtxStart, VtxStart + 1, VtxStart + 2, VtxStart + 3 };
+				Tetras[LocalTetrahedronIdx] = { VtxStart, VtxStart + 1, VtxStart + 2, VtxStart + 3 };
 			}
 
 			if (SplitVertices.Num() && Tetras.Num())
diff --git a/Engine/Plugins/Experimental/GeometryCollectionPlugin/Source/GeometryCollectionNodes/Private/Dataflow/GeometryCollectionNodes.cpp b/Engine/Plugins/Experimental/GeometryCollectionPlugin/Source/GeometryCollectionNodes/Private/Dataflow/GeometryCollectionNodes.cpp
--- a/Engine/Plugins/Experimental/GeometryCollectionPlugin/Source/GeometryCollectionNodes/Private/Dataflow/GeometryCollectionNodes.cpp
+++ b/Engine/Plugins/Experimental/GeometryCollectionPlugin/Source/GeometryCollectionNodes/Private/Dataflow/GeometryCollectionNodes.cpp
@@ -1866,12 +1866,15 @@ void FVisualizeTetrahedronsDataflowNode::Evaluate(UE::Dataflow::FContext& Contex
 					{
 						const FTransform CollectionSpaceTransform = CollectionSpaceTransforms[TransformIndex];
 						const int32 GeoIndex = TransformToGeometryIndex[TransformIndex];
-						const int32 VertexStart = VertexStartArr[GeoIndex];
-						const int32 VertexCount = VertexCountArr[GeoIndex];
-
-						for (int32 VertexIdx = VertexStart; VertexIdx < VertexStart + VertexCount; ++VertexIdx)
+						if (VertexStartArr.IsValidIndex(GeoIndex))
 						{
-							VerticesInCollectionSpace[VertexIdx] = CollectionSpaceTransform.TransformPosition((FVector)Vertex[VertexIdx]);
+							const int32 VertexStart = VertexStartArr[GeoIndex];
+							const int32 VertexCount = VertexCountArr[GeoIndex];
+
+							for (int32 VertexIdx = VertexStart; VertexIdx < VertexStart + VertexCount; ++VertexIdx)
+							{
+								VerticesInCollectionSpace[VertexIdx] = CollectionSpaceTransform.TransformPosition((FVector)Vertex[VertexIdx]);
+							}
 						}
 					}
 
