Dataflow scene , update rendering slowtask to the dataflow specific scoped progress notification
- fix the component names to use the mesh name

diff --git a/Engine/Plugins/Experimental/Dataflow/Source/DataflowEditor/Private/Dataflow/DataflowConstructionScene.cpp b/Engine/Plugins/Experimental/Dataflow/Source/DataflowEditor/Private/Dataflow/DataflowConstructionScene.cpp
--- a/Engine/Plugins/Experimental/Dataflow/Source/DataflowEditor/Private/Dataflow/DataflowConstructionScene.cpp
+++ b/Engine/Plugins/Experimental/Dataflow/Source/DataflowEditor/Private/Dataflow/DataflowConstructionScene.cpp
@@ -10,6 +10,7 @@
 #include "Dataflow/DataflowObject.h"
 #include "Dataflow/DataflowEditorStyle.h"
 #include "Dataflow/DataflowRenderingFactory.h"
+#include "Dataflow/DataflowEditorUtil.h"
 #include "DataflowRendering/DataflowRenderableComponents.h"
 #include "DataflowRendering/DataflowRenderableTypeInstance.h"
 #include "Drawing/MeshElementsVisualizer.h"
@@ -260,10 +261,8 @@ void FDataflowConstructionScene::ResetSceneComponents()
 
 TObjectPtr<UDynamicMeshComponent>& FDataflowConstructionScene::AddDynamicMeshComponent(FDataflowRenderKey InKey, const FString& MeshName, UE::Geometry::FDynamicMesh3&& DynamicMesh, const TArray<UMaterialInterface*>& MaterialSet, const FTransform& InMeshTransform)
 {
-	// Don't use the MakeUniqueObjectName for the component, we need to keep the name aligned with the collection so selection will work in 
-	// other editors. 
-	// const FName UniqueObjectName = MakeUniqueObjectName(RootSceneActor, UDataflowEditorCollectionComponent::StaticClass(), FName(MeshName));
-	TObjectPtr<UDataflowEditorCollectionComponent> DynamicMeshComponent = NewObject<UDataflowEditorCollectionComponent>(RootSceneActor, FName(MeshName));
+	const FName UniqueObjectName = MakeUniqueObjectName(RootSceneActor, UDataflowEditorCollectionComponent::StaticClass(), FName(MeshName));
+	TObjectPtr<UDataflowEditorCollectionComponent> DynamicMeshComponent = NewObject<UDataflowEditorCollectionComponent>(RootSceneActor, UniqueObjectName);
 
 	DynamicMeshComponent->MeshIndex = InKey.Value;
 	DynamicMeshComponent->Node = InKey.Key;
@@ -573,10 +572,6 @@ void FDataflowConstructionScene::UpdateComponentsFromNode(const UDataflowEdNode*
 		return;
 	}
 
-	const float AmountOfWork = 1;
-	FScopedSlowTask SlowTask(AmountOfWork, LOCTEXT("FDataflowConstructionScene_UpdateCOmponentFromNode_Progress", "Generating construction view components..."));
-	SlowTask.MakeDialog(/*bShowCancelButton*/false);
-
 	using namespace UE::Geometry;//FDynamicMesh3
 
 	const FDataflowEditorToolkit* EditorToolkit = DataflowEditor ? static_cast<FDataflowEditorToolkit*>(DataflowEditor->GetInstanceInterface()) : nullptr;
@@ -587,8 +582,14 @@ void FDataflowConstructionScene::UpdateComponentsFromNode(const UDataflowEdNode*
 
 	if (EditorContent && DataflowContext)
 	{
+		const int32 TotalWork = 4.f;
+		UE::Dataflow::FScopedProgressNotification ProgressNotification(LOCTEXT("FDataflowConstructionScene_UpdateCOmponentFromNode_Progress", "Generating construction view components..."), TotalWork);
+
 		const UE::Dataflow::IDataflowConstructionViewMode* ConstructionViewMode = EditorContent->GetConstructionViewMode();
 
+		// step 1 : create rendering facade
+		ProgressNotification.SetProgress(0);
+
 		TSharedPtr<FManagedArrayCollection> RenderCollection(new FManagedArrayCollection);
 		GeometryCollection::Facades::FRenderingFacade RenderingFacade(*RenderCollection);
 		RenderingFacade.DefineSchema();
@@ -597,6 +598,8 @@ void FDataflowConstructionScene::UpdateComponentsFromNode(const UDataflowEdNode*
 			? UE::Dataflow::CanRenderNodeOutput(*EdNode, *EditorContent, *ConstructionViewMode)
 			: false;
 
+		// step 2 : Render nodes that have explictly specifiy how to be rendered ( Old path )
+		ProgressNotification.AddProgress(1.f);
 		bool bHasRenderCollectionPrimitives = false;
 		if (bCanRenderOutput)
 		{
@@ -607,6 +610,9 @@ void FDataflowConstructionScene::UpdateComponentsFromNode(const UDataflowEdNode*
 			}
 		}
 
+		// step 3 : Collect primitive from nodes that explicly want to generate them 
+		ProgressNotification.AddProgress(1.f);
+
 		// Dataflow primitive nodes node explicily provide its own components 
 		// We do not check for bHasRenderCollectionPrimitives because some nodes actualy set HasRenderCollectionPrimitives to false while still implementing AddPrimitiveComponents
 		// and we need to support that for the time being 
@@ -620,11 +626,16 @@ void FDataflowConstructionScene::UpdateComponentsFromNode(const UDataflowEdNode*
 			}
 		}
 
+		// step 3 : convert render collection results in dynamic meshes or get the components from the new rendering system  
+		ProgressNotification.AddProgress(1.f);
+
 		if (!bHasRenderCollectionPrimitives)
 		{
 			if (bCanRenderOutput) // Path where the node has defined it's way of rendering the output using a MACRO on the node definition itself 
 			{
 				const int32 NumGeometry = RenderingFacade.NumGeometry();
+				const float GeometryProgress = (NumGeometry > 0) ? (1.f / (float)NumGeometry) : 0;
+
 				for (int32 MeshIndex = 0; MeshIndex < NumGeometry; ++MeshIndex)
 				{
 					FDataflowRenderKey DynamicMeshKey{ EdNode, MeshIndex };
@@ -663,6 +674,8 @@ void FDataflowConstructionScene::UpdateComponentsFromNode(const UDataflowEdNode*
 							}
 						}
 					}
+
+					ProgressNotification.AddProgress(GeometryProgress);
 				}
 				UpdateDebugMesh(RenderingFacade);
 			}
@@ -679,12 +692,16 @@ void FDataflowConstructionScene::UpdateComponentsFromNode(const UDataflowEdNode*
 				NodeRenderableTypeInstances.Append(RenderableTypeInstances);
 
 				// extract components from instances
+				const float RenderableTypeProgress = (NodeRenderableTypeInstances.Num() > 0) ? (1.f / (float)NodeRenderableTypeInstances.Num()) : 0;
+
 				UE::Dataflow::FRenderableComponents NodeRenderableComponents(RootSceneActor);
 				TArray<UPrimitiveComponent*> ComponentsFromOutputTypes;
 				for (UE::Dataflow::FRenderableTypeInstance& Instance : NodeRenderableTypeInstances)
 				{
 					// todo : filter by option ? - what option should be enmabled by default ?
 					Instance.GetPrimitiveComponents(NodeRenderableComponents);
+
+					ProgressNotification.AddProgress(RenderableTypeProgress);
 				}
 				PrimitiveComponents.Append(NodeRenderableComponents.GetComponents());
 			}
