Fix localized text identifier duplication

diff --git a/Engine/Plugins/Experimental/Dataflow/Source/DataflowEditor/Private/Dataflow/DataflowSNode.cpp b/Engine/Plugins/Experimental/Dataflow/Source/DataflowEditor/Private/Dataflow/DataflowSNode.cpp
--- a/Engine/Plugins/Experimental/Dataflow/Source/DataflowEditor/Private/Dataflow/DataflowSNode.cpp
+++ b/Engine/Plugins/Experimental/Dataflow/Source/DataflowEditor/Private/Dataflow/DataflowSNode.cpp
@@ -271,7 +271,7 @@ TArray<FOverlayWidgetInfo> SDataflowEdNode::GetOverlayWidgets(bool bSelected, co
 					FText PerfText;
 					if (PerfData.NumCalls > 1)
 					{
-						PerfText = FText::Format(LOCTEXT("DataflowNodePerfDataFormat", "Time: {0} ms - {1} calls"), FText::AsNumber(PerfData.ExclusiveTimeMs, &FmtOptions), FText::AsNumber(PerfData.NumCalls));
+						PerfText = FText::Format(LOCTEXT("DataflowNodePerfDataWithCallFormat", "Time: {0} ms - {1} calls"), FText::AsNumber(PerfData.ExclusiveTimeMs, &FmtOptions), FText::AsNumber(PerfData.NumCalls));
 					}
 					else
 					{
