Dataflow Editor - Fixed crash where the PreviewActor object gets garbage collected before the destruction of the FDataflowSimulationScene instance.

#rnx

diff --git a/Engine/Plugins/Experimental/Dataflow/Source/DataflowEditor/Private/Dataflow/DataflowSimulationScene.cpp b/Engine/Plugins/Experimental/Dataflow/Source/DataflowEditor/Private/Dataflow/DataflowSimulationScene.cpp
--- a/Engine/Plugins/Experimental/Dataflow/Source/DataflowEditor/Private/Dataflow/DataflowSimulationScene.cpp
+++ b/Engine/Plugins/Experimental/Dataflow/Source/DataflowEditor/Private/Dataflow/DataflowSimulationScene.cpp
@@ -461,6 +461,7 @@ void FDataflowSimulationScene::AddReferencedObjects(FReferenceCollector& Collect
 	FDataflowPreviewSceneBase::AddReferencedObjects(Collector);
 
 	Collector.AddReferencedObject(SceneDescription);
+	Collector.AddReferencedObject(PreviewActor);
 }
 
 void FDataflowSimulationScene::SceneDescriptionPropertyChanged(const FName& PropertyName)
