Dataflow Editor: make the tool add node buttons work for tools that don't operate on ManagedArrayCollections.

Update FMeshWrapLandmarksNode to work with this.

diff --git a/Engine/Plugins/Experimental/Dataflow/Source/DataflowEditor/Private/Dataflow/DataflowEditorMode.cpp b/Engine/Plugins/Experimental/Dataflow/Source/DataflowEditor/Private/Dataflow/DataflowEditorMode.cpp
--- a/Engine/Plugins/Experimental/Dataflow/Source/DataflowEditor/Private/Dataflow/DataflowEditorMode.cpp
+++ b/Engine/Plugins/Experimental/Dataflow/Source/DataflowEditor/Private/Dataflow/DataflowEditorMode.cpp
@@ -242,8 +242,10 @@ void UDataflowEditorMode::RegisterDataflowTool(TSharedPtr<FUICommandInfo> UIComm
 
 void UDataflowEditorMode::AddNode(FName NewNodeType)
 {
-	const FName ConnectionType = FManagedArrayCollection::StaticType();
-	const FName ConnectionName("Collection");
+	const UE::Dataflow::FDataflowToolRegistry& ToolRegistry = UE::Dataflow::FDataflowToolRegistry::Get();
+
+	const FName ConnectionType = ToolRegistry.HasToolInfoForNodeType(NewNodeType) ? ToolRegistry.GetAddNodeConnectionType(NewNodeType) : FManagedArrayCollection::StaticType();
+	const FName ConnectionName = ToolRegistry.HasToolInfoForNodeType(NewNodeType) ? ToolRegistry.GetAddNodeConnectionName(NewNodeType) : FName("Collection");
 
 	UEdGraphNode* const CurrentlySelectedNode = GetSingleSelectedNodeWithOutputType(ConnectionType);
 	checkf(CurrentlySelectedNode, TEXT("No node with FManagedArrayCollection output is currently selected in the Dataflow graph"));
@@ -260,7 +262,10 @@ void UDataflowEditorMode::AddNode(FName NewNodeType)
 
 bool UDataflowEditorMode::CanAddNode(FName NewNodeType) const
 {
-	const UEdGraphNode* const CurrentlySelectedNode = GetSingleSelectedNodeWithOutputType(FManagedArrayCollection::StaticType());
+	const UE::Dataflow::FDataflowToolRegistry& ToolRegistry = UE::Dataflow::FDataflowToolRegistry::Get();
+	const FName ConnectionType = ToolRegistry.HasToolInfoForNodeType(NewNodeType) ? ToolRegistry.GetAddNodeConnectionType(NewNodeType) : FManagedArrayCollection::StaticType();
+	
+	const UEdGraphNode* const CurrentlySelectedNode = GetSingleSelectedNodeWithOutputType(ConnectionType);
 	if (CurrentlySelectedNode && GetToolManager())
 	{
 		const bool bToolActive = GetToolManager()->HasActiveTool(EToolSide::Left);
diff --git a/Engine/Plugins/Experimental/Dataflow/Source/DataflowEditor/Private/Dataflow/DataflowToolRegistry.cpp b/Engine/Plugins/Experimental/Dataflow/Source/DataflowEditor/Private/Dataflow/DataflowToolRegistry.cpp
--- a/Engine/Plugins/Experimental/Dataflow/Source/DataflowEditor/Private/Dataflow/DataflowToolRegistry.cpp
+++ b/Engine/Plugins/Experimental/Dataflow/Source/DataflowEditor/Private/Dataflow/DataflowToolRegistry.cpp
@@ -19,30 +19,16 @@ namespace UE::Dataflow
 	}
 
 	void FDataflowToolRegistry::AddNodeToToolMapping(const FName& NodeName, TObjectPtr<UInteractiveToolBuilder> ToolBuilder, const TSharedRef<const IDataflowToolActionCommands>& ToolActionCommands,
-		const FSlateIcon& AddNodeButtonIcon, const FText& AddNodeButtonText, const FName& ToolCategory)
+		const FSlateIcon& AddNodeButtonIcon, const FText& AddNodeButtonText, const FName& ToolCategory, const FName& AddNodeConnectionType, const FName& AddNodeConnectionName)
 	{
 		ensureMsgf(!FDataflowEditorCommands::IsRegistered(), TEXT("FDataflowToolRegistry: DataflowEditorCommands have already been registered. \
 				Newly registered Tools may not be available in the Editor. \
 				Ensure that AddNodeToToolMapping is called before the DataflowEditor module is loaded."));
 
 		// FUICommandInfo is uninitialized, it will be created later in FDataflowEditorCommandsImpl::RegisterCommands
-		NodeTypeToToolMap.Add(NodeName, FToolInfo{ ToolBuilder, ToolActionCommands, AddNodeButtonIcon, AddNodeButtonText, nullptr, nullptr, ToolCategory });
+		NodeTypeToToolMap.Add(NodeName, FToolInfo{ ToolBuilder, ToolActionCommands, AddNodeButtonIcon, AddNodeButtonText, AddNodeConnectionType, AddNodeConnectionName, nullptr, nullptr, ToolCategory });
 	}
 
-	PRAGMA_DISABLE_DEPRECATION_WARNINGS
-	void FDataflowToolRegistry::AddNodeToToolMapping(const FName& NodeName,
-		TObjectPtr<UInteractiveToolBuilder> ToolBuilder,
-		const TSharedRef<const IDataflowToolActionCommands>& ToolActionCommands)
-	{
-		ensureMsgf(!FDataflowEditorCommands::IsRegistered(), TEXT("FDataflowToolRegistry: DataflowEditorCommands have already been registered. \
-				Newly registered Tools may not be available in the Editor. \
-				Ensure that AddNodeToToolMapping is called before the DataflowEditor module is loaded."));
-
-		// FUICommandInfo is uninitialized, it will be created later in FDataflowEditorCommandsImpl::RegisterCommands
-		NodeTypeToToolMap.Add(NodeName, FToolInfo{ ToolBuilder, ToolActionCommands, FSlateIcon(), FText(), nullptr, nullptr, FName("General") });
-	}
-	PRAGMA_ENABLE_DEPRECATION_WARNINGS
-
 	void FDataflowToolRegistry::RemoveNodeToToolMapping(const FName& NodeName)
 	{
 		NodeTypeToToolMap.Remove(NodeName);
@@ -79,12 +65,24 @@ namespace UE::Dataflow
 		return NodeTypeToToolMap[NodeName].ToolCommand;
 	}
 
-	const FName& FDataflowToolRegistry::GetToolCategoryForNode(const FName& NodeName)
+	const FName& FDataflowToolRegistry::GetToolCategoryForNode(const FName& NodeName) const
 	{
 		check(NodeTypeToToolMap.Contains(NodeName));
 		return NodeTypeToToolMap[NodeName].ToolCategory;
 	}
 
+	const FName& FDataflowToolRegistry::GetAddNodeConnectionType(const FName& NodeName) const
+	{
+		check(NodeTypeToToolMap.Contains(NodeName));
+		return NodeTypeToToolMap[NodeName].AddNodeConnectionType;
+	}
+
+	const FName& FDataflowToolRegistry::GetAddNodeConnectionName(const FName& NodeName) const
+	{
+		check(NodeTypeToToolMap.Contains(NodeName));
+		return NodeTypeToToolMap[NodeName].AddNodeConnectionName;
+	}
+
 	UInteractiveToolBuilder* FDataflowToolRegistry::GetToolBuilderForNode(const FName& NodeName)
 	{
 		check(NodeTypeToToolMap.Contains(NodeName));
diff --git a/Engine/Plugins/Experimental/Dataflow/Source/DataflowEditor/Public/Dataflow/DataflowToolRegistry.h b/Engine/Plugins/Experimental/Dataflow/Source/DataflowEditor/Public/Dataflow/DataflowToolRegistry.h
--- a/Engine/Plugins/Experimental/Dataflow/Source/DataflowEditor/Public/Dataflow/DataflowToolRegistry.h
+++ b/Engine/Plugins/Experimental/Dataflow/Source/DataflowEditor/Public/Dataflow/DataflowToolRegistry.h
@@ -38,25 +38,38 @@ namespace UE::Dataflow
 			const TSharedRef<const IDataflowToolActionCommands>& ToolActionCommands, 
 			const FSlateIcon& AddNodeButtonIcon,
 			const FText& AddNodeButtonText, 
-			const FName& ToolCategory = FName("General"));
+			const FName& ToolCategory = FName("General"),
+			const FName& AddNodeConnectionType = FName("FManagedArrayCollection"),
+			const FName& AddNodeConnectionName = FName("Collection")
+		);
 
 		UE_DEPRECATED(5.6, "Please use the version of AddNodeToToolMapping taking AddNodeButtonIcon and AddNodeButtonText")
-		UE_API void AddNodeToToolMapping(const FName& NodeName,
+		void AddNodeToToolMapping(const FName& NodeName,
 			TObjectPtr<UInteractiveToolBuilder> ToolBuilder,
-			const TSharedRef<const IDataflowToolActionCommands>& ToolActionCommands);
+			const TSharedRef<const IDataflowToolActionCommands>& ToolActionCommands)
+		{
+			AddNodeToToolMapping(NodeName, ToolBuilder, ToolActionCommands, FSlateIcon(), FText());
+		}
 
 		UE_API void RemoveNodeToToolMapping(const FName& NodeName);
 
 		UE_API TArray<FName> GetNodeNames() const;
 
+		bool HasToolInfoForNodeType(const FName& NodeType) const
+		{
+			return NodeTypeToToolMap.Contains(NodeType);
+		}
+
 		UE_API TSharedPtr<FUICommandInfo>& GetAddNodeCommandForNode(const FName& NodeType);
 		UE_API const FSlateIcon& GetAddNodeButtonIcon(const FName& NodeType) const;
 		UE_API const FText& GetAddNodeButtonText(const FName& NodeType) const;
+		UE_API const FName& GetAddNodeConnectionType(const FName& NodeType) const;
+		UE_API const FName& GetAddNodeConnectionName(const FName& NodeType) const;
 
 		UE_API TSharedPtr<FUICommandInfo>& GetToolCommandForNode(const FName& NodeName);
 
 		/** Return the category the tool has been registered in */
-		UE_API const FName& GetToolCategoryForNode(const FName& NodeName);
+		UE_API const FName& GetToolCategoryForNode(const FName& NodeName) const;
 
 		UE_API UInteractiveToolBuilder* GetToolBuilderForNode(const FName& NodeName);
 		UE_API const UInteractiveToolBuilder* GetToolBuilderForNode(const FName& NodeName) const;
@@ -73,6 +86,8 @@ namespace UE::Dataflow
 			TSharedRef<const IDataflowToolActionCommands> ToolActionCommands;
 			FSlateIcon AddNodeButtonIcon;
 			FText AddNodeButtonText;
+			FName AddNodeConnectionType;
+			FName AddNodeConnectionName;
 
 			// Constructed automatically in FDataflowEditorCommandsImpl::RegisterCommands
 			TSharedPtr<FUICommandInfo> AddNodeCommand;
diff --git a/Engine/Plugins/Experimental/MeshResizing/Source/MeshResizingEditorTools/Private/MeshResizing/MeshResizingEditorToolsModule.cpp b/Engine/Plugins/Experimental/MeshResizing/Source/MeshResizingEditorTools/Private/MeshResizing/MeshResizingEditorToolsModule.cpp
--- a/Engine/Plugins/Experimental/MeshResizing/Source/MeshResizingEditorTools/Private/MeshResizing/MeshResizingEditorToolsModule.cpp
+++ b/Engine/Plugins/Experimental/MeshResizing/Source/MeshResizingEditorTools/Private/MeshResizing/MeshResizingEditorToolsModule.cpp
@@ -24,7 +24,9 @@ namespace UE::MeshResizing
 				MeshResizingToolActions,                                                     // Tool action bindings
 				FSlateIcon(),                                                                // Button icon 
 				LOCTEXT("AddMeshWrapLandmarksNodeButtonText", "Mesh Resizing Landmarks"),    // Button text
-				FName("General"));                                                           // Tool category
+				FName("General"),                                                            // Tool category
+				FName("TObjectPtr<UDataflowMesh>"),                                          // AddNodeConnectionType
+				FName("Mesh"));                                                              // AddNodeConnectionName
 		}
 
 		virtual void ShutdownModule() override
