Set DataflowGraph package on FStructOnScopes in the DataflowEditor, so USE_STABLE_LOCALIZATION_KEYS works correctly.

diff --git a/Engine/Plugins/Experimental/Dataflow/Source/DataflowEditor/Private/Dataflow/DataflowEditorCommands.cpp b/Engine/Plugins/Experimental/Dataflow/Source/DataflowEditor/Private/Dataflow/DataflowEditorCommands.cpp
--- a/Engine/Plugins/Experimental/Dataflow/Source/DataflowEditor/Private/Dataflow/DataflowEditorCommands.cpp
+++ b/Engine/Plugins/Experimental/Dataflow/Source/DataflowEditor/Private/Dataflow/DataflowEditorCommands.cpp
@@ -554,6 +554,10 @@ void FDataflowEditorCommands::OnSelectedNodesChanged(TSharedPtr<IStructureDetail
 				}
 				if (StructData.Num())
 				{
+					for (TSharedPtr<FStructOnScope>& StructOnScope : StructData)
+					{
+						StructOnScope->SetPackage(Graph->GetPackage());
+					}
 					PropertiesEditor->SetStructureProvider(MakeShared<FStructOnScopeStructureDataProvider>(StructData));
 				}
 			}
