DataflowEditor: Set the new Dataflow, PreviewCacheAsset, and PreviewGeometryCacheAsset on the asset when doing SaveAs.

This also changes the behavior to only prompt to save the cache if it's not using the default cache asset that lives in the Engine folder.

#FailedValidations Preflight

diff --git a/Engine/Plugins/Experimental/Dataflow/Source/DataflowEditor/Private/Dataflow/DataflowEditorToolkit.cpp b/Engine/Plugins/Experimental/Dataflow/Source/DataflowEditor/Private/Dataflow/DataflowEditorToolkit.cpp
--- a/Engine/Plugins/Experimental/Dataflow/Source/DataflowEditor/Private/Dataflow/DataflowEditorToolkit.cpp
+++ b/Engine/Plugins/Experimental/Dataflow/Source/DataflowEditor/Private/Dataflow/DataflowEditorToolkit.cpp
@@ -48,10 +48,13 @@
 #include "Dataflow/DataflowPath.h"
 #include "Dataflow/DataflowSimulationManager.h"
 #include "Dataflow/DataflowSimulationNodes.h"
+#include "Chaos/CacheCollection.h"
 #include "EditorModeManager.h"
 #include "EditorViewportTabContent.h"
 #include "Tools/EdModeInteractiveToolsContext.h"
+#include "FileHelpers.h"
 #include "GameFramework/Actor.h"
+#include "Interfaces/Interface_AssetUserData.h"
 #include "ISkeletonTree.h"
 #include "ISceneOutliner.h"
 #include "IStructureDetailsView.h"
@@ -63,6 +66,7 @@
 #include "WorkflowOrientedApp/WorkflowTabManager.h"
 #include "Framework/Notifications/NotificationManager.h"
 #include "Toolkits/AssetEditorToolkitMenuContext.h"
+#include "UObject/PackageReload.h"
 #include "Widgets/Notifications/SNotificationList.h"
 #include "Widgets/Docking/SDockTab.h"
 #include "Widgets/Layout/SSpacer.h"
@@ -1103,20 +1107,90 @@ void FDataflowEditorToolkit::GetSaveableObjects(TArray<UObject*>& OutObjects) co
 		{
 			check(DataflowAsset->IsAsset());
 			OutObjects.AddUnique(DataflowAsset);
+			if (DataflowAsset->PreviewCacheAsset.IsValid())
+			{
+				OutObjects.AddUnique(DataflowAsset->PreviewCacheAsset.Get());
+			}
+			if (DataflowAsset->PreviewGeometryCacheAsset.IsValid())
+			{
+				OutObjects.AddUnique(DataflowAsset->PreviewGeometryCacheAsset.Get());
+			}
 		}
+	}
+}
 
-		if(SimulationScene && SimulationScene->GetPreviewSceneDescription())
+void FDataflowEditorToolkit::OnAssetsSavedAs(const TArray<UObject*>& SavedObjects)
+{
+	// Find new assets 
+	UDataflow* NewDataflowAsset = nullptr;
+	UObject* NewAssetWithInterface = nullptr;
+	IDataflowInstanceInterface* NewDataflowInterface = nullptr;
+	UChaosCacheCollection* NewChaosCache = nullptr;
+	UGeometryCache* NewGeomCache = nullptr;
+	for (UObject* const SavedObj : SavedObjects)
+	{
+		if (SavedObj->IsA<UDataflow>())
 		{
-			if(TObjectPtr<UChaosCacheCollection> CacheCollection = SimulationScene->GetPreviewSceneDescription()->CacheAsset)
-			{
-				OutObjects.AddUnique(CacheCollection);
-			}
-			if(TObjectPtr<UGeometryCache> GeometryCache = SimulationScene->GetPreviewSceneDescription()->GeometryCacheAsset)
+			NewDataflowAsset = Cast<UDataflow>(SavedObj);
+		}
+		else if (Cast<IDataflowInstanceInterface>(SavedObj))
+		{
+			NewAssetWithInterface = SavedObj;
+			NewDataflowInterface = Cast<IDataflowInstanceInterface>(SavedObj);
+		}
+		else if (SavedObj->IsA<UChaosCacheCollection>())
+		{
+			NewChaosCache = Cast<UChaosCacheCollection>(SavedObj);
+		}
+		else if (SavedObj->IsA<UGeometryCache>())
+		{
+			NewGeomCache = Cast<UGeometryCache>(SavedObj);
+		}
+		else if (IInterface_AssetUserData* AssetUserDataProvider = Cast<IInterface_AssetUserData>(SavedObj))
+		{
+			if (const TArray<UAssetUserData*>* AssetUserDataArray = AssetUserDataProvider->GetAssetUserDataArray())
 			{
-				OutObjects.AddUnique(GeometryCache);
+				for (UAssetUserData* UserData : *AssetUserDataArray)
+				{
+					if (IDataflowInstanceInterface* DataflowInterface = Cast<IDataflowInstanceInterface>(UserData))
+					{
+						NewDataflowInterface = DataflowInterface;
+						NewAssetWithInterface = SavedObj;
+						break;
+					}
+				}
 			}
 		}
 	}
+
+	TArray<UPackage*> PackagesToSave;
+	if (NewDataflowAsset && (NewChaosCache || NewGeomCache))
+	{
+		NewDataflowAsset->PreviewCacheAsset = NewChaosCache;
+		NewDataflowAsset->PreviewGeometryCacheAsset = NewGeomCache;
+
+		// Now save the new Dataflow asset again since we've updated its Property
+		NewDataflowAsset->MarkPackageDirty();
+		PackagesToSave.AddUnique(NewDataflowAsset->GetOutermost());
+	}
+	if (NewDataflowAsset && NewDataflowInterface && NewAssetWithInterface)
+	{
+		NewDataflowInterface->GetDataflowInstance().SetDataflowAsset(NewDataflowAsset);
+
+		// Now save the new IDataflowInstanceInterface asset again since we've updated its Property
+		NewAssetWithInterface->MarkPackageDirty();
+		PackagesToSave.Add(NewAssetWithInterface->GetOutermost());
+	}
+
+	if (!PackagesToSave.IsEmpty())
+	{
+		constexpr bool bPromptToSave = false;
+		FEditorFileUtils::PromptForCheckoutAndSave(PackagesToSave, bCheckDirtyOnAssetSave, bPromptToSave);
+
+		// Reload the package so the editor can update anything that references the Dataflow asset
+		// (Sadly this function is only called after the editor re-launches as part of the Save As process, so the editor is already open at this point.)
+		UPackageTools::ReloadPackages(PackagesToSave);
+	}
 }
 
 //~ End FBaseCharacterFXEditorToolkit overrides
diff --git a/Engine/Plugins/Experimental/Dataflow/Source/DataflowEditor/Public/Dataflow/DataflowEditorToolkit.h b/Engine/Plugins/Experimental/Dataflow/Source/DataflowEditor/Public/Dataflow/DataflowEditorToolkit.h
--- a/Engine/Plugins/Experimental/Dataflow/Source/DataflowEditor/Public/Dataflow/DataflowEditorToolkit.h
+++ b/Engine/Plugins/Experimental/Dataflow/Source/DataflowEditor/Public/Dataflow/DataflowEditorToolkit.h
@@ -203,6 +203,7 @@ private:
 	UE_API virtual bool OnRequestClose(EAssetEditorCloseReason InCloseReason) override;
 	UE_API virtual void PostInitAssetEditor() override;
 	UE_API virtual void GetSaveableObjects(TArray<UObject*>& OutObjects) const override;
+	UE_API virtual void OnAssetsSavedAs(const TArray<UObject*>& SavedObjects) override;
 
 	// FBaseAssetToolkit interface
 	UE_API virtual void CreateWidgets() override;
