Dataflow : address comment issue

diff --git a/Engine/Plugins/Experimental/Dataflow/Source/DataflowEditor/Private/Dataflow/DataflowConstructionScene.cpp b/Engine/Plugins/Experimental/Dataflow/Source/DataflowEditor/Private/Dataflow/DataflowConstructionScene.cpp
--- a/Engine/Plugins/Experimental/Dataflow/Source/DataflowEditor/Private/Dataflow/DataflowConstructionScene.cpp
+++ b/Engine/Plugins/Experimental/Dataflow/Source/DataflowEditor/Private/Dataflow/DataflowConstructionScene.cpp
@@ -626,7 +626,7 @@ void FDataflowConstructionScene::UpdateComponentsFromNode(const UDataflowEdNode*
 			}
 		}
 
-		// step 3 : convert render collection results in dynamic meshes or get the components from the new rendering system  
+		// step 4 : convert render collection results in dynamic meshes or get the components from the new rendering system  
 		ProgressNotification.AddProgress(1.f);
 
 		if (!bHasRenderCollectionPrimitives)
