Dataflow : enable unique component names for renderable types
- also update each surface renerable to use a relevant name

diff --git a/Engine/Plugins/Experimental/Dataflow/Source/DataflowEditor/Private/DataflowRendering/DataflowGeometryCollectionProximityRenderableType.cpp b/Engine/Plugins/Experimental/Dataflow/Source/DataflowEditor/Private/DataflowRendering/DataflowGeometryCollectionProximityRenderableType.cpp
--- a/Engine/Plugins/Experimental/Dataflow/Source/DataflowEditor/Private/DataflowRendering/DataflowGeometryCollectionProximityRenderableType.cpp
+++ b/Engine/Plugins/Experimental/Dataflow/Source/DataflowEditor/Private/DataflowRendering/DataflowGeometryCollectionProximityRenderableType.cpp
@@ -56,8 +56,10 @@ namespace UE::Dataflow::Private
 			if (Collection.HasAttribute(ProximityAttributeName, FGeometryCollection::GeometryGroup) &&
 				Collection.HasAttribute(BoundingBoxAttributeName, FGeometryCollection::GeometryGroup))
 			{
-				UPointSetComponent* PointComponent = OutComponents.AddNewComponent<UPointSetComponent>();
-				ULineSetComponent* LineComponent = OutComponents.AddNewComponent<ULineSetComponent>();
+				static const FName PointComponentName = TEXT("Connectivity_Points");
+				UPointSetComponent* PointComponent = OutComponents.AddNewComponent<UPointSetComponent>(PointComponentName);
+				static const FName LineComponentName = TEXT("Connectivity_Lines");
+				ULineSetComponent* LineComponent = OutComponents.AddNewComponent<ULineSetComponent>(LineComponentName);
 
 				if (PointComponent && LineComponent)
 				{
diff --git a/Engine/Plugins/Experimental/Dataflow/Source/DataflowEditor/Private/DataflowRendering/DataflowGeometryCollectionRenderableType.cpp b/Engine/Plugins/Experimental/Dataflow/Source/DataflowEditor/Private/DataflowRendering/DataflowGeometryCollectionRenderableType.cpp
--- a/Engine/Plugins/Experimental/Dataflow/Source/DataflowEditor/Private/DataflowRendering/DataflowGeometryCollectionRenderableType.cpp
+++ b/Engine/Plugins/Experimental/Dataflow/Source/DataflowEditor/Private/DataflowRendering/DataflowGeometryCollectionRenderableType.cpp
@@ -90,8 +90,9 @@ namespace UE::Dataflow::Private
 				{
 					if (DynamicMeshes[TransformIndex].VertexCount() > 0)
 					{
-						const FName ComponentName = BoneNameAttribute ? FName((*BoneNameAttribute)[TransformIndex]) :
-								FName(TEXT("Bone_Dynamic_Mesh_") + FString::FromInt(TransformIndex));
+						static const FString DefaultBoneString = TEXT("Geometry_Collection_Bone");
+						const FString ComponentNamePostFix = BoneNameAttribute ? (*BoneNameAttribute)[TransformIndex] : DefaultBoneString;
+						const FName ComponentName = FName(FString::Printf(TEXT("[%03d]_%s"), TransformIndex, *ComponentNamePostFix));
 						if (UDynamicMeshComponent* Component = OutComponents.AddNewComponent<UDynamicMeshComponent>(ComponentName))
 						{
 							Component->SetCastShadow(false);
diff --git a/Engine/Plugins/Experimental/Dataflow/Source/DataflowEditor/Private/DataflowRendering/DataflowSkeletalMeshRenderableType.cpp b/Engine/Plugins/Experimental/Dataflow/Source/DataflowEditor/Private/DataflowRendering/DataflowSkeletalMeshRenderableType.cpp
--- a/Engine/Plugins/Experimental/Dataflow/Source/DataflowEditor/Private/DataflowRendering/DataflowSkeletalMeshRenderableType.cpp
+++ b/Engine/Plugins/Experimental/Dataflow/Source/DataflowEditor/Private/DataflowRendering/DataflowSkeletalMeshRenderableType.cpp
@@ -30,7 +30,7 @@ namespace UE::Dataflow::Private
 		{
 			if (const TObjectPtr<USkeletalMesh> SkeletalMesh = GetSkeletalMesh(Instance, nullptr))
 			{
-				if (USkeletalMeshComponent* Component = OutComponents.AddNewComponent<USkeletalMeshComponent>())
+				if (USkeletalMeshComponent* Component = OutComponents.AddNewComponent<USkeletalMeshComponent>(SkeletalMesh->GetFName()))
 				{
 					Component->SetSkeletalMesh(SkeletalMesh);
 				}
diff --git a/Engine/Plugins/Experimental/Dataflow/Source/DataflowEditor/Private/DataflowRendering/DataflowStaticMeshRenderableType.cpp b/Engine/Plugins/Experimental/Dataflow/Source/DataflowEditor/Private/DataflowRendering/DataflowStaticMeshRenderableType.cpp
--- a/Engine/Plugins/Experimental/Dataflow/Source/DataflowEditor/Private/DataflowRendering/DataflowStaticMeshRenderableType.cpp
+++ b/Engine/Plugins/Experimental/Dataflow/Source/DataflowEditor/Private/DataflowRendering/DataflowStaticMeshRenderableType.cpp
@@ -30,7 +30,7 @@ namespace UE::Dataflow::Private
 		{
 			if (const TObjectPtr<UStaticMesh> StaticMesh = GetStaticMesh(Instance, nullptr))
 			{
-				if (UStaticMeshComponent* Component = OutComponents.AddNewComponent<UStaticMeshComponent>())
+				if (UStaticMeshComponent* Component = OutComponents.AddNewComponent<UStaticMeshComponent>(StaticMesh->GetFName()))
 				{
 					Component->SetStaticMesh(StaticMesh);
 				}
diff --git a/Engine/Plugins/Experimental/Dataflow/Source/DataflowEditor/Public/DataflowRendering/DataflowRenderableComponents.h b/Engine/Plugins/Experimental/Dataflow/Source/DataflowEditor/Public/DataflowRendering/DataflowRenderableComponents.h
--- a/Engine/Plugins/Experimental/Dataflow/Source/DataflowEditor/Public/DataflowRendering/DataflowRenderableComponents.h
+++ b/Engine/Plugins/Experimental/Dataflow/Source/DataflowEditor/Public/DataflowRendering/DataflowRenderableComponents.h
@@ -30,11 +30,12 @@ namespace UE::Dataflow
 		template<typename T UE_REQUIRES(TIsDerivedFrom<T, UPrimitiveComponent>::Value)>
 		T* AddNewComponent(FName ComponentName = NAME_None, UObject* ComponentParent = nullptr)
 		{
-			// Todo: Will need to fix the name clashing
-			//const FName ObjectName = MakeUniqueObjectName(ObjectOwner, T::StaticClass(), FName(ComponentName));
+			UObject* ParentObject = ComponentParent ? ComponentParent : RootActor;
+
+			const FName ObjectName = MakeUniqueObjectName(ParentObject, T::StaticClass(), FName(ComponentName));
 			
 			// for now always allocate a new one - in the future we could use a pool 
-			T* NewComponent = ComponentParent ? NewObject<T>(ComponentParent) : NewObject<T>(RootActor);
+			T* NewComponent = NewObject<T>(ParentObject, ObjectName);
 			if (NewComponent)
 			{
 				Components.Add(NewComponent);
