Dataflow : always show the pins on the nodes so that the new rendering system works with it

diff --git a/Engine/Plugins/Experimental/Dataflow/Source/DataflowEditor/Private/Dataflow/DataflowSNode.cpp b/Engine/Plugins/Experimental/Dataflow/Source/DataflowEditor/Private/Dataflow/DataflowSNode.cpp
--- a/Engine/Plugins/Experimental/Dataflow/Source/DataflowEditor/Private/Dataflow/DataflowSNode.cpp
+++ b/Engine/Plugins/Experimental/Dataflow/Source/DataflowEditor/Private/Dataflow/DataflowSNode.cpp
@@ -208,8 +208,7 @@ TSharedRef<SWidget> SDataflowEdNode::CreateTitleRightWidget()
 
 	const bool bHasRenderingCheckBox = (
 		DataflowInterface->NodesHaveToggleWidget() && 
-		DataflowGraphNode->GetDataflowNode() &&
-		(DataflowGraphNode->GetDataflowNode()->GetRenderParameters().Num() || DataflowGraphNode->GetDataflowNode()->CanDebugDraw())
+		DataflowGraphNode->GetDataflowNode()
 		);
 
 	if (bHasRenderingCheckBox)
