Dataflow : make sure Reroute node don't try to render in the construction view
- Also rename variable category for properties to be less confusing

diff --git a/Engine/Plugins/Experimental/Dataflow/Source/DataflowEditor/Private/DataflowRendering/DataflowRenderableTypeInstance.cpp b/Engine/Plugins/Experimental/Dataflow/Source/DataflowEditor/Private/DataflowRendering/DataflowRenderableTypeInstance.cpp
--- a/Engine/Plugins/Experimental/Dataflow/Source/DataflowEditor/Private/DataflowRendering/DataflowRenderableTypeInstance.cpp
+++ b/Engine/Plugins/Experimental/Dataflow/Source/DataflowEditor/Private/DataflowRendering/DataflowRenderableTypeInstance.cpp
@@ -2,6 +2,7 @@
 
 #include "DataflowRendering/DataflowRenderableTypeInstance.h"
 
+#include "Dataflow/DataflowCoreNodes.h"
 #include "DataflowRendering/DataflowRenderableComponents.h"
 #include "DataflowRendering/DataflowRenderableType.h"
 #include "DataflowRendering/DataflowRenderableTypeRegistry.h"
@@ -34,6 +35,12 @@ namespace UE::Dataflow
 	{
 		OutRenderableTypeInstances.Reset();
 
+		// we do not want reroute node to be able to render 
+		if (Node.AsType<FDataflowReRouteNode>())
+		{
+			return;
+		}
+
 		TArray<FDataflowOutput*> Outputs = Node.GetOutputs();
 		for (const FDataflowOutput* Output : Outputs)
 		{
diff --git a/Engine/Source/Runtime/Experimental/Dataflow/Engine/Public/Dataflow/DataflowVariableNodes.h b/Engine/Source/Runtime/Experimental/Dataflow/Engine/Public/Dataflow/DataflowVariableNodes.h
--- a/Engine/Source/Runtime/Experimental/Dataflow/Engine/Public/Dataflow/DataflowVariableNodes.h
+++ b/Engine/Source/Runtime/Experimental/Dataflow/Engine/Public/Dataflow/DataflowVariableNodes.h
@@ -54,10 +54,10 @@ private:
 	void ChangeOutputType(FDataflowOutput* Output, FName NewType);
 
 private:
-	UPROPERTY(EditAnywhere, Category = Variables, meta = (FixedLayout, ShowOnlyInnerProperties))
+	UPROPERTY(EditAnywhere, Category = "Default Value", meta = (FixedLayout, ShowOnlyInnerProperties))
 	FInstancedPropertyBag VariablePropertyBag;
 
-	UPROPERTY(VisibleAnywhere, Category = Variables)
+	UPROPERTY(VisibleAnywhere, Category = Variable )
 	FName VariableName;
 
 	void RegisterHandlers();
