Dataflow : Do not render nodes when selecting multiple of them
- Rational is that multiple selection is amde for organizational purpose ( move part of the graph , delete , copy ... ) and should not block the game thread because of rendering in the construction view
- if one is already rendered this will keep it in the construction view

diff --git a/Engine/Plugins/Experimental/Dataflow/Source/DataflowEditor/Private/Dataflow/DataflowEditorToolkit.cpp b/Engine/Plugins/Experimental/Dataflow/Source/DataflowEditor/Private/Dataflow/DataflowEditorToolkit.cpp
--- a/Engine/Plugins/Experimental/Dataflow/Source/DataflowEditor/Private/Dataflow/DataflowEditorToolkit.cpp
+++ b/Engine/Plugins/Experimental/Dataflow/Source/DataflowEditor/Private/Dataflow/DataflowEditorToolkit.cpp
@@ -1526,9 +1526,9 @@ TSet<TObjectPtr<UDataflowEdNode>> FDataflowEditorToolkit::FilterDataflowEdNodesF
 	return DataflowEdNodes;
 }
 
-const TObjectPtr<UDataflowEdNode> FDataflowEditorToolkit::GetFirstFromSet(const TSet<TObjectPtr<UDataflowEdNode>>& Set)
+const TObjectPtr<UDataflowEdNode> FDataflowEditorToolkit::GetOnlyFromSet(const TSet<TObjectPtr<UDataflowEdNode>>& Set)
 {
-	if (Set.Num() > 0)
+	if (Set.Num() == 1)
 	{
 		return *Set.CreateConstIterator();
 	}
@@ -1551,15 +1551,18 @@ void FDataflowEditorToolkit::OnNodeSelectionChanged(const TSet<UObject*>& InNewS
 		{
 			for (UDataflowEdNode* const Node : Nodes)
 			{
-				Node->SetShouldRenderNode(bValue);
+				if (Node)
+				{
+					Node->SetShouldRenderNode(bValue);
+				}
 			}
 		};
 
-	SetShouldRenderNodeFlag(DeselectedNodes, false);
-	SetShouldRenderNodeFlag(NewlySelectedNodes, true);
+	const TObjectPtr<UDataflowEdNode> PrevSelectedNode = GetOnlyFromSet(PreviouslySelectedNodes);
+	const TObjectPtr<UDataflowEdNode> NewSelectedNode = GetOnlyFromSet(SelectedDataflowNodes);
 
-	const TObjectPtr<UDataflowEdNode> PrevSelectedNode = GetFirstFromSet(PreviouslySelectedNodes);
-	const TObjectPtr<UDataflowEdNode> NewSelectedNode = GetFirstFromSet(SelectedDataflowNodes);
+	SetShouldRenderNodeFlag(DeselectedNodes, false);
+	SetShouldRenderNodeFlag({ NewSelectedNode }, true);
 
 	// todo : move the tool logic somewhere else: we should not be activelly looking for a collection here, we should have some sort of interface to talk to the tools
 	UDataflowEditorMode* const DataflowMode = Cast<UDataflowEditorMode>(EditorModeManager->GetActiveScriptableMode(UDataflowEditorMode::EM_DataflowEditorModeId));
diff --git a/Engine/Plugins/Experimental/Dataflow/Source/DataflowEditor/Public/Dataflow/DataflowEditorToolkit.h b/Engine/Plugins/Experimental/Dataflow/Source/DataflowEditor/Public/Dataflow/DataflowEditorToolkit.h
--- a/Engine/Plugins/Experimental/Dataflow/Source/DataflowEditor/Public/Dataflow/DataflowEditorToolkit.h
+++ b/Engine/Plugins/Experimental/Dataflow/Source/DataflowEditor/Public/Dataflow/DataflowEditorToolkit.h
@@ -299,7 +299,7 @@ private:
 
 	void UpdateViewsFromNode(UDataflowEdNode* Node = nullptr);
 	static TSet<TObjectPtr<UDataflowEdNode>> FilterDataflowEdNodesFromSet(const TSet<UObject*>& Set);
-	static const TObjectPtr<UDataflowEdNode> GetFirstFromSet(const TSet<TObjectPtr<UDataflowEdNode>>& Set);
+	static const TObjectPtr<UDataflowEdNode> GetOnlyFromSet(const TSet<TObjectPtr<UDataflowEdNode>>& Set);
 
 	// List of editor commands used  for the dataflow asset
 	TSharedPtr<FUICommandList> GraphEditorCommands;
