Dataflow : add settings for Dataflow editor
- implement settings for Dataflow simulation scene
- change Flesh and geometry collection Editors to set the simulation settings

diff --git a/Engine/Plugins/Experimental/ChaosFlesh/Source/ChaosFleshEditor/Private/ChaosFlesh/Asset/AssetDefinition_FleshAsset.cpp b/Engine/Plugins/Experimental/ChaosFlesh/Source/ChaosFleshEditor/Private/ChaosFlesh/Asset/AssetDefinition_FleshAsset.cpp
--- a/Engine/Plugins/Experimental/ChaosFlesh/Source/ChaosFleshEditor/Private/ChaosFlesh/Asset/AssetDefinition_FleshAsset.cpp
+++ b/Engine/Plugins/Experimental/ChaosFlesh/Source/ChaosFleshEditor/Private/ChaosFlesh/Asset/AssetDefinition_FleshAsset.cpp
@@ -6,6 +6,7 @@
 #include "Dataflow/DataflowEditor.h"
 #include "Dataflow/DataflowEditorModule.h"
 #include "Dataflow/DataflowEditorToolkit.h"
+#include "Dataflow/DataflowSimulationScene.h"
 #include "Math/Color.h"
 #include "ThumbnailRendering/SceneThumbnailInfo.h"
 #include "Toolkits/SimpleAssetEditor.h"
@@ -73,6 +74,15 @@ EAssetCommandResult UAssetDefinition_FleshAsset::OpenAssets(const FAssetOpenArgs
 			{
 				UAssetEditorSubsystem* const AssetEditorSubsystem = GEditor->GetEditorSubsystem<UAssetEditorSubsystem>();
 				UDataflowEditor* const AssetEditor = NewObject<UDataflowEditor>(AssetEditorSubsystem, NAME_None, RF_Transient);
+				
+				if (UDataflowSimulationSettings* SimulationSettings = NewObject<UDataflowSimulationSettings>())
+				{
+					SimulationSettings->bIsSimulationPlayingByDefault = false;
+					SimulationSettings->bIsAsyncCachingSupported = true;
+					SimulationSettings->bIsAsyncCachingEnabledByDefault = true;
+					AssetEditor->AddEditorSettings(SimulationSettings);
+				}
+				
 				AssetEditor->RegisterToolCategories({"General"});
 
 				const TSubclassOf<AActor> ActorClass = StaticLoadClass(AActor::StaticClass(), nullptr,
diff --git a/Engine/Plugins/Experimental/Dataflow/Source/DataflowEditor/Private/Dataflow/DataflowEditor.cpp b/Engine/Plugins/Experimental/Dataflow/Source/DataflowEditor/Private/Dataflow/DataflowEditor.cpp
--- a/Engine/Plugins/Experimental/Dataflow/Source/DataflowEditor/Private/Dataflow/DataflowEditor.cpp
+++ b/Engine/Plugins/Experimental/Dataflow/Source/DataflowEditor/Private/Dataflow/DataflowEditor.cpp
@@ -25,6 +25,12 @@ TSharedPtr<FBaseAssetToolkit> UDataflowEditor::CreateToolkit()
 	return DataflowToolkit;
 }
 
+
+void UDataflowEditor::AddEditorSettings(const UDataflowEditorSettings* Settings)
+{
+	EditorSettings.Add(Settings);
+}
+
 void UDataflowEditor::Initialize(const TArray<TObjectPtr<UObject>>& InObjects, const TSubclassOf<AActor>& InPreviewClass)
 {
 	if(!InObjects.IsEmpty())
diff --git a/Engine/Plugins/Experimental/Dataflow/Source/DataflowEditor/Private/Dataflow/DataflowSimulationScene.cpp b/Engine/Plugins/Experimental/Dataflow/Source/DataflowEditor/Private/Dataflow/DataflowSimulationScene.cpp
--- a/Engine/Plugins/Experimental/Dataflow/Source/DataflowEditor/Private/Dataflow/DataflowSimulationScene.cpp
+++ b/Engine/Plugins/Experimental/Dataflow/Source/DataflowEditor/Private/Dataflow/DataflowSimulationScene.cpp
@@ -77,6 +77,22 @@ FDataflowSimulationScene::FDataflowSimulationScene(FPreviewScene::ConstructionVa
 #endif
 
 	CreateSimulationScene();
+
+	// apply settings
+	if (const UDataflowSimulationSettings* Settings = InEditor->FindEditorSettings<UDataflowSimulationSettings>())
+	{
+		SetSimulationEnabled(Settings->bIsSimulationPlayingByDefault);
+
+		SceneDescription->CacheParams.bEnableAsyncCaching = Settings->bIsAsyncCachingSupported;
+		if (Settings->bIsAsyncCachingSupported)
+		{
+			SceneDescription->CacheParams.bAsyncCaching = Settings->bIsAsyncCachingEnabledByDefault;
+		}
+		else
+		{
+			SceneDescription->CacheParams.bAsyncCaching = false;
+		}
+	}
 }
 
 void FDataflowSimulationScene::OnObjectsReinstanced(const TMap<UObject*, UObject*>& ObjectsMap)
diff --git a/Engine/Plugins/Experimental/Dataflow/Source/DataflowEditor/Public/Dataflow/DataflowEditor.h b/Engine/Plugins/Experimental/Dataflow/Source/DataflowEditor/Public/Dataflow/DataflowEditor.h
--- a/Engine/Plugins/Experimental/Dataflow/Source/DataflowEditor/Public/Dataflow/DataflowEditor.h
+++ b/Engine/Plugins/Experimental/Dataflow/Source/DataflowEditor/Public/Dataflow/DataflowEditor.h
@@ -21,6 +21,18 @@
 class FDataflowEditorToolkit;
 class UDataflowBaseContent;
 
+/**
+* Base class defining settings for Dataflow
+* Each part of the dataflow editor can define settings and assets can define instance of those settings
+* The interface does not define any method as it is only used as based class for those settings to be discovered and also recognized by the editor
+*/
+
+UCLASS(MinimalAPI)
+class UDataflowEditorSettings : public UObject
+{
+	GENERATED_BODY()
+};
+
 /** 
  * The actual asset editor class doesn't have that much in it, intentionally. 
  * 
@@ -62,12 +74,29 @@ public:
 	TArray<TObjectPtr<UDataflowBaseContent>>& GetTerminalContents() { return TerminalContents; }
 	const TArray<TObjectPtr<UDataflowBaseContent>>& GetTerminalContents() const { return TerminalContents; }
 
-	/** Get the registered tool catagories */
+	/** Get the registered tool categories */
 	const TArray<FName>& GetToolCategories() const { return ToolCategories; }
 
+	/** Get the editor settings */
+	template<typename T>
+	inline const T* FindEditorSettings() const 
+	{ 
+		for (const UDataflowEditorSettings* Settings : EditorSettings)
+		{
+			if (const T* TypedSettings = Cast<T>(Settings))
+			{
+				return TypedSettings;
+			}
+		}
+		return nullptr;
+	}
+
 	/** Registere tool categories available for that construction scene */
 	void RegisterToolCategories(const TArray<FName>& InToolCategrories) { ToolCategories = InToolCategrories; }
 
+	/** set the settings for this editor - this needs to be called before initialize for it to be applied */
+	UE_API void AddEditorSettings(const UDataflowEditorSettings* Settings);
+
 private :
 
 	friend class FDataflowEditorToolkit;
@@ -91,6 +120,10 @@ private :
 	UPROPERTY()
 	TArray<TObjectPtr<UDataflowBaseContent>> TerminalContents;
 
+	/** array of settings for this instance */
+	UPROPERTY()
+	TArray<TObjectPtr<const UDataflowEditorSettings>> EditorSettings;
+
 	/** List of tool categories registered for that editor */
 	TArray<FName> ToolCategories;
 };
diff --git a/Engine/Plugins/Experimental/Dataflow/Source/DataflowEditor/Public/Dataflow/DataflowSimulationScene.h b/Engine/Plugins/Experimental/Dataflow/Source/DataflowEditor/Public/Dataflow/DataflowSimulationScene.h
--- a/Engine/Plugins/Experimental/Dataflow/Source/DataflowEditor/Public/Dataflow/DataflowSimulationScene.h
+++ b/Engine/Plugins/Experimental/Dataflow/Source/DataflowEditor/Public/Dataflow/DataflowSimulationScene.h
@@ -4,6 +4,7 @@
 
 #include "Dataflow/DataflowEditorPreviewSceneBase.h"
 #include "Dataflow/DataflowSimulationGenerator.h"
+#include "Dataflow/DataflowEditor.h"
 #include "Chaos/CacheCollection.h"
 
 #include "DataflowSimulationScene.generated.h"
@@ -19,6 +20,22 @@ class FFleshCollection;
 
 DECLARE_EVENT(UDataflowSimulationSceneDescription, FDataflowSimulationSceneDescriptionChanged)
 
+UCLASS(MinimalAPI)
+class UDataflowSimulationSettings : public UDataflowEditorSettings
+{
+public:
+	GENERATED_BODY()
+
+	UPROPERTY()
+	bool bIsSimulationPlayingByDefault = false;
+
+	UPROPERTY()
+	bool bIsAsyncCachingSupported = false;
+
+	UPROPERTY()
+	bool bIsAsyncCachingEnabledByDefault = false;
+};
+
 UCLASS(MinimalAPI)
 class UDataflowSimulationSceneDescription : public UObject
 {
diff --git a/Engine/Plugins/Experimental/GeometryCollectionPlugin/Source/GeometryCollectionEditor/Private/GeometryCollection/AssetDefinition_GeometryCollection.cpp b/Engine/Plugins/Experimental/GeometryCollectionPlugin/Source/GeometryCollectionEditor/Private/GeometryCollection/AssetDefinition_GeometryCollection.cpp
--- a/Engine/Plugins/Experimental/GeometryCollectionPlugin/Source/GeometryCollectionEditor/Private/GeometryCollection/AssetDefinition_GeometryCollection.cpp
+++ b/Engine/Plugins/Experimental/GeometryCollectionPlugin/Source/GeometryCollectionEditor/Private/GeometryCollection/AssetDefinition_GeometryCollection.cpp
@@ -6,6 +6,7 @@
 #include "Dataflow/DataflowEditor.h"
 #include "Dataflow/DataflowEditorModule.h"
 #include "Dataflow/DataflowEditorToolkit.h"
+#include "Dataflow/DataflowSimulationScene.h"
 #include "Math/Color.h"
 #include "ThumbnailRendering/SceneThumbnailInfo.h"
 #include "Toolkits/SimpleAssetEditor.h"
@@ -96,7 +97,17 @@ EAssetCommandResult UAssetDefinition_GeometryCollection::OpenAssets(const FAsset
 				UAssetEditorSubsystem* const AssetEditorSubsystem = GEditor->GetEditorSubsystem<UAssetEditorSubsystem>();
 				
 				UDataflowEditor* const AssetEditor = NewObject<UDataflowEditor>(AssetEditorSubsystem, NAME_None, RF_Transient);
+
+				if (UDataflowSimulationSettings* SimulationSettings = NewObject<UDataflowSimulationSettings>())
+				{
+					SimulationSettings->bIsSimulationPlayingByDefault = false;
+					SimulationSettings->bIsAsyncCachingSupported = false;
+					SimulationSettings->bIsAsyncCachingEnabledByDefault = false;
+					AssetEditor->AddEditorSettings(SimulationSettings);
+				}
+
 				AssetEditor->RegisterToolCategories({"General"});
+
 				const TSubclassOf<AActor> ActorClass = StaticLoadClass(AActor::StaticClass(), nullptr,
 					TEXT("/GeometryCollectionPlugin/BP_GeometryCollectionPreview.BP_GeometryCollectionPreview_C"), nullptr, LOAD_None, nullptr);
 				AssetEditor->Initialize({ GeometryCollection }, ActorClass);
diff --git a/Engine/Source/Runtime/Experimental/Dataflow/Engine/Public/Dataflow/DataflowPreview.h b/Engine/Source/Runtime/Experimental/Dataflow/Engine/Public/Dataflow/DataflowPreview.h
--- a/Engine/Source/Runtime/Experimental/Dataflow/Engine/Public/Dataflow/DataflowPreview.h
+++ b/Engine/Source/Runtime/Experimental/Dataflow/Engine/Public/Dataflow/DataflowPreview.h
@@ -40,7 +40,16 @@ public :
 	UPROPERTY(EditAnywhere, Category = "Caching", meta = (EditCondition = "bRestartSimulation"))
 	FVector2f RestartTimeRange = FVector2f(0.0f, 0.0f);
 
-	/** Boolean to check if the caching will be done on an async thread (if yes no GT dependency) */
-	UPROPERTY(EditAnywhere, AdvancedDisplay, Category = "Caching")
+	/** Control whether the async caching feature is available or not */
+	UPROPERTY()
+	bool bEnableAsyncCaching = false;
+
+	/** 
+	* When enabled runs the caching proces on an async thread (allowing the UI to remain responsive)
+	* This option will be disabled if the asset associated with the graph does not support asynchronous caching 
+	*/
+	UPROPERTY(EditAnywhere, AdvancedDisplay, Category = "Caching", meta = (EditCondition = "bEnableAsyncCaching", HideEditConditionToggle))
 	bool bAsyncCaching = false;
+
+	
 };
