Dataflow : Fix build issues

diff --git a/Engine/Plugins/Experimental/Dataflow/Source/DataflowEditor/Public/DataflowRendering/DataflowRenderableComponents.h b/Engine/Plugins/Experimental/Dataflow/Source/DataflowEditor/Public/DataflowRendering/DataflowRenderableComponents.h
--- a/Engine/Plugins/Experimental/Dataflow/Source/DataflowEditor/Public/DataflowRendering/DataflowRenderableComponents.h
+++ b/Engine/Plugins/Experimental/Dataflow/Source/DataflowEditor/Public/DataflowRendering/DataflowRenderableComponents.h
@@ -5,9 +5,8 @@
 #include "Containers/Array.h"
 #include "Templates/UnrealTypeTraits.h"
 #include "UObject/UnrealNames.h"
-
-class AActor;
-class UPrimitiveComponent;
+#include "GameFramework/Actor.h"
+#include "Components/PrimitiveComponent.h"
 
 namespace UE::Dataflow
 {
