Dataflow : do not merge edges when generating dynamic meshes from geometry collection
Saves about 40% of the generation time

diff --git a/Engine/Plugins/Experimental/Dataflow/Source/DataflowEditor/Private/DataflowRendering/DataflowGeometryCollectionRenderableType.cpp b/Engine/Plugins/Experimental/Dataflow/Source/DataflowEditor/Private/DataflowRendering/DataflowGeometryCollectionRenderableType.cpp
--- a/Engine/Plugins/Experimental/Dataflow/Source/DataflowEditor/Private/DataflowRendering/DataflowGeometryCollectionRenderableType.cpp
+++ b/Engine/Plugins/Experimental/Dataflow/Source/DataflowEditor/Private/DataflowRendering/DataflowGeometryCollectionRenderableType.cpp
@@ -80,7 +80,8 @@ namespace UE::Dataflow::Private
 								TArrayView<const int32> TransformIndicesToConvert(&TransformIndex, 1);
 
 								FTransform UnusedTransform;
-								::ConvertGeometryCollectionToDynamicMesh(DynamicMesh, UnusedTransform, false, *GeomCollection, true, BoneTransforms, true, TransformIndicesToConvert);
+								constexpr bool bWeldEdges = false; // disabling this as this can take 43% of the conversion time for large meshes
+								::ConvertGeometryCollectionToDynamicMesh(DynamicMesh, UnusedTransform, false, *GeomCollection, bWeldEdges, BoneTransforms, true, TransformIndicesToConvert);
 							});
 					}
 					Instance.CacheValue(DynamicMeshes);
