Dataflow : Add support to generate dynamic mesh instead of skeletal mesh for the edit skin weight tool
- this is a first step to avoid too many converstion between types and avoid the costly dynamic mesh to skeletal mesh conversion
- keep it behind a cvar for the time being because the tool needs to be adapted to support it

diff --git a/Engine/Plugins/Experimental/Dataflow/Source/DataflowNodes/DataflowNodes.Build.cs b/Engine/Plugins/Experimental/Dataflow/Source/DataflowNodes/DataflowNodes.Build.cs
--- a/Engine/Plugins/Experimental/Dataflow/Source/DataflowNodes/DataflowNodes.Build.cs
+++ b/Engine/Plugins/Experimental/Dataflow/Source/DataflowNodes/DataflowNodes.Build.cs
@@ -25,6 +25,13 @@ namespace UnrealBuildTool.Rules
 				}
 			);
 
+			PrivateDependencyModuleNames.AddRange(
+				new string[]
+				{
+					"GeometryFramework"
+				}
+			);
+
 			if (Target.Type == TargetType.Editor)
 			{
 				PrivateDependencyModuleNames.AddRange(
diff --git a/Engine/Plugins/Experimental/Dataflow/Source/DataflowNodes/Private/Dataflow/DataflowCollectionEditSkinWeightsNode.cpp b/Engine/Plugins/Experimental/Dataflow/Source/DataflowNodes/Private/Dataflow/DataflowCollectionEditSkinWeightsNode.cpp
--- a/Engine/Plugins/Experimental/Dataflow/Source/DataflowNodes/Private/Dataflow/DataflowCollectionEditSkinWeightsNode.cpp
+++ b/Engine/Plugins/Experimental/Dataflow/Source/DataflowNodes/Private/Dataflow/DataflowCollectionEditSkinWeightsNode.cpp
@@ -7,6 +7,7 @@
 #include "Dataflow/CollectionRenderingPatternUtility.h"
 #include "DataflowCollectionSkeletalMeshUtils.h"
 #include "DynamicMeshToMeshDescription.h"
+#include "DynamicMesh/DynamicBoneAttribute.h"
 #include "MeshConversionOptions.h"
 #include "SkeletalMeshAttributes.h"
 #include "DynamicMesh/DynamicMesh3.h"
@@ -14,6 +15,7 @@
 #include "Rendering/SkeletalMeshModel.h"
 #include "Engine/SkeletalMesh.h"
 #include "Animation/Skeleton.h"
+#include "Components/DynamicMeshComponent.h"
 #include "Components/SkeletalMeshComponent.h"
 #include "InteractiveToolChange.h"
 
@@ -25,9 +27,13 @@
 
 #define LOCTEXT_NAMESPACE "DataflowCollectionEditSkinWeights"
 
+
 namespace UE::Dataflow::Private
 {
-	
+	bool bEditSkinWeightsNodeGeneratesDynamicMesh = false;
+	FAutoConsoleVariableRef CVarEditSkinWeightsNodeGeneratesDynamicMesh(TEXT("p.Dataflow.EditSkinWeightsNodeGeneratesDynamicMesh"), bEditSkinWeightsNodeGeneratesDynamicMesh, TEXT("When on the edit skin weight node will generate dynamic mesh component instead of the mode expensive skeletal mesh ones"));
+
+
 template<typename ScalarType, typename VectorType, int32 NumComponents>
 bool SetAttributeValues(FManagedArrayCollection& SelectedCollection,
 	const FCollectionAttributeKey& AttributeKey, const TArray<TArray<ScalarType>>& AttributeValues, const ScalarType& DefaultValue, const bool bRenormalizeValues)
@@ -281,6 +287,74 @@ void FDataflowCollectionEditSkinWeightsNode::AddPrimitiveComponents(UE::Dataflow
 		GeometryCollection::Facades::FRenderingFacade RenderingFacade(*RenderCollection);
 		const int32 NumGeometry = RenderingFacade.IsValid() ? RenderingFacade.NumGeometry() : 0;
 		
+		if (UE::Dataflow::Private::bEditSkinWeightsNodeGeneratesDynamicMesh)
+		{
+			const bool bNeedsConstruction = (DynamicMeshes.Num() != NumGeometry);
+			if (bNeedsConstruction)
+			{
+				TObjectPtr<USkeletalMesh> InSkeletalMesh = GetValue(Context, &SkeletalMesh, SkeletalMesh);
+
+				DynamicMeshes.SetNum(NumGeometry);
+				ParallelFor(DynamicMeshes.Num(),
+					[this, &RenderingFacade, &InSkeletalMesh](int32 GeometryIndex)
+					{
+						FDynamicMesh3& DynamicMesh = DynamicMeshes[GeometryIndex];
+						UE::Dataflow::Conversion::RenderingFacadeToDynamicMesh(RenderingFacade, GeometryIndex, DynamicMesh);
+
+						if (InSkeletalMesh)
+						{
+							const FReferenceSkeleton& RefSkeleton = InSkeletalMesh->GetRefSkeleton();
+							const int32 NumBones = RefSkeleton.GetRawBoneNum();
+							if (NumBones)
+							{
+								using namespace UE::Geometry; 
+								if (FDynamicMeshAttributeSet* Attributes = DynamicMesh.Attributes())
+								{
+									Attributes->EnableBones(NumBones);
+									FDynamicMeshBoneNameAttribute* BoneNameAttrib = Attributes->GetBoneNames();
+									FDynamicMeshBoneParentIndexAttribute* BoneParentIndices = Attributes->GetBoneParentIndices();
+									FDynamicMeshBonePoseAttribute* BonePoses = Attributes->GetBonePoses();
+									FDynamicMeshBoneColorAttribute* BoneColors = Attributes->GetBoneColors();
+
+									for (int32 BoneIdx = 0; BoneIdx < NumBones; ++BoneIdx)
+									{
+										BoneNameAttrib->SetValue(BoneIdx, RefSkeleton.GetRawRefBoneInfo()[BoneIdx].Name);
+										BoneParentIndices->SetValue(BoneIdx, RefSkeleton.GetRawRefBoneInfo()[BoneIdx].ParentIndex);
+										BonePoses->SetValue(BoneIdx, RefSkeleton.GetRawRefBonePose()[BoneIdx]);
+										BoneColors->SetValue(BoneIdx, FVector4f::One());
+									}
+								}
+							}
+						}
+					});
+			}
+			if (!DynamicMeshes.IsEmpty())
+			{
+				// copy the meshes because the dynamic component needs to get the ownership and we don't want to loose the cached ones (DynamicMeshes)
+				TArray<FDynamicMesh3> DynamicMeshesToSet;
+				DynamicMeshesToSet.SetNum(NumGeometry);
+				ParallelFor(DynamicMeshes.Num(),
+					[this, &DynamicMeshesToSet](int32 GeometryIndex)
+					{
+						DynamicMeshesToSet[GeometryIndex] = DynamicMeshes[GeometryIndex];
+					});
+
+				PrimitiveComponents.Reserve(PrimitiveComponents.Num() + NumGeometry);
+				UMaterial* DefaultTwoSidedMaterial = Cast<UMaterial>(StaticLoadObject(UMaterial::StaticClass(), NULL, TEXT("/Engine/EditorMaterials/Dataflow/DataflowTwoSidedVertexMaterial")));
+				for (int32 GeometryIndex = 0; GeometryIndex < NumGeometry; ++GeometryIndex)
+				{
+					const FName DynamicMeshName(FString::Printf(TEXT("[%d]_DynamicMesh_(%s)"), GeometryIndex, *FGuid::NewGuid().ToString()));
+					UDynamicMeshComponent* DynamicMeshComponent = NewObject<UDynamicMeshComponent>(RootActor, DynamicMeshName);
+					FDynamicMesh3 DynamicMeshesToMove = DynamicMeshesToSet[GeometryIndex]; // need to copy to not loose the cached ones
+					DynamicMeshComponent->SetMesh(MoveTemp(DynamicMeshesToMove));
+					DynamicMeshComponent->SetOverrideRenderMaterial(DefaultTwoSidedMaterial);
+					PrimitiveComponents.Add(DynamicMeshComponent);
+				}
+			}
+
+			return;
+		}
+
 		const bool bNeedsConstruction = (SkeletalMeshes.Num() != NumGeometry) || !bValidSkeletalMeshes;
 		
 		if(SkeletalMeshes.Num() != NumGeometry)
@@ -494,6 +568,7 @@ FCollectionAttributeKey FDataflowCollectionEditSkinWeightsNode::GetBoneWeightsKe
 void FDataflowCollectionEditSkinWeightsNode::OnInvalidate()
 {
 	bValidSkeletalMeshes = false;
+	DynamicMeshes.Reset();
 }
 
 #if WITH_EDITOR
diff --git a/Engine/Plugins/Experimental/Dataflow/Source/DataflowNodes/Public/Dataflow/DataflowCollectionEditSkinWeightsNode.h b/Engine/Plugins/Experimental/Dataflow/Source/DataflowNodes/Public/Dataflow/DataflowCollectionEditSkinWeightsNode.h
--- a/Engine/Plugins/Experimental/Dataflow/Source/DataflowNodes/Public/Dataflow/DataflowCollectionEditSkinWeightsNode.h
+++ b/Engine/Plugins/Experimental/Dataflow/Source/DataflowNodes/Public/Dataflow/DataflowCollectionEditSkinWeightsNode.h
@@ -8,6 +8,7 @@
 #include "DataflowSkeletalMeshNodes.h"
 #include "Dataflow/DataflowCollectionAttributeKeyNodes.h"
 #include "Dataflow/DataflowNode.h"
+#include "DynamicMesh/DynamicMesh3.h"
 #include "GeometryCollection/GeometryCollection.h"
 #include "GeometryCollection/ManagedArrayCollection.h"
 #include "GPUSkinPublicDefs.h"
@@ -194,6 +195,9 @@ private:
 
 	/** Valid skeletal mesh boolean to trigger the construction */
 	bool bValidSkeletalMeshes = false;
+
+	/** Transient dynamic meshes built from dataflow render collection */
+	TArray<UE::Geometry::FDynamicMesh3> DynamicMeshes;
 };
 
 #undef UE_API
