Dataflow : Add ways to render only a fraction of the curves from a geometry collection if they are present
Cvar is p.Dataflow.Curve.Surface.VertexThreshold and is set to 50000 vertices by default

diff --git a/Engine/Plugins/Experimental/Dataflow/Source/DataflowEditor/Private/DataflowRendering/DataflowGeometryCollectionRenderableType.cpp b/Engine/Plugins/Experimental/Dataflow/Source/DataflowEditor/Private/DataflowRendering/DataflowGeometryCollectionRenderableType.cpp
--- a/Engine/Plugins/Experimental/Dataflow/Source/DataflowEditor/Private/DataflowRendering/DataflowGeometryCollectionRenderableType.cpp
+++ b/Engine/Plugins/Experimental/Dataflow/Source/DataflowEditor/Private/DataflowRendering/DataflowGeometryCollectionRenderableType.cpp
@@ -17,6 +17,8 @@
 #include "GeometryCollection/GeometryCollectionAlgo.h"
 #include "GeometryCollection/ManagedArrayCollection.h"
 
+#include "GeometryCollection/Facades/CollectionCurveFacade.h"
+
 #include "UObject/ObjectPtr.h"
 
 #include "PlanarCut.h" // for ConvertGeometryCollectionToDynamicMesh
@@ -24,6 +26,9 @@
 
 namespace UE::Dataflow::Private
 {
+	int32 CurveSurfaceVertexThreshold = 50000;
+	FAutoConsoleVariableRef CVarCurveSurfaceVertexThreshold(TEXT("p.Dataflow.Curve.Surface.VertexThreshold"), CurveSurfaceVertexThreshold, TEXT("Vertex limit used to reduce curve based rendering in the contruction viewport for performance control"));
+
 	////////////////////////////////////////////////////////////////////////////////////////////////////////////////////////////
 
 	class FGeometryCollectionSurfaceRenderableType : public IRenderableType
@@ -65,8 +70,10 @@ namespace UE::Dataflow::Private
 				else
 				{
 					// TODO(dataflow) : need a version of ConvertGeometryCollectionToDynamicMesh that only takes a managed array collection
-					if (TUniquePtr<const FGeometryCollection> GeomCollection = TUniquePtr<FGeometryCollection>(Collection.NewCopy<FGeometryCollection>()))
+					if (TUniquePtr<FGeometryCollection> GeomCollection = TUniquePtr<FGeometryCollection>(Collection.NewCopy<FGeometryCollection>()))
 					{
+						CullCurvesIfNeeded(*GeomCollection);
+
 						// prepare the dynamic meshes
 						DynamicMeshes.SetNum(TransformAttribute->Num());
 
@@ -104,6 +111,53 @@ namespace UE::Dataflow::Private
 				}
 			}
 		}
+
+		void CullCurvesIfNeeded(FGeometryCollection& GeomCollection) const
+		{
+			// Performance control for curve based geometry that can ghave a lot of geometry
+					// if the threshold is reached, we only render a certain percentage of the curves
+			const int32 NumVertices = GeomCollection.NumElements(FGeometryCollection::VerticesGroup);
+			if (CurveSurfaceVertexThreshold > 0 && NumVertices > 0)
+			{
+				const float CullingRatio = 1.0 - FMath::Clamp((float)CurveSurfaceVertexThreshold / (float)NumVertices, 0, 1);
+
+				GeometryCollection::Facades::FCollectionCurveGeometryFacade CurveFacade(GeomCollection);
+				const bool bHasCurves = CurveFacade.IsValid();
+				if (bHasCurves)
+				{
+					TManagedArray<bool>& FaceVisible = GeomCollection.AddAttribute<bool>(FGeometryCollection::FaceVisibleAttribute, FGeometryCollection::FacesGroup);
+
+					FRandomStream Random(0);
+
+					// assumption that vertex index are Pointindex * 2
+					const TArray<int32>& CurvePointOffsets = CurveFacade.GetCurvePointOffsets();
+					int32 StartFace = 0;
+					const int32 NumCurves = CurvePointOffsets.Num();
+					for (int32 CurveIndex = 0; CurveIndex < NumCurves; ++CurveIndex)
+					{
+						const float CurveRand = Random.GetFraction();
+						const bool bCurveVisible = (CurveRand > CullingRatio);
+
+						const int32 CurvPointIndex = CurvePointOffsets[CurveIndex];
+						const int32 PrevPointIndex = (CurveIndex > 0) ? CurvePointOffsets[CurveIndex - 1] : 0;
+						const int32 CurveLength = FMath::Max(0, CurvPointIndex - PrevPointIndex);
+						if (CurveLength > 0)
+						{
+							const int32 CurveNumFaces = (CurveLength - 1) * 2;
+							for (int32 FaceIndex = 0; FaceIndex < CurveNumFaces; ++FaceIndex)
+							{
+								const int32 GlobalFaceIndex = StartFace + FaceIndex;
+								if (FaceVisible.IsValidIndex(GlobalFaceIndex))
+								{
+									FaceVisible[GlobalFaceIndex] = bCurveVisible;
+								}
+							}
+							StartFace += CurveNumFaces;
+						}
+					}
+				}
+			}
+		}
 	};
 
 	////////////////////////////////////////////////////////////////////////////////////////////////////////////////////////////
