Dataflow : do not invalidate node when changing their override color properties

diff --git a/Engine/Plugins/Experimental/Dataflow/Source/DataflowEditor/Private/Dataflow/DataflowEditorCommands.cpp b/Engine/Plugins/Experimental/Dataflow/Source/DataflowEditor/Private/Dataflow/DataflowEditorCommands.cpp
--- a/Engine/Plugins/Experimental/Dataflow/Source/DataflowEditor/Private/Dataflow/DataflowEditorCommands.cpp
+++ b/Engine/Plugins/Experimental/Dataflow/Source/DataflowEditor/Private/Dataflow/DataflowEditorCommands.cpp
@@ -402,8 +402,10 @@ void FDataflowEditorCommands::OnPropertyValueChanged(UDataflow* OutDataflow, TSh
 			{
 				if (TSharedPtr<FDataflowNode> DataflowNode = SelectedNode->GetDataflowNode())
 				{
+					const FName PropertyName = InPropertyChangedEvent.GetPropertyName();
+
 					// Active state update
-					if (InPropertyChangedEvent.GetPropertyName() == FDataflowNode::GetActivePropertyName())
+					if (PropertyName == FDataflowNode::GetActivePropertyName())
 					{
 						// Reflect the active state on the drawing of the node
 						constexpr bool bCheckIsActiveFlagOnly = true;
@@ -414,10 +416,13 @@ void FDataflowEditorCommands::OnPropertyValueChanged(UDataflow* OutDataflow, TSh
 					}
 					else
 					{
-						// make sure the ditable values are up-to-date with the changes
+						// make sure the editable values are up-to-date with the changes
 						SelectedNode->UpdatePinsDefaultValuesFromNode();
 					}
-					if (DataflowNode->ShouldInvalidateOnPropertyChanged(InPropertyChangedEvent))
+
+					// we do not want to invalidate when changing the color of the node as this is irrelevant to the output of the node
+					const bool bIsOverrideColorBeingChanged = (PropertyName == GET_MEMBER_NAME_CHECKED(FDataflowNode, bOverrideColor)) || (PropertyName == GET_MEMBER_NAME_CHECKED(FDataflowNode, OverrideColor));
+					if (!bIsOverrideColorBeingChanged && DataflowNode->ShouldInvalidateOnPropertyChanged(InPropertyChangedEvent))
 					{
 						// Invalidate the node and reset the editor timestamp
 						DataflowNode->Invalidate();
