Add ShouldReopenEditorForSavedAsset callbacks for Dataflow Editor on SaveAs. This will make the editor not attempt to open the Dataflow Asset unless it can be opened as a standalone. It also won't open the caches.

I must have unshelved an old version of the file without this change when I submitted the previous CL fixing SaveAs.

diff --git a/Engine/Plugins/Experimental/Dataflow/Source/DataflowEditor/Private/Dataflow/AssetDefinition_DataflowAsset.cpp b/Engine/Plugins/Experimental/Dataflow/Source/DataflowEditor/Private/Dataflow/AssetDefinition_DataflowAsset.cpp
--- a/Engine/Plugins/Experimental/Dataflow/Source/DataflowEditor/Private/Dataflow/AssetDefinition_DataflowAsset.cpp
+++ b/Engine/Plugins/Experimental/Dataflow/Source/DataflowEditor/Private/Dataflow/AssetDefinition_DataflowAsset.cpp
@@ -105,7 +105,7 @@ namespace UE::DataflowAssetDefinitionHelpers
 				SMessageDialog::FButton(LOCTEXT("Dataflow_NewText", "Create new Dataflow")),
 				SMessageDialog::FButton(LOCTEXT("Dataflow_OpenText", "Open existing Dataflow")),
 				SMessageDialog::FButton(LOCTEXT("Dataflow_ContinueText", "Continue without Dataflow")),
-					 });
+				});
 
 		const int32 ResultButtonIdx = ConfirmDialog->ShowModal();
 		switch (ResultButtonIdx)
@@ -133,6 +133,12 @@ namespace UE::DataflowAssetDefinitionHelpers
 
 		return DataflowAsset;
 	}
+
+	bool CanOpenDataflowAssetInEditor(const UObject* Asset)
+	{
+		const UDataflow* const DataflowAsset = Cast<UDataflow>(Asset);
+		return DataflowAsset && (bCanEditDataflow || DataflowAsset->Type == EDataflowType::Simulation);
+	}
 }
 
 
@@ -188,7 +194,7 @@ EAssetCommandResult UAssetDefinition_DataflowAsset::OpenAssets(const FAssetOpenA
 
 	if (DataflowObjects.Num() == 1)
 	{
-		if (bCanEditDataflow || DataflowObjects[0]->Type == EDataflowType::Simulation)
+		if (UE::DataflowAssetDefinitionHelpers::CanOpenDataflowAssetInEditor(DataflowObjects[0]))
 		{
 			UAssetEditorSubsystem* const AssetEditorSubsystem = GEditor->GetEditorSubsystem<UAssetEditorSubsystem>();
 			UDataflowEditor* const AssetEditor = NewObject<UDataflowEditor>(AssetEditorSubsystem, NAME_None, RF_Transient);
diff --git a/Engine/Plugins/Experimental/Dataflow/Source/DataflowEditor/Private/Dataflow/DataflowEditorToolkit.cpp b/Engine/Plugins/Experimental/Dataflow/Source/DataflowEditor/Private/Dataflow/DataflowEditorToolkit.cpp
--- a/Engine/Plugins/Experimental/Dataflow/Source/DataflowEditor/Private/Dataflow/DataflowEditorToolkit.cpp
+++ b/Engine/Plugins/Experimental/Dataflow/Source/DataflowEditor/Private/Dataflow/DataflowEditorToolkit.cpp
@@ -8,6 +8,7 @@
 #include "AssetViewerSettings.h"
 #include "DataflowEditorStyle.h"
 #include "DetailCategoryBuilder.h"
+#include "Dataflow/AssetDefinition_DataflowAsset.h"
 #include "Dataflow/DataflowAdvancedPreviewDetailsTab.h"
 #include "Dataflow/DataflowAssetViewerSettingsCustomization.h"
 #include "Dataflow/DataflowDebugDrawComponent.h"
@@ -1193,6 +1194,24 @@ void FDataflowEditorToolkit::OnAssetsSavedAs(const TArray<UObject*>& SavedObject
 	}
 }
 
+bool FDataflowEditorToolkit::ShouldReopenEditorForSavedAsset(const UObject* Asset) const
+{
+	// Do not re-open known dependent assets.
+	if (Asset->IsA<UDataflow>())
+	{
+		return UE::DataflowAssetDefinitionHelpers::CanOpenDataflowAssetInEditor(Asset);
+	}
+	if (Asset->IsA<UChaosCacheCollection>())
+	{
+		return false;
+	}
+	if (Asset->IsA<UGeometryCache>())
+	{
+		return false;
+	}
+	return true;
+}
+
 //~ End FBaseCharacterFXEditorToolkit overrides
 
 class FDataflowPreviewSceneDescriptionCustomization : public IDetailCustomization
diff --git a/Engine/Plugins/Experimental/Dataflow/Source/DataflowEditor/Public/Dataflow/AssetDefinition_DataflowAsset.h b/Engine/Plugins/Experimental/Dataflow/Source/DataflowEditor/Public/Dataflow/AssetDefinition_DataflowAsset.h
--- a/Engine/Plugins/Experimental/Dataflow/Source/DataflowEditor/Public/Dataflow/AssetDefinition_DataflowAsset.h
+++ b/Engine/Plugins/Experimental/Dataflow/Source/DataflowEditor/Public/Dataflow/AssetDefinition_DataflowAsset.h
@@ -8,7 +8,6 @@
 
 #define UE_API DATAFLOWEDITOR_API
 
-
 namespace UE::DataflowAssetDefinitionHelpers
 {
 	// Return true if we should proceed, false if we should re-open the dialog
@@ -22,6 +21,9 @@ namespace UE::DataflowAssetDefinitionHelpers
 
 	// Create a new UDataflow if one doesn't already exist for the Cloth Asset
 	DATAFLOWEDITOR_API UObject* NewOrOpenDataflowAsset(const UObject* Asset);
+
+	// Determine whether or not a UDataflow can be opened directly in the Dataflow Editor
+	DATAFLOWEDITOR_API bool CanOpenDataflowAssetInEditor(const UObject* DataflowAsset);
 }
 
 struct FDataflowInput;
diff --git a/Engine/Plugins/Experimental/Dataflow/Source/DataflowEditor/Public/Dataflow/DataflowEditorToolkit.h b/Engine/Plugins/Experimental/Dataflow/Source/DataflowEditor/Public/Dataflow/DataflowEditorToolkit.h
--- a/Engine/Plugins/Experimental/Dataflow/Source/DataflowEditor/Public/Dataflow/DataflowEditorToolkit.h
+++ b/Engine/Plugins/Experimental/Dataflow/Source/DataflowEditor/Public/Dataflow/DataflowEditorToolkit.h
@@ -206,6 +206,7 @@ private:
 	UE_API virtual void PostInitAssetEditor() override;
 	UE_API virtual void GetSaveableObjects(TArray<UObject*>& OutObjects) const override;
 	UE_API virtual void OnAssetsSavedAs(const TArray<UObject*>& SavedObjects) override;
+	UE_API virtual bool ShouldReopenEditorForSavedAsset(const UObject* Asset) const override;
 
 	// FBaseAssetToolkit interface
 	UE_API virtual void CreateWidgets() override;
