Dataflow : make sure we expand the construction view components by default in the Dataflow scene outliner

diff --git a/Engine/Plugins/Experimental/Dataflow/Source/DataflowEditor/Private/Dataflow/DataflowOutlinerMode.cpp b/Engine/Plugins/Experimental/Dataflow/Source/DataflowEditor/Private/Dataflow/DataflowOutlinerMode.cpp
--- a/Engine/Plugins/Experimental/Dataflow/Source/DataflowEditor/Private/Dataflow/DataflowOutlinerMode.cpp
+++ b/Engine/Plugins/Experimental/Dataflow/Source/DataflowEditor/Private/Dataflow/DataflowOutlinerMode.cpp
@@ -228,8 +228,20 @@ void FDataflowOutlinerMode::OnItemSelectionChanged(FSceneOutlinerTreeItemPtr Ite
 
 void FDataflowOutlinerMode::OnItemAdded(FSceneOutlinerTreeItemPtr Item)
 {
-	SceneOutliner->SetItemExpansion(Item, false);
-	Item->Flags.bIsExpanded = false;
+	const UObject* AddedObject = UE::Dataflow::Private::GetOutlinerItemObject<UObject>(Item);
+	const FName ObjectName = AddedObject ? AddedObject->GetFName(): NAME_None;
+
+	// auto expand the construction view object since this is the one we use the most 
+	bool bShouldExpand = false;
+	if (const TSharedPtr<FDataflowConstructionScene> ConstructionScenePtr = ConstructionScene.Pin())
+	{
+		if (ConstructionScenePtr->GetRootActor() == AddedObject)
+		{
+			bShouldExpand = true;
+		}
+	}
+	SceneOutliner->SetItemExpansion(Item, bShouldExpand);
+	Item->Flags.bIsExpanded = bShouldExpand;
 }
 
 void FDataflowOutlinerMode::OnItemDoubleClick(FSceneOutlinerTreeItemPtr SelectedItem)
