Dataflow : better filtering of the curve decimation code for construction view rendering of geometry collection

diff --git a/Engine/Plugins/Experimental/Dataflow/Source/DataflowEditor/Private/DataflowRendering/DataflowGeometryCollectionRenderableType.cpp b/Engine/Plugins/Experimental/Dataflow/Source/DataflowEditor/Private/DataflowRendering/DataflowGeometryCollectionRenderableType.cpp
--- a/Engine/Plugins/Experimental/Dataflow/Source/DataflowEditor/Private/DataflowRendering/DataflowGeometryCollectionRenderableType.cpp
+++ b/Engine/Plugins/Experimental/Dataflow/Source/DataflowEditor/Private/DataflowRendering/DataflowGeometryCollectionRenderableType.cpp
@@ -117,7 +117,7 @@ namespace UE::Dataflow::Private
 			// Performance control for curve based geometry that can ghave a lot of geometry
 					// if the threshold is reached, we only render a certain percentage of the curves
 			const int32 NumVertices = GeomCollection.NumElements(FGeometryCollection::VerticesGroup);
-			if (CurveSurfaceVertexThreshold > 0 && NumVertices > 0)
+			if (CurveSurfaceVertexThreshold > 0 && NumVertices > CurveSurfaceVertexThreshold)
 			{
 				const float CullingRatio = 1.0 - FMath::Clamp((float)CurveSurfaceVertexThreshold / (float)NumVertices, 0, 1);
 
