Dataflow simulation: fixed bugs
-unable to play back cache after async simulation or after stopping real-time simulation update

diff --git a/Engine/Plugins/Experimental/Dataflow/Source/DataflowEditor/Private/Dataflow/DataflowSimulationScene.cpp b/Engine/Plugins/Experimental/Dataflow/Source/DataflowEditor/Private/Dataflow/DataflowSimulationScene.cpp
--- a/Engine/Plugins/Experimental/Dataflow/Source/DataflowEditor/Private/Dataflow/DataflowSimulationScene.cpp
+++ b/Engine/Plugins/Experimental/Dataflow/Source/DataflowEditor/Private/Dataflow/DataflowSimulationScene.cpp
@@ -230,7 +230,7 @@ void FDataflowSimulationScene::RebuildSimulationScene()
 	}
 
 	// Set the locked simulation flag to false
-	bIsSimulationLocked = false;
+	bIsSimulationLocked = !IsSimulationEnabled();
 
 	// Update simulation visualizations
 	for (const TPair<FName, TUniquePtr<UE::Dataflow::IDataflowSimulationVisualization>>& Visualization : UE::Dataflow::FDataflowSimulationVisualizationRegistry::GetInstance().GetVisualizations())
@@ -313,6 +313,7 @@ void FDataflowSimulationScene::RecordSimulationCache()
 			if(SimulationGenerator.IsValid())
 			{
 				SimulationGenerator->RequestGeneratorAction(UE::Dataflow::EDataflowGeneratorActions::StartGenerate);
+				bIsSimulationLocked = true;
 			}
 		}
 		else
