Datafolw : allow evaluation of the dataflow graph while in PIE for cloth asset
- added evaluation settings to the dataflow editors

diff --git a/Engine/Plugins/ChaosClothAssetEditor/Source/ChaosClothAssetEditor/Private/ChaosClothAsset/AssetDefinition_ClothAsset.cpp b/Engine/Plugins/ChaosClothAssetEditor/Source/ChaosClothAssetEditor/Private/ChaosClothAsset/AssetDefinition_ClothAsset.cpp
--- a/Engine/Plugins/ChaosClothAssetEditor/Source/ChaosClothAssetEditor/Private/ChaosClothAsset/AssetDefinition_ClothAsset.cpp
+++ b/Engine/Plugins/ChaosClothAssetEditor/Source/ChaosClothAssetEditor/Private/ChaosClothAsset/AssetDefinition_ClothAsset.cpp
@@ -18,6 +18,7 @@
 #include "Dataflow/DataflowEditor.h"
 #include "Dataflow/DataflowSNode.h"
 #include "Dataflow/DataflowSimulationScene.h"
+#include "Dataflow/DataflowEditorToolkit.h"
 #include "Toolkits/SimpleAssetEditor.h"
 
 #include UE_INLINE_GENERATED_CPP_BY_NAME(AssetDefinition_ClothAsset)
@@ -210,6 +211,12 @@ EAssetCommandResult UAssetDefinition_ClothAsset::OpenAssets(const FAssetOpenArgs
 					AssetEditor->AddEditorSettings(SimulationSettings);
 				}
 
+				if (UDataflowEvaluationSettings* EvaluationSettings = NewObject<UDataflowEvaluationSettings>())
+				{
+					EvaluationSettings->bAllowEvaluationInPIE = true;
+					AssetEditor->AddEditorSettings(EvaluationSettings);
+				}
+
 				const TSubclassOf<AActor> ActorClass = StaticLoadClass(AActor::StaticClass(), nullptr,
 					TEXT("/ChaosClothAssetEditor/BP_ClothPreview.BP_ClothPreview_C"), nullptr, LOAD_None, nullptr);
 				AssetEditor->RegisterToolCategories({"General", "Cloth"});
diff --git a/Engine/Plugins/Experimental/Dataflow/Source/DataflowEditor/Private/Dataflow/DataflowEditorToolkit.cpp b/Engine/Plugins/Experimental/Dataflow/Source/DataflowEditor/Private/Dataflow/DataflowEditorToolkit.cpp
--- a/Engine/Plugins/Experimental/Dataflow/Source/DataflowEditor/Private/Dataflow/DataflowEditorToolkit.cpp
+++ b/Engine/Plugins/Experimental/Dataflow/Source/DataflowEditor/Private/Dataflow/DataflowEditorToolkit.cpp
@@ -249,6 +249,12 @@ FDataflowEditorToolkit::FDataflowEditorToolkit(UAssetEditor* InOwningAssetEditor
 	// When saving, only prompt to checkout and save assets that are actually modified
 	bCheckDirtyOnAssetSave = true;
 
+	bAllowEvaluationInPIE = false;
+	if (const UDataflowEvaluationSettings* EvaluationSettings = DataflowEditor->FindEditorSettings<UDataflowEvaluationSettings>())
+	{
+		bAllowEvaluationInPIE = EvaluationSettings->bAllowEvaluationInPIE;
+	}
+
 	check(DataflowEditor);
 
 	ConstructionDefaultLayout = FTabManager::NewLayout(FName("DataflowConstructionLayout12"))
@@ -2123,7 +2129,8 @@ void FDataflowEditorToolkit::EvaluateNode(const FDataflowNode* Node, const FData
 	if (bDataflowEnableGraphEval)
 	{
 		const bool bIsInPIEOrSimulate = GEditor->PlayWorld || GEditor->bIsSimulatingInEditor;
-		if (!bIsInPIEOrSimulate)  // TODO: make this test optional for some toolkit
+		const bool bCanEvaluate = bAllowEvaluationInPIE || !bIsInPIEOrSimulate;
+		if (bCanEvaluate)
 		{
 			if (const TObjectPtr<UDataflowBaseContent>& EditorContent = GetEditorContent())
 			{
diff --git a/Engine/Plugins/Experimental/Dataflow/Source/DataflowEditor/Public/Dataflow/DataflowEditorToolkit.h b/Engine/Plugins/Experimental/Dataflow/Source/DataflowEditor/Public/Dataflow/DataflowEditorToolkit.h
--- a/Engine/Plugins/Experimental/Dataflow/Source/DataflowEditor/Public/Dataflow/DataflowEditorToolkit.h
+++ b/Engine/Plugins/Experimental/Dataflow/Source/DataflowEditor/Public/Dataflow/DataflowEditorToolkit.h
@@ -6,6 +6,9 @@
 #include "Misc/NotifyHook.h"
 #include "TickableEditorObject.h"
 #include "Dataflow/DataflowContent.h"
+#include "Dataflow/DataflowEditor.h"
+
+#include "DataflowEditorToolkit.generated.h"
 
 #define UE_API DATAFLOWEDITOR_API
 
@@ -45,6 +48,17 @@ namespace UE::Dataflow
 }
 namespace EMessageSeverity { enum Type : int; }
 
+
+UCLASS(MinimalAPI)
+class UDataflowEvaluationSettings : public UDataflowEditorSettings
+{
+public:
+	GENERATED_BODY()
+
+	UPROPERTY()
+	bool bAllowEvaluationInPIE = false;
+};
+
 class FDataflowEditorToolkit final : public FBaseCharacterFXEditorToolkit, public FTickableEditorObject, public FNotifyHook, public FGCObject
 {
 	using FBaseCharacterFXEditorToolkit::ObjectScene;
@@ -362,6 +376,9 @@ private:
 	/** Cached value of the p.Dataflow.EnableGraphEval cvar, to avoid calling FindConsoleVariable too often */
 	bool bDataflowEnableGraphEval;
 
+	/** Cached value from the editor evaluation settings */
+	bool bAllowEvaluationInPIE = false;
+
 	EDataflowEditorEvaluationMode EvaluationMode;
 
 	/** Delegate for updating the cached value of p.Dataflow.EnableGraphEval */
