Visually update welded vertex weight values when painting with Dataflow paint tool.

diff --git a/Engine/Plugins/Experimental/Dataflow/Source/DataflowEditor/Private/DataflowEditorTools/DataflowEditorVertexAttributePaintTool.cpp b/Engine/Plugins/Experimental/Dataflow/Source/DataflowEditor/Private/DataflowEditorTools/DataflowEditorVertexAttributePaintTool.cpp
--- a/Engine/Plugins/Experimental/Dataflow/Source/DataflowEditor/Private/DataflowEditorTools/DataflowEditorVertexAttributePaintTool.cpp
+++ b/Engine/Plugins/Experimental/Dataflow/Source/DataflowEditor/Private/DataflowEditorTools/DataflowEditorVertexAttributePaintTool.cpp
@@ -1082,6 +1082,25 @@ void UDataflowEditorVertexAttributePaintTool::UpdateROI(const FSculptBrushStamp&
 		}
 	}
 
+	if (VertexData.HaveDynamicMeshToWeightConversion())
+	{
+		// Find triangles whose vertices map to the same welded vertex as any vertex in VertexSetBuffer and add them to TriangleROI
+		for (const int VertexID : VertexSetBuffer)
+		{
+			for (const int OtherVertexID : VertexData.GetWeightToDynamicMesh()[VertexData.GetDynamicMeshToWeight()[VertexID]])
+			{
+				if (OtherVertexID != VertexID)
+				{
+					Mesh->EnumerateVertexTriangles(OtherVertexID, [this](int32 AdjacentTri)
+						{
+							TriangleROI.Add(AdjacentTri);
+						});
+				}
+			}
+		}
+
+	}
+
 	// Construct ROI vertex set
 	VertexSetBuffer.Reset();
 	for (int32 tid : TriangleROI)
diff --git a/Engine/Plugins/Experimental/Dataflow/Source/DataflowEditor/Public/DataflowEditorTools/DataflowEditorVertexAttributePaintTool.h b/Engine/Plugins/Experimental/Dataflow/Source/DataflowEditor/Public/DataflowEditorTools/DataflowEditorVertexAttributePaintTool.h
--- a/Engine/Plugins/Experimental/Dataflow/Source/DataflowEditor/Public/DataflowEditorTools/DataflowEditorVertexAttributePaintTool.h
+++ b/Engine/Plugins/Experimental/Dataflow/Source/DataflowEditor/Public/DataflowEditorTools/DataflowEditorVertexAttributePaintTool.h
@@ -295,6 +295,10 @@ public:
 
 	void UpdateNode(TObjectPtr<UDataflowContextObject> InDataflowEditorContextObject, IToolsContextTransactionsAPI* TransactionAPI);
 
+	bool HaveDynamicMeshToWeightConversion() const { return bHaveDynamicMeshToWeightConversion; }
+	const TArray<int32>& GetDynamicMeshToWeight() const { return  DynamicMeshToWeight; }
+	const TArray<TArray<int32>>& GetWeightToDynamicMesh() const { return WeightToDynamicMesh; }
+
 private:
 	// DynamicMesh might be unwelded mesh, but weights are on the welded mesh.
 	bool bHaveDynamicMeshToWeightConversion = false;
