Dataflow node for setting skinning skeletal mesh in  a collection

diff --git a/Engine/Plugins/Experimental/Dataflow/Source/DataflowNodes/Private/Dataflow/DataflowCollectionSetSkinningSkeletalMesh.cpp b/Engine/Plugins/Experimental/Dataflow/Source/DataflowNodes/Private/Dataflow/DataflowCollectionSetSkinningSkeletalMesh.cpp
new file mode 100644
--- /dev/null
+++ b/Engine/Plugins/Experimental/Dataflow/Source/DataflowNodes/Private/Dataflow/DataflowCollectionSetSkinningSkeletalMesh.cpp
@@ -0,0 +1,48 @@
+// Copyright Epic Games, Inc. All Rights Reserved.
+
+#include "Dataflow/DataflowCollectionSetSkinningSkeletalMesh.h"
+#include "Dataflow/DataflowObjectInterface.h"
+#include "GeometryCollection/Facades/CollectionVertexBoneWeightsFacade.h"
+
+
+#include UE_INLINE_GENERATED_CPP_BY_NAME(DataflowCollectionSetSkinningSkeletalMesh)
+
+FDataflowCollectionSetSkinningSkeletalMesh::FDataflowCollectionSetSkinningSkeletalMesh(const UE::Dataflow::FNodeParameters& InParam, FGuid InGuid)
+	: FDataflowNode(InParam, InGuid)
+{
+	RegisterInputConnection(&Collection);
+	RegisterInputConnection(&SkeletalMesh);
+	RegisterInputConnection(&GeometrySelection);
+	RegisterInputConnection(&LODIndex).SetCanHidePin(true).SetPinIsHidden(true);
+	RegisterOutputConnection(&Collection, &Collection);
+}
+
+
+void FDataflowCollectionSetSkinningSkeletalMesh::Evaluate(UE::Dataflow::FContext& Context, const FDataflowOutput* Out) const
+{
+	if (Out->IsA(&Collection))
+	{
+		FManagedArrayCollection GeometryCollection = GetValue(Context, &Collection);
+		GeometryCollection::Facades::FVertexBoneWeightsFacade SkinningFacade(GeometryCollection, false);
+		SkinningFacade.DefineSchema();
+
+		const TObjectPtr<USkeletalMesh> InputSkelMesh = GetValue(Context, &SkeletalMesh);
+		const FDataflowGeometrySelection InputSelection = GetValue(Context, &GeometrySelection);
+		const bool bValidSelection = InputSelection.IsValidForCollection(GeometryCollection);
+
+		if (SkinningFacade.IsValid() && InputSkelMesh && InputSkelMesh->IsValidLODIndex(LODIndex))
+		{
+			for (int32 GeometryIndex = 0; GeometryIndex < SkinningFacade.NumGeometry(); ++GeometryIndex)
+			{
+				if (!bValidSelection || (bValidSelection && InputSelection.IsSelected(GeometryIndex)))
+				{
+					SkinningFacade.ModifyGeometryBinding(GeometryIndex, InputSkelMesh, LODIndex);
+				}
+			}
+		}
+
+		SetValue(Context, MoveTemp(GeometryCollection), &Collection);
+	}
+}
+
+
diff --git a/Engine/Plugins/Experimental/Dataflow/Source/DataflowNodes/Private/Dataflow/DataflowNodesPlugin.cpp b/Engine/Plugins/Experimental/Dataflow/Source/DataflowNodes/Private/Dataflow/DataflowNodesPlugin.cpp
--- a/Engine/Plugins/Experimental/Dataflow/Source/DataflowNodes/Private/Dataflow/DataflowNodesPlugin.cpp
+++ b/Engine/Plugins/Experimental/Dataflow/Source/DataflowNodes/Private/Dataflow/DataflowNodesPlugin.cpp
@@ -12,6 +12,7 @@
 #include "Dataflow/DataflowCollectionAddScalarVertexPropertyNode.h"
 #include "Dataflow/DataflowCollectionEditSkinWeightsNode.h"
 #include "Dataflow/DataflowCollectionEditSkeletonBonesNode.h"
+#include "Dataflow/DataflowCollectionSetSkinningSkeletalMesh.h"
 #include "Dataflow/DataflowSkeletonAssetTerminalNode.h"
 
 #define LOCTEXT_NAMESPACE "DataflowNodes"
@@ -52,6 +53,7 @@ void IDataflowNodesPlugin::StartupModule()
 	DATAFLOW_NODE_REGISTER_CREATION_FACTORY(FDataflowCollectionAddScalarVertexPropertyNode);
 	DATAFLOW_NODE_REGISTER_CREATION_FACTORY(FDataflowCollectionEditSkinWeightsNode);
 	DATAFLOW_NODE_REGISTER_CREATION_FACTORY(FDataflowCollectionEditSkeletonBonesNode);
+	DATAFLOW_NODE_REGISTER_CREATION_FACTORY(FDataflowCollectionSetSkinningSkeletalMesh);
 	DATAFLOW_NODE_REGISTER_CREATION_FACTORY(FSkeletonAssetTerminalNode);
 
 	UE::Dataflow::RegisterNodeFilter(FDataflowTerminalNode::StaticType());
diff --git a/Engine/Plugins/Experimental/Dataflow/Source/DataflowNodes/Public/Dataflow/DataflowCollectionSetSkinningSkeletalMesh.h b/Engine/Plugins/Experimental/Dataflow/Source/DataflowNodes/Public/Dataflow/DataflowCollectionSetSkinningSkeletalMesh.h
new file mode 100644
--- /dev/null
+++ b/Engine/Plugins/Experimental/Dataflow/Source/DataflowNodes/Public/Dataflow/DataflowCollectionSetSkinningSkeletalMesh.h
@@ -0,0 +1,45 @@
+// Copyright Epic Games, Inc. All Rights Reserved.
+
+#pragma once 
+
+#include "CoreMinimal.h"
+#include "ChaosLog.h"
+#include "Dataflow/DataflowConnectionTypes.h"
+#include "Dataflow/DataflowCore.h"
+#include "Dataflow/DataflowSelection.h"
+#include "GeometryCollection/ManagedArrayCollection.h"
+#include "Engine/SkeletalMesh.h"
+
+#include "DataflowCollectionSetSkinningSkeletalMesh.generated.h"
+
+/** Set the skeletal mesh for the collection */
+USTRUCT(meta = (Experimental, DataflowGroom))
+struct FDataflowCollectionSetSkinningSkeletalMesh : public FDataflowNode
+{
+	GENERATED_BODY()
+
+	DATAFLOW_NODE_DEFINE_INTERNAL(FDataflowCollectionSetSkinningSkeletalMesh, "SetSkinningSkeletalMesh", "Collection", "")
+
+public:
+	FDataflowCollectionSetSkinningSkeletalMesh(const UE::Dataflow::FNodeParameters& InParam, FGuid InGuid = FGuid::NewGuid());
+
+	//~ Begin FDataflowNode interface
+	virtual void Evaluate(UE::Dataflow::FContext& Context, const FDataflowOutput* Out) const override;
+	//~ End FDataflowNode interface
+
+	/** Managed array collection to be used to store datas */
+	UPROPERTY(meta = (DataflowInput, DataflowOutput, DisplayName = "Collection", DataflowPassthrough = "Collection", DataflowRenderGroups = "Surface"))
+	FManagedArrayCollection Collection;
+
+	/** Skeletal mesh binding to be stored in the collection */
+	UPROPERTY(EditAnywhere, Category = "Skeletal Mesh", meta = (DataflowInput, DataflowPassthrough = "SkeletalMesh"))
+	TObjectPtr<USkeletalMesh> SkeletalMesh = nullptr;
+
+	UPROPERTY(EditAnywhere, Category = "Skeletal Mesh", meta = (DataflowInput, DisplayName = "LOD Index"))
+	int32 LODIndex = 0;
+
+	/** Geometries to set skinning skeletal mesh on*/
+	UPROPERTY(EditAnywhere, Category = "Skeletal Mesh", meta = (DataflowInput))
+	FDataflowGeometrySelection GeometrySelection;
+};
+
