Dataflow : add pointer check for bounds
also remove warning from UHT for unessary UPROPERTY

diff --git a/Engine/Plugins/Experimental/Dataflow/Source/DataflowEditor/Private/Dataflow/DataflowEditorPreviewSceneBase.cpp b/Engine/Plugins/Experimental/Dataflow/Source/DataflowEditor/Private/Dataflow/DataflowEditorPreviewSceneBase.cpp
--- a/Engine/Plugins/Experimental/Dataflow/Source/DataflowEditor/Private/Dataflow/DataflowEditorPreviewSceneBase.cpp
+++ b/Engine/Plugins/Experimental/Dataflow/Source/DataflowEditor/Private/Dataflow/DataflowEditorPreviewSceneBase.cpp
@@ -127,7 +127,7 @@ FBox FDataflowPreviewSceneBase::GetBoundingBox() const
 		USelection* const SelectedComponents = DataflowModeManager->GetSelectedComponents();
 
 		TArray<TWeakObjectPtr<UObject>> SelectedObjects;
-		const int32 NumSelected = SelectedComponents->GetSelectedObjects(SelectedObjects);
+		const int32 NumSelected = SelectedComponents ? SelectedComponents->GetSelectedObjects(SelectedObjects): 0;
 		
 		if(NumSelected > 0)
 		{
diff --git a/Engine/Plugins/Experimental/Dataflow/Source/DataflowEditor/Private/Dataflow/DataflowSimulationScene.cpp b/Engine/Plugins/Experimental/Dataflow/Source/DataflowEditor/Private/Dataflow/DataflowSimulationScene.cpp
--- a/Engine/Plugins/Experimental/Dataflow/Source/DataflowEditor/Private/Dataflow/DataflowSimulationScene.cpp
+++ b/Engine/Plugins/Experimental/Dataflow/Source/DataflowEditor/Private/Dataflow/DataflowSimulationScene.cpp
@@ -405,7 +405,7 @@ FBox FDataflowSimulationScene::GetBoundingBox() const
 		USelection* const SelectedComponents = DataflowModeManager->GetSelectedComponents();
 
 		TArray<TWeakObjectPtr<UObject>> SelectedObjects;
-		const int32 NumSelected = SelectedComponents->GetSelectedObjects(SelectedObjects);
+		const int32 NumSelected = SelectedComponents? SelectedComponents->GetSelectedObjects(SelectedObjects): 0;
 
 		if (NumSelected > 0)
 		{
diff --git a/Engine/Plugins/Experimental/Dataflow/Source/DataflowEditor/Public/Dataflow/DataflowEditorToolkit.h b/Engine/Plugins/Experimental/Dataflow/Source/DataflowEditor/Public/Dataflow/DataflowEditorToolkit.h
--- a/Engine/Plugins/Experimental/Dataflow/Source/DataflowEditor/Public/Dataflow/DataflowEditorToolkit.h
+++ b/Engine/Plugins/Experimental/Dataflow/Source/DataflowEditor/Public/Dataflow/DataflowEditorToolkit.h
@@ -344,7 +344,6 @@ private:
 	FDelegateHandle OnContextHasErrorDelegateHandle;
 
 	// The currently selected set of dataflow nodes. 
-	UPROPERTY()
 	TSet<TObjectPtr<UDataflowEdNode>> SelectedDataflowNodes;
 
 	/** Pointer to the construction viewport scene. Note this is an alias of ObjectScene in FBaseCharacterFXEditorToolkit but with the specific type */
