Dataflow : only show the pin on node that have outputs or can debug draw

diff --git a/Engine/Plugins/Experimental/Dataflow/Source/DataflowEditor/Private/Dataflow/DataflowSNode.cpp b/Engine/Plugins/Experimental/Dataflow/Source/DataflowEditor/Private/Dataflow/DataflowSNode.cpp
--- a/Engine/Plugins/Experimental/Dataflow/Source/DataflowEditor/Private/Dataflow/DataflowSNode.cpp
+++ b/Engine/Plugins/Experimental/Dataflow/Source/DataflowEditor/Private/Dataflow/DataflowSNode.cpp
@@ -206,9 +206,12 @@ TSharedRef<SWidget> SDataflowEdNode::CreateTitleRightWidget()
 						.Image(this, &SDataflowEdNode::GetPinButtonImage)
 				];
 
+	const TSharedPtr<const FDataflowNode> DataflowNode = DataflowGraphNode ? DataflowGraphNode->GetDataflowNode() : TSharedPtr<const FDataflowNode>();
+
 	const bool bHasRenderingCheckBox = (
-		DataflowInterface->NodesHaveToggleWidget() && 
-		DataflowGraphNode->GetDataflowNode()
+		DataflowInterface->NodesHaveToggleWidget() &&
+		DataflowNode && 
+		(DataflowNode->NumOutputs() > 0 || DataflowNode->CanDebugDraw())
 		);
 
 	if (bHasRenderingCheckBox)
