Dataflow : display all available view mode when no node is selected

diff --git a/Engine/Plugins/Experimental/Dataflow/Source/DataflowEditor/Private/Dataflow/DataflowConstructionViewport.cpp b/Engine/Plugins/Experimental/Dataflow/Source/DataflowEditor/Private/Dataflow/DataflowConstructionViewport.cpp
--- a/Engine/Plugins/Experimental/Dataflow/Source/DataflowEditor/Private/Dataflow/DataflowConstructionViewport.cpp
+++ b/Engine/Plugins/Experimental/Dataflow/Source/DataflowEditor/Private/Dataflow/DataflowConstructionViewport.cpp
@@ -104,9 +104,11 @@ void SDataflowConstructionViewport::BuildViewMenu(FToolMenuSection& MenuSection)
 					{
 						if (UDataflowEditorMode* EditorMode = Viewport->GetEdMode())
 						{
+							const bool bIsAnyNodeSelected = EditorMode->IsAnyNodeSelected();
+
 							for (const TPair<FName, TSharedPtr<FUICommandInfo>>& NameAndCommand : FDataflowEditorCommandsImpl::Get().SetConstructionViewModeCommands)
 							{
-								if (EditorMode->CanChangeConstructionViewModeTo(NameAndCommand.Key))
+								if (!bIsAnyNodeSelected || EditorMode->CanChangeConstructionViewModeTo(NameAndCommand.Key))
 								{
 									SimulationSection.AddMenuEntry(NameAndCommand.Value);
 								}
@@ -224,7 +226,8 @@ void SDataflowConstructionViewport::BindCommands()
 			{ 
 				if (const UDataflowEditorMode* const EdMode = GetEdMode())
 				{
-					return EdMode->CanChangeConstructionViewModeTo(ViewModeName);
+					const bool bIsAnyNodeSelected = EdMode->IsAnyNodeSelected();
+					return !bIsAnyNodeSelected || EdMode->CanChangeConstructionViewModeTo(ViewModeName);
 				}
 				return false; 
 			}),
diff --git a/Engine/Plugins/Experimental/Dataflow/Source/DataflowEditor/Private/Dataflow/DataflowEditorMode.cpp b/Engine/Plugins/Experimental/Dataflow/Source/DataflowEditor/Private/Dataflow/DataflowEditorMode.cpp
--- a/Engine/Plugins/Experimental/Dataflow/Source/DataflowEditor/Private/Dataflow/DataflowEditorMode.cpp
+++ b/Engine/Plugins/Experimental/Dataflow/Source/DataflowEditor/Private/Dataflow/DataflowEditorMode.cpp
@@ -1150,6 +1150,18 @@ void UDataflowEditorMode::OnDataflowNodeDeleted(const TSet<UObject*>& DeletedNod
 	ToolsContext->EndTool(bCanCancel ? EToolShutdownType::Cancel : EToolShutdownType::Completed);
 }
 
+bool UDataflowEditorMode::IsAnyNodeSelected() const
+{
+	const TSharedPtr<const SDataflowGraphEditor> PinnedDataflowGraphEditor = DataflowGraphEditor.Pin();
+	if (!PinnedDataflowGraphEditor)
+	{
+		return false;
+	}
+
+	UEdGraphNode* const SelectedNode = PinnedDataflowGraphEditor->GetSingleSelectedNode();
+	return (SelectedNode != nullptr);
+}
+
 UEdGraphNode* UDataflowEditorMode::GetSingleSelectedNodeWithOutputType(const FName& SelectedNodeOutputTypeName) const
 {
 	const TSharedPtr<const SDataflowGraphEditor> PinnedDataflowGraphEditor = DataflowGraphEditor.Pin();
diff --git a/Engine/Plugins/Experimental/Dataflow/Source/DataflowEditor/Public/Dataflow/DataflowEditorMode.h b/Engine/Plugins/Experimental/Dataflow/Source/DataflowEditor/Public/Dataflow/DataflowEditorMode.h
--- a/Engine/Plugins/Experimental/Dataflow/Source/DataflowEditor/Public/Dataflow/DataflowEditorMode.h
+++ b/Engine/Plugins/Experimental/Dataflow/Source/DataflowEditor/Public/Dataflow/DataflowEditorMode.h
@@ -98,6 +98,8 @@ public:
 	 */
 	UE_API const FName* FindRecentlyUsedViewMode(const TArray<FName>& AvailableViewModes) const;
 
+	UE_API bool IsAnyNodeSelected() const;
+
 private:
 	friend class FDataflowEditorToolkit;
 	friend class FDataflowEditorModeToolkit;
