Dataflow skeleton proxy object:  added the possibility to override the transform
- preparation change for bone manipulation in dataflow tools

diff --git a/Engine/Plugins/Experimental/Dataflow/Source/DataflowNodes/Private/Dataflow/DataflowSkeletalMeshNodes.cpp b/Engine/Plugins/Experimental/Dataflow/Source/DataflowNodes/Private/Dataflow/DataflowSkeletalMeshNodes.cpp
--- a/Engine/Plugins/Experimental/Dataflow/Source/DataflowNodes/Private/Dataflow/DataflowSkeletalMeshNodes.cpp
+++ b/Engine/Plugins/Experimental/Dataflow/Source/DataflowNodes/Private/Dataflow/DataflowSkeletalMeshNodes.cpp
@@ -38,6 +38,12 @@ namespace UE::Dataflow
 ///////////////////////////////////////////////////////////////////////////////////////////////////////////////////////////////////
 ///
 
+void FDataflowDebugDrawSkeletonObject::OverrideBoneTransforms(const TArray<FTransform>& InBoneTransforms)
+{
+	FScopeLock Lock(&TransformOverridesLockGuard);
+	TransformOverrides = InBoneTransforms;
+}
+
 void FDataflowDebugDrawSkeletonObject::PopulateDataflowElements()
 {
 	ElementsOffset = DataflowElements.Num();
@@ -87,7 +93,7 @@ FBox FDataflowDebugDrawSkeletonObject::ComputeBoundingBox() const
 
 void FDataflowDebugDrawSkeletonObject::DrawDataflowElements(FPrimitiveDrawInterface* PDI) 
 {
-	const uint32 NumBones = ReferenceSkeleton.GetNum();
+	const int32 NumBones = ReferenceSkeleton.GetNum();
 
 	TArray<FBoneIndexType> RequiredBones;
 	TArray<TRefCountPtr<HHitProxy>> HitProxies;
@@ -97,6 +103,20 @@ void FDataflowDebugDrawSkeletonObject::DrawDataflowElements(FPrimitiveDrawInterf
 	TArray<FName> BoneNames;
 	
 	ReferenceSkeleton.GetBoneAbsoluteTransforms(WorldTransforms);
+
+	// apply the transform overrides
+	{
+		FScopeLock Lock(&TransformOverridesLockGuard);
+		check(NumBones == WorldTransforms.Num());
+		for (int32 BoneIndex = 0; BoneIndex < NumBones; ++BoneIndex)
+		{
+			if (TransformOverrides.IsValidIndex(BoneIndex))
+			{
+				WorldTransforms[BoneIndex] = TransformOverrides[BoneIndex];
+			}
+		}
+	}
+
 	RequiredBones.Reserve(NumBones);
 
 	bool bSelectionChanged = false;
diff --git a/Engine/Plugins/Experimental/Dataflow/Source/DataflowNodes/Public/Dataflow/DataflowSkeletalMeshNodes.h b/Engine/Plugins/Experimental/Dataflow/Source/DataflowNodes/Public/Dataflow/DataflowSkeletalMeshNodes.h
--- a/Engine/Plugins/Experimental/Dataflow/Source/DataflowNodes/Public/Dataflow/DataflowSkeletalMeshNodes.h
+++ b/Engine/Plugins/Experimental/Dataflow/Source/DataflowNodes/Public/Dataflow/DataflowSkeletalMeshNodes.h
@@ -27,7 +27,7 @@ struct FDataflowDebugDrawSkeletonObject : public FDataflowDebugDrawBaseObject
 	FDataflowDebugDrawSkeletonObject(IDataflowDebugDrawInterface::FDataflowElementsType& InDataflowElements, const FReferenceSkeleton& InReferenceSkeleton, const bool bBonesVisible = true) :
 		FDataflowDebugDrawBaseObject(InDataflowElements), ReferenceSkeleton(InReferenceSkeleton), bBonesVisible(bBonesVisible)
 	{}
-	
+
 	/** Populate dataflow elements */
 	UE_API virtual void PopulateDataflowElements() override;
 
@@ -37,13 +37,19 @@ struct FDataflowDebugDrawSkeletonObject : public FDataflowDebugDrawBaseObject
 	/** Compute the dataflow elements bounding box */
 	UE_API virtual FBox ComputeBoundingBox() const override;
 
+	/** override the bone transforms */
+	UE_API void OverrideBoneTransforms(const TArray<FTransform>& InBoneTransforms);
+
 	/** Delegate to broadcast bones selection changes */
 	FDataflowBoneSelectionChangedNotifyDelegate OnBoneSelectionChanged;
 
 private :
+	FCriticalSection TransformOverridesLockGuard;
+	TArray<FTransform> TransformOverrides;
+
 	/** Skeletal mesh to use to populate/draw scene elements */
 	const FReferenceSkeleton& ReferenceSkeleton;
-	
+		
 	/** Previous element selection */
 	TArray<bool> PreviousSelection;
 
