Dataflow : Weight map paint tool : save brush type in config
* also set Replace as the default value instead off Add

diff --git a/Engine/Plugins/Experimental/Dataflow/Source/DataflowEditor/Public/DataflowEditorTools/DataflowEditorVertexAttributePaintTool.h b/Engine/Plugins/Experimental/Dataflow/Source/DataflowEditor/Public/DataflowEditorTools/DataflowEditorVertexAttributePaintTool.h
--- a/Engine/Plugins/Experimental/Dataflow/Source/DataflowEditor/Public/DataflowEditorTools/DataflowEditorVertexAttributePaintTool.h
+++ b/Engine/Plugins/Experimental/Dataflow/Source/DataflowEditor/Public/DataflowEditorTools/DataflowEditorVertexAttributePaintTool.h
@@ -91,7 +91,7 @@ struct FDataflowEditorVertexAttributePaintToolBrushProperties
 	GENERATED_BODY()
 
 	UPROPERTY()
-	EDataflowEditorToolEditOperation BrushMode = EDataflowEditorToolEditOperation::Add;
+	EDataflowEditorToolEditOperation BrushMode = EDataflowEditorToolEditOperation::Replace;
 
 	/** Adaptive size of brush - Relative to the model size */
 	UPROPERTY()
@@ -217,7 +217,7 @@ public:
 	UPROPERTY()
 	EDataflowEditorToolEditMode EditingMode = EDataflowEditorToolEditMode::Brush;
 	
-	UPROPERTY()
+	UPROPERTY(config)
 	FDataflowEditorVertexAttributePaintToolBrushProperties BrushProperties;
 
 	UPROPERTY()
