Dataflow : Added options to not update the visual simulation viewport when in PIE or when the editor is not focused
- by default they are true ( no update in both case )
- but they can be controlled in the preview details view

diff --git a/Engine/Plugins/Experimental/Dataflow/Source/DataflowEditor/Private/Dataflow/DataflowSimulationViewportClient.cpp b/Engine/Plugins/Experimental/Dataflow/Source/DataflowEditor/Private/Dataflow/DataflowSimulationViewportClient.cpp
--- a/Engine/Plugins/Experimental/Dataflow/Source/DataflowEditor/Private/Dataflow/DataflowSimulationViewportClient.cpp
+++ b/Engine/Plugins/Experimental/Dataflow/Source/DataflowEditor/Private/Dataflow/DataflowSimulationViewportClient.cpp
@@ -19,11 +19,13 @@
 #include "Selection.h"
 #include "SGraphPanel.h"
 #include "SNodePanel.h"
+#include "Widgets/Docking/SDockTab.h"
 
 FDataflowSimulationViewportClient::FDataflowSimulationViewportClient(FEditorModeTools* InModeTools,
                                                              TWeakPtr<FDataflowSimulationScene> InDataflowSimulationScene,  const bool bCouldTickScene,
                                                              const TWeakPtr<SEditorViewport> InEditorViewportWidget)
 	: FDataflowEditorViewportClientBase(InModeTools, InDataflowSimulationScene, bCouldTickScene, InEditorViewportWidget)
+	, DataflowSimulationScenePtr(InDataflowSimulationScene)
 {
 	// We want our near clip plane to be quite close so that we can zoom in further.
 	OverrideNearClipPlane(KINDA_SMALL_NUMBER);
@@ -61,12 +63,50 @@ void FDataflowSimulationViewportClient::SetToolCommandList(TWeakPtr<FUICommandLi
 	ToolCommandList = InToolCommandList;
 }
 
+bool FDataflowSimulationViewportClient::IsDataflowEditorTabActive() const
+{
+	if (TSharedPtr<FDataflowEditorToolkit> DataflowEditorToolkit = DataflowEditorToolkitPtr.Pin())
+	{
+		if (TSharedPtr<FTabManager> TabManager = DataflowEditorToolkit->GetTabManager())
+		{
+			if (TSharedPtr<SDockTab> DataflowEditorTab = FGlobalTabmanager::Get()->GetMajorTabForTabManager(TabManager.ToSharedRef()))
+			{
+				bool bIsTabInSamePanelAsActiveOne = false;
+				if (TSharedPtr<class SDockTab> ActiveTab = FGlobalTabmanager::Get()->GetActiveTab())
+				{
+					TSharedPtr<FTabManager> ActiveTabManager = ActiveTab->GetTabManagerPtr();
+					bIsTabInSamePanelAsActiveOne = (ActiveTabManager == TabManager);
+				}
+				const bool bIsTabInForecround = DataflowEditorTab->IsForeground();
+				return (bIsTabInSamePanelAsActiveOne && bIsTabInForecround);
+			}
+		}
+	}
+	return false;
+}
+
 void FDataflowSimulationViewportClient::Tick(float DeltaSeconds)
 {
 	Super::Tick(DeltaSeconds);
-	if (const TSharedPtr<FDataflowPreviewSceneBase> DataflowPreviewScenePtr = DataflowPreviewScene.Pin())
+	if (const TSharedPtr<FDataflowSimulationScene> DataflowSimulationScene = DataflowSimulationScenePtr.Pin())
 	{
-		DataflowPreviewScenePtr->TickDataflowScene(DeltaSeconds);
+		bool bPauseInPIE = true;
+		bool bPauseWhenNotFocused = true;
+		if (UDataflowSimulationSceneDescription* SceneDescription = DataflowSimulationScene->GetPreviewSceneDescription())
+		{
+			bPauseInPIE = SceneDescription->bPauseSimulationViewportWhenPlayingInEditor;
+			bPauseWhenNotFocused = SceneDescription->bPauseSimulationViewportWhenNotFocused;
+		}
+
+		const bool bIsTabActive = IsDataflowEditorTabActive();
+		const bool bIsInPIEOrSIE = GEditor->PlayWorld != NULL || GEditor->bIsSimulatingInEditor;
+		const bool bShouldPause = (bPauseInPIE && bIsInPIEOrSIE) || (bPauseWhenNotFocused && !bIsTabActive);
+		if (bShouldPause)
+		{
+			return;
+		}
+
+		DataflowSimulationScene->TickDataflowScene(DeltaSeconds);
 	}
 }
 
diff --git a/Engine/Plugins/Experimental/Dataflow/Source/DataflowEditor/Public/Dataflow/DataflowSimulationScene.h b/Engine/Plugins/Experimental/Dataflow/Source/DataflowEditor/Public/Dataflow/DataflowSimulationScene.h
--- a/Engine/Plugins/Experimental/Dataflow/Source/DataflowEditor/Public/Dataflow/DataflowSimulationScene.h
+++ b/Engine/Plugins/Experimental/Dataflow/Source/DataflowEditor/Public/Dataflow/DataflowSimulationScene.h
@@ -60,6 +60,12 @@ public:
 	UPROPERTY(EditAnywhere, Category = "Scene")
 	FTransform BlueprintTransform = FTransform::Identity;
 
+	UPROPERTY(EditAnywhere, Category = "Viewport")
+	bool bPauseSimulationViewportWhenPlayingInEditor = true;
+
+	UPROPERTY(EditAnywhere, Category = "Viewport")
+	bool bPauseSimulationViewportWhenNotFocused = true;
+
 	/** Caching asset to be used to record the simulation  */
 	UPROPERTY(EditAnywhere, Category="Caching", DisplayName="Cache Collection")
 	TObjectPtr<UChaosCacheCollection> CacheAsset = nullptr;
diff --git a/Engine/Plugins/Experimental/Dataflow/Source/DataflowEditor/Public/Dataflow/DataflowSimulationViewportClient.h b/Engine/Plugins/Experimental/Dataflow/Source/DataflowEditor/Public/Dataflow/DataflowSimulationViewportClient.h
--- a/Engine/Plugins/Experimental/Dataflow/Source/DataflowEditor/Public/Dataflow/DataflowSimulationViewportClient.h
+++ b/Engine/Plugins/Experimental/Dataflow/Source/DataflowEditor/Public/Dataflow/DataflowSimulationViewportClient.h
@@ -40,6 +40,7 @@ public:
 	FOnSelectionChangedMulticast OnSelectionChangedMulticast;
 
 private:
+	bool IsDataflowEditorTabActive() const;
 
 	// FDataflowEditorViewportClientBase
 	UE_API virtual void OnViewportClicked(HHitProxy* HitProxy) override;
@@ -53,6 +54,9 @@ private:
 	/** Toolkit used to create the viewport client */
 	TWeakPtr<FDataflowEditorToolkit> DataflowEditorToolkitPtr = nullptr;
 
+	/** strongly typed simulation scene */
+	TWeakPtr<FDataflowSimulationScene> DataflowSimulationScenePtr;
+
 	// @todo(brice) : Is this needed?
 	TWeakPtr<FUICommandList> ToolCommandList;
 	
