Dataflow : add Fov and planes options to the view menu in both simulation and construction views

diff --git a/Engine/Plugins/Experimental/Dataflow/Source/DataflowEditor/Private/Dataflow/DataflowConstructionViewport.cpp b/Engine/Plugins/Experimental/Dataflow/Source/DataflowEditor/Private/Dataflow/DataflowConstructionViewport.cpp
--- a/Engine/Plugins/Experimental/Dataflow/Source/DataflowEditor/Private/Dataflow/DataflowConstructionViewport.cpp
+++ b/Engine/Plugins/Experimental/Dataflow/Source/DataflowEditor/Private/Dataflow/DataflowConstructionViewport.cpp
@@ -140,7 +140,11 @@ TSharedPtr<SWidget> SDataflowConstructionViewport::BuildViewportToolbar()
 		{
 			RightSection.Alignment = EToolMenuSectionAlign::Last;
 
-			RightSection.AddEntry(UE::UnrealEd::CreateCameraSubmenu(UE::UnrealEd::FViewportCameraMenuOptions().ShowCameraMovement()));
+			RightSection.AddEntry(UE::UnrealEd::CreateCameraSubmenu(
+				UE::UnrealEd::FViewportCameraMenuOptions()
+				.ShowCameraMovement()
+				.ShowLensControls()
+			));
 			RightSection.AddEntry(UE::UnrealEd::CreateViewModesSubmenu());
 			RightSection.AddEntry(UE::UnrealEd::CreateDefaultShowSubmenu());
 			RightSection.AddEntry(UE::UnrealEd::CreateAssetViewerProfileSubmenu());
diff --git a/Engine/Plugins/Experimental/Dataflow/Source/DataflowEditor/Private/Dataflow/DataflowSimulationViewport.cpp b/Engine/Plugins/Experimental/Dataflow/Source/DataflowEditor/Private/Dataflow/DataflowSimulationViewport.cpp
--- a/Engine/Plugins/Experimental/Dataflow/Source/DataflowEditor/Private/Dataflow/DataflowSimulationViewport.cpp
+++ b/Engine/Plugins/Experimental/Dataflow/Source/DataflowEditor/Private/Dataflow/DataflowSimulationViewport.cpp
@@ -104,7 +104,11 @@ TSharedPtr<SWidget> SDataflowSimulationViewport::BuildViewportToolbar()
 			FToolMenuSection& RightSection = Menu->AddSection("Right");
 			RightSection.Alignment = EToolMenuSectionAlign::Last;
 			
-			RightSection.AddEntry(UE::UnrealEd::CreateCameraSubmenu(UE::UnrealEd::FViewportCameraMenuOptions().ShowCameraMovement()));
+			RightSection.AddEntry(UE::UnrealEd::CreateCameraSubmenu(
+				UE::UnrealEd::FViewportCameraMenuOptions()
+					.ShowCameraMovement()
+					.ShowLensControls()
+			));
 			RightSection.AddEntry(UE::UnrealEd::CreateViewModesSubmenu());
 			RightSection.AddEntry(UE::UnrealEd::CreateDefaultShowSubmenu());
 			RightSection.AddEntry(UE::UnrealEd::CreateAssetViewerProfileSubmenu());
