Dataflow : shutdown tool with auto accept if the user deselect the current node
- did some cleanup for better logic readibility

diff --git a/Engine/Plugins/Experimental/Dataflow/Source/DataflowEditor/Private/Dataflow/DataflowEditorToolkit.cpp b/Engine/Plugins/Experimental/Dataflow/Source/DataflowEditor/Private/Dataflow/DataflowEditorToolkit.cpp
--- a/Engine/Plugins/Experimental/Dataflow/Source/DataflowEditor/Private/Dataflow/DataflowEditorToolkit.cpp
+++ b/Engine/Plugins/Experimental/Dataflow/Source/DataflowEditor/Private/Dataflow/DataflowEditorToolkit.cpp
@@ -1591,6 +1591,7 @@ void FDataflowEditorToolkit::OnNodeSelectionChanged(const TSet<UObject*>& InNewS
 
 	const TObjectPtr<UDataflowEdNode> PrevSelectedNode = GetOnlyFromSet(PreviouslySelectedNodes);
 	const TObjectPtr<UDataflowEdNode> NewSelectedNode = GetOnlyFromSet(SelectedDataflowNodes);
+	const bool bSingleSelectionChanged = (PrevSelectedNode != NewSelectedNode);
 
 	SetShouldRenderNodeFlag(DeselectedNodes, false);
 	SetShouldRenderNodeFlag({ NewSelectedNode }, true);
@@ -1598,63 +1599,35 @@ void FDataflowEditorToolkit::OnNodeSelectionChanged(const TSet<UObject*>& InNewS
 	// todo : move the tool logic somewhere else: we should not be activelly looking for a collection here, we should have some sort of interface to talk to the tools
 	UDataflowEditorMode* const DataflowMode = Cast<UDataflowEditorMode>(EditorModeManager->GetActiveScriptableMode(UDataflowEditorMode::EM_DataflowEditorModeId));
 
-	bool bSingleSelectionChanged = false;
-	if (const TObjectPtr<UDataflowBaseContent>& EditorContent = GetEditorContent(); EditorContent->GetDataflowAsset())
+	// First close any running tool by accepting the changes.
+	// New tool will start at the end of this method if a single node was selected
+	if (DataflowMode && bSingleSelectionChanged)
 	{
-		if (NewSelectedNode != PrevSelectedNode)
-		{
-			bSingleSelectionChanged = true;
+		DataflowMode->ShutdownActiveToolIfNeeded(EToolShutdownType::Accept);
+	}
 
-			EditorContent->SetSelectedNode(nullptr);
-			EditorContent->SetSelectedCollection(nullptr, /*bCollectionIsInput=*/ false);
+	if (const TObjectPtr<UDataflowBaseContent>& EditorContent = GetEditorContent(); EditorContent->GetDataflowAsset())
+	{
+		EditorContent->SetSelectedNode(NewSelectedNode);
+		EditorContent->SetSelectedCollection(nullptr, /*bCollectionIsInput=*/ false);
 
-			if (NewSelectedNode)
+		if (bSingleSelectionChanged)
+		{
+			// if alt is pressed, select make sure the first component is selected 
+			if (GetDataflowGraphEditor()->IsAltDown())
 			{
-				if (DataflowMode)
+				if (ConstructionScene)
 				{
-					// Close any running tool by accepting the changes.
-					// new tool will start at the end of this method if a single node was selected
-					DataflowMode->ShutdownActiveToolIfNeeded(EToolShutdownType::Accept);
-
-					EditorContent->SetSelectedNode(NewSelectedNode);
-
-
-// todo(dataflow) : delete when confirmed this is not causing any issue
-#if 0 // after discussion this does not seem to be used at all in the context of the Dataflow editor - disabling it for now and see if this cause any issues
-
-					// Call the node's OnSelected function. Some nodes use this to cache information from the inputs (e.g. FDataflowCollectionAddScalarVertexPropertyNode::CachedCollectionGroupNames)
-					TSharedPtr<UE::Dataflow::FEngineContext> DataflowContext = EditorContent->GetDataflowContext();
-					if (NewSelectedNode && DataflowContext.IsValid())
-					{
-						if (const TSharedPtr<FDataflowNode> DataflowNode = NewSelectedNode->GetDataflowNode())
-						{
-							// Update selected Collection in the ContextObject
-							for (const FDataflowOutput* const Output : DataflowNode->GetOutputs())
-							{
-								if (Output->GetType() == FName(TEXT("FManagedArrayCollection")))
-								{
-									const FManagedArrayCollection DefaultValue;
-									TSharedRef<FManagedArrayCollection> Collection = MakeShared<FManagedArrayCollection>(Output->GetValue<FManagedArrayCollection>(*DataflowContext, DefaultValue));
-									constexpr bool bCollectionIsInput = false;
-									EditorContent->SetSelectedCollection(Collection, bCollectionIsInput);
-								}
-							}
-						}
-					}
-#endif
+					ConstructionScene->SelectNodeComponents(NewSelectedNode);
 				}
 			}
-			
-			if (GetDataflowGraphEditor()->IsAltDown())
-			{
-				ConstructionScene->SelectNodeComponents(NewSelectedNode);
-			}
 		}
-
-		// refresh the views
-		UpdateViewsFromNode(NewSelectedNode);
 	}
 
+	// refresh the views
+	UpdateViewsFromNode(NewSelectedNode);
+
+
 	// Check if the current view mode can render the selected node. If not, try to find a view mode that can. 
 	if (DataflowMode)
 	{
