Dataflow : add bone manipulator

diff --git a/Engine/Plugins/Experimental/Dataflow/Source/DataflowEditor/Private/DataflowEditorTools/DataflowEditorBoneManipulator.cpp b/Engine/Plugins/Experimental/Dataflow/Source/DataflowEditor/Private/DataflowEditorTools/DataflowEditorBoneManipulator.cpp
new file mode 100644
--- /dev/null
+++ b/Engine/Plugins/Experimental/Dataflow/Source/DataflowEditor/Private/DataflowEditorTools/DataflowEditorBoneManipulator.cpp
@@ -0,0 +1,141 @@
+// Copyright Epic Games, Inc. All Rights Reserved.
+
+#include "DataflowEditorTools/DataflowEditorBoneManipulator.h"
+
+#include "BaseGizmos/CombinedTransformGizmo.h"
+#include "BaseGizmos/TransformGizmoUtil.h"
+#include "SkeletalMesh/RefSkeletonPoser.h"
+#include "Transforms/TransformGizmoDataBinder.h"
+
+#include UE_INLINE_GENERATED_CPP_BY_NAME(DataflowEditorBoneManipulator)
+
+//////////////////////////////////////////////////////////////////////////////////////////////////////////////
+
+void UDataflowBoneManipulator::Setup(UInteractiveToolManager* ToolManager, const FReferenceSkeleton& RefSkeleton)
+{
+	ensure(ToolManager);
+	if (!ToolManager)
+	{
+		return;
+	}
+
+	UInteractiveGizmoManager* const GizmoManager = ToolManager->GetPairedGizmoManager();
+	ensure(GizmoManager);
+	if (!GizmoManager)
+	{
+		return;
+	}
+
+	RefSkeletonPoser = NewObject<URefSkeletonPoser>();
+	RefSkeletonPoser->SetRefSkeleton(RefSkeleton);
+
+	// create the transform rotation only gizmo
+	TransformProxy = NewObject<UTransformProxy>(this);
+	TransformProxy->SetTransform(FTransform::Identity);
+	TransformProxy->OnBeginTransformEdit.AddUObject(this, &UDataflowBoneManipulator::OnStartTransformEdit);
+	TransformProxy->OnTransformChanged.AddUObject(this, &UDataflowBoneManipulator::OnTransformUpdated);
+	TransformProxy->OnEndTransformEdit.AddUObject(this, &UDataflowBoneManipulator::OnEndTransformEdit);
+	//TransformProxy->OnTransformChangedUndoRedo.AddUObject(this, &UDataflowBoneManipulator::OnWidgetTransformUpdated);
+
+	// create the rotation only gizmo
+	TransformGizmo = UE::TransformGizmoUtil::CreateCustomTransformGizmo(GizmoManager, ETransformGizmoSubElements::FreeRotate | ETransformGizmoSubElements::RotateAllAxes, this);
+	TransformGizmo->SetActiveTarget(TransformProxy, ToolManager);
+	TransformGizmo->SetVisibility(false);
+	TransformGizmo->bUseContextCoordinateSystem = false;
+	TransformGizmo->bUseContextGizmoMode = false;
+	TransformGizmo->CurrentCoordinateSystem = EToolContextCoordinateSystem::Local;
+}
+
+void UDataflowBoneManipulator::Shutdown(UInteractiveToolManager* ToolManager)
+{
+	if (TransformProxy)
+	{
+		TransformProxy->OnTransformChanged.RemoveAll(this);
+		TransformProxy->OnEndTransformEdit.RemoveAll(this);
+	}
+
+	if (ensure(ToolManager))
+	{
+		ToolManager->GetPairedGizmoManager()->DestroyAllGizmosByOwner(this);
+	}
+	TransformGizmo = nullptr;
+}
+
+void UDataflowBoneManipulator::SetEnabled(bool bInEnabled)
+{
+	if (bEnabled == bInEnabled)
+	{
+		return;
+	}
+	bEnabled = bInEnabled;
+
+	// update gizmo accordingly
+	SetSelectedBoneByIndex(BoneIndex);
+}
+
+void UDataflowBoneManipulator::SetSelectedBoneByName(FName BoneName)
+{
+	if (RefSkeletonPoser)
+	{
+		const FReferenceSkeleton& RefSkeleton = RefSkeletonPoser->GetRefSkeleton();
+		BoneIndex = RefSkeleton.FindBoneIndex(BoneName);
+		SetSelectedBoneByIndex(BoneIndex);
+	}
+}
+
+void UDataflowBoneManipulator::SetSelectedBoneByIndex(int32 InBoneIndex)
+{
+	if (RefSkeletonPoser)
+	{
+		const FReferenceSkeleton& RefSkeleton = RefSkeletonPoser->GetRefSkeleton();
+		BoneIndex = RefSkeleton.IsValidIndex(InBoneIndex) ? InBoneIndex : INDEX_NONE;
+		
+		if (TransformGizmo && TransformProxy)
+		{
+			const bool bValidBone = (BoneIndex != INDEX_NONE);
+			TransformGizmo->SetVisibility(bValidBone && bEnabled);
+			if (bValidBone)
+			{
+				const FTransform& BoneTransform = RefSkeletonPoser->GetComponentSpaceTransform(BoneIndex);
+				TransformGizmo->ReinitializeGizmoTransform(BoneTransform, /*bKeepGizmoUnscaled*/true);
+			}
+		}
+	}
+}
+
+void UDataflowBoneManipulator::OnStartTransformEdit(UTransformProxy*)
+{
+	if (RefSkeletonPoser && BoneIndex != INDEX_NONE)
+	{
+		StartTransform = RefSkeletonPoser->GetComponentSpaceTransform(BoneIndex);
+		RefSkeletonPoser->BeginPoseChange();
+	}
+}
+
+void UDataflowBoneManipulator::OnTransformUpdated(UTransformProxy*, FTransform NewTransform)
+{
+	if (BoneIndex != INDEX_NONE)
+	{
+		if (RefSkeletonPoser && RefSkeletonPoser->IsRecordingPoseChange())
+		{
+			RefSkeletonPoser->ModifyBoneAdditiveTransform(BoneIndex,
+				[this, &NewTransform](FTransform& InTransform)
+				{
+					InTransform = NewTransform.GetRelativeTransform(StartTransform);
+				});
+
+			if (OnReferenceSkeletonUpdated.IsBound())
+			{
+				OnReferenceSkeletonUpdated.Broadcast(*this);
+			}
+		}
+	}
+}
+
+void UDataflowBoneManipulator::OnEndTransformEdit(UTransformProxy*)
+{
+	if (RefSkeletonPoser && RefSkeletonPoser->IsRecordingPoseChange())
+	{
+		RefSkeletonPoser->EndPoseChange();
+	}
+}
\ No newline at end of file
diff --git a/Engine/Plugins/Experimental/Dataflow/Source/DataflowEditor/Public/DataflowEditorTools/DataflowEditorBoneManipulator.h b/Engine/Plugins/Experimental/Dataflow/Source/DataflowEditor/Public/DataflowEditorTools/DataflowEditorBoneManipulator.h
new file mode 100644
--- /dev/null
+++ b/Engine/Plugins/Experimental/Dataflow/Source/DataflowEditor/Public/DataflowEditorTools/DataflowEditorBoneManipulator.h
@@ -0,0 +1,60 @@
+// Copyright Epic Games, Inc. All Rights Reserved.
+
+#pragma once
+
+#include "CoreMinimal.h"
+
+#include "DataflowEditorBoneManipulator.generated.h"
+
+#define UE_API DATAFLOWEDITOR_API
+
+struct FReferenceSkeleton;
+
+class UCombinedTransformGizmo;
+class UInteractiveToolManager;
+class URefSkeletonPoser;
+class UTransformProxy;
+
+UCLASS(MinimalAPI)
+class UDataflowBoneManipulator : public UObject
+{
+	GENERATED_BODY()
+
+public:
+	void Setup(UInteractiveToolManager* ToolManager, const FReferenceSkeleton& RefSkeleton);
+	void Shutdown(UInteractiveToolManager* ToolManager);
+
+	void SetSelectedBoneByName(FName BoneName);
+	void SetSelectedBoneByIndex(int32 BoneIndex);
+
+	URefSkeletonPoser* GetRefSkeletonPoser() const { return RefSkeletonPoser; }
+
+	void SetEnabled(bool bEnabled);
+
+	DECLARE_MULTICAST_DELEGATE_OneParam(FOnReferenceSkeletonUpdated, UDataflowBoneManipulator&);
+	FOnReferenceSkeletonUpdated OnReferenceSkeletonUpdated;
+
+private:
+	void OnStartTransformEdit(UTransformProxy*);
+	void OnTransformUpdated(UTransformProxy*, FTransform NewTransform);
+	void OnEndTransformEdit(UTransformProxy*);
+
+	// Source mesh transform gizmo support
+	UPROPERTY(Transient)
+	TObjectPtr<UTransformProxy> TransformProxy;
+
+	UPROPERTY(Transient)
+	TObjectPtr<UCombinedTransformGizmo> TransformGizmo;
+
+	UPROPERTY(Transient)
+	TObjectPtr<URefSkeletonPoser> RefSkeletonPoser;
+
+	UPROPERTY(Transient)
+	int32 BoneIndex = INDEX_NONE;
+
+	FTransform StartTransform;
+
+	bool bEnabled = true;
+};
+
+#undef UE_API
