Dataflow : properly draw attention to output log and main editor when needed
- errors in the graph will draw attention to the outputlog tab
- clicking on an entry in the log or double clicking on a variable properly open the main graph tab

diff --git a/Engine/Plugins/Experimental/Dataflow/Source/DataflowEditor/Private/Dataflow/DataflowEditorToolkit.cpp b/Engine/Plugins/Experimental/Dataflow/Source/DataflowEditor/Private/Dataflow/DataflowEditorToolkit.cpp
--- a/Engine/Plugins/Experimental/Dataflow/Source/DataflowEditor/Private/Dataflow/DataflowEditorToolkit.cpp
+++ b/Engine/Plugins/Experimental/Dataflow/Source/DataflowEditor/Private/Dataflow/DataflowEditorToolkit.cpp
@@ -1902,6 +1902,10 @@ void FDataflowEditorToolkit::OnContextHasWarning(const FDataflowNode* Node, cons
 void FDataflowEditorToolkit::OnContextHasError(const FDataflowNode* Node, const FDataflowOutput* Output, const FString& Error)
 {
 	LogMessage(EMessageSeverity::Error, Node, Output, Error);
+	if (TSharedPtr<SDockTab> OutputLogDockTab = WeakOutputLogDockTab.Pin())
+	{
+		OutputLogDockTab->DrawAttention();
+	}
 }
 
 void FDataflowEditorToolkit::OnBeginEvaluate()
@@ -1996,14 +2000,10 @@ void FDataflowEditorToolkit::OnOutputLogMessageTokenClicked(const FString TokenS
 		}
 		else
 		{
-			// TODO : try to find a way to focus back on the main graph
-			//			current issue seem to be related to the tabspawner beeing unknown when calling thos method
-			//			maybe a better way would be to leverage the subgraph spawner and show the main graph in the subgraph list 
-			//			so that all graphs are spawned and treated the same way ( but this may be an issue to keep the pointer of the main graph )
-			//if (TSharedPtr<SDockTab> MainGraphTab = FGlobalTabmanager::Get()->TryInvokeTab(GraphCanvasTabId))
-			//{
-			//	MainGraphTab->ActivateInParent(SetDirectly);
-			//}
+			if (GraphEditorTab)
+			{
+				GraphEditorTab->DrawAttention();
+			}
 		}
 
 		// Clear node selection
@@ -2644,6 +2644,8 @@ TSharedRef<SDockTab> FDataflowEditorToolkit::SpawnTab_OutputLog(const FSpawnTabA
 
 	DockableTab->SetOnTabClosed(SDockTab::FOnTabClosedCallback::CreateRaw(this, &FDataflowEditorToolkit::OnTabClosed));
 
+	WeakOutputLogDockTab = DockableTab;
+
 	return DockableTab;
 }
 
@@ -2976,14 +2978,10 @@ FGuid FDataflowEditorToolkit::FocusOnNextVariableNode(FName VariableName, const
 					}
 					else
 					{
-						// TODO : try to find a way to focus back on the main graph
-						//			current issue seem to be related to the tabspawner beeing unknown when calling thos method
-						//			maybe a better way would be to leverage the subgraph spawner and show the main graph in the subgraph list 
-						//			so that all graphs are spawned and treated the same way ( but this may be an issue to keep the pointer of the main graph )
-						//if (TSharedPtr<SDockTab> MainGraphTab = FGlobalTabmanager::Get()->TryInvokeTab(GraphCanvasTabId))
-						//{
-						//	MainGraphTab->ActivateInParent(SetDirectly);
-						//}
+						if (GraphEditorTab)
+						{
+							GraphEditorTab->DrawAttention();
+						}
 					}
 					if (ActiveEditor)
 					{
diff --git a/Engine/Plugins/Experimental/Dataflow/Source/DataflowEditor/Public/Dataflow/DataflowEditorToolkit.h b/Engine/Plugins/Experimental/Dataflow/Source/DataflowEditor/Public/Dataflow/DataflowEditorToolkit.h
--- a/Engine/Plugins/Experimental/Dataflow/Source/DataflowEditor/Public/Dataflow/DataflowEditorToolkit.h
+++ b/Engine/Plugins/Experimental/Dataflow/Source/DataflowEditor/Public/Dataflow/DataflowEditorToolkit.h
@@ -388,6 +388,8 @@ private:
 	/** Document tracker for dynamic tabs ( like subgraphs ) */
 	TSharedPtr<FDocumentTracker> DocumentManager;
 
+	TWeakPtr<SDockTab> WeakOutputLogDockTab;
+
 	TSet<FGuid> NodesToEvaluateOnTick;
 
 	FDateTime GraphEvaluationBegin;
