Dataflow bone manipulator : fix missing delegate removal

diff --git a/Engine/Plugins/Experimental/Dataflow/Source/DataflowEditor/Private/DataflowEditorTools/DataflowEditorBoneManipulator.cpp b/Engine/Plugins/Experimental/Dataflow/Source/DataflowEditor/Private/DataflowEditorTools/DataflowEditorBoneManipulator.cpp
--- a/Engine/Plugins/Experimental/Dataflow/Source/DataflowEditor/Private/DataflowEditorTools/DataflowEditorBoneManipulator.cpp
+++ b/Engine/Plugins/Experimental/Dataflow/Source/DataflowEditor/Private/DataflowEditorTools/DataflowEditorBoneManipulator.cpp
@@ -50,6 +50,7 @@ void UDataflowBoneManipulator::Shutdown(UInteractiveToolManager* ToolManager)
 {
 	if (TransformProxy)
 	{
+		TransformProxy->OnBeginTransformEdit.RemoveAll(this);
 		TransformProxy->OnTransformChanged.RemoveAll(this);
 		TransformProxy->OnEndTransformEdit.RemoveAll(this);
 	}
