Dataflow : Update EditSkinWeightTool
- Add support for dynamic mesh component target
- Add bone manipulator
- Check if the tool can start and output errors to the log is it cannot
- Also avoid crashing if there's no bones in the target when opening the tool in EditSkinWeightEditor
- Only enable the accept button when the weights have actually changed

diff --git a/Engine/Plugins/Experimental/Dataflow/Source/DataflowEditor/Private/DataflowEditorTools/DataflowEditorSkinWeightsPaintTool.cpp b/Engine/Plugins/Experimental/Dataflow/Source/DataflowEditor/Private/DataflowEditorTools/DataflowEditorSkinWeightsPaintTool.cpp
--- a/Engine/Plugins/Experimental/Dataflow/Source/DataflowEditor/Private/DataflowEditorTools/DataflowEditorSkinWeightsPaintTool.cpp
+++ b/Engine/Plugins/Experimental/Dataflow/Source/DataflowEditor/Private/DataflowEditorTools/DataflowEditorSkinWeightsPaintTool.cpp
@@ -8,6 +8,7 @@
 #include "Dataflow/DataflowObjectInterface.h"
 #include "Dataflow/DataflowSNode.h"
 #include "Dataflow/DataflowEditorUtil.h"
+#include "DynamicMesh/MeshNormals.h"
 #include "Engine/World.h"
 #include "InteractiveToolManager.h"
 #include "ModelingToolTargetUtil.h"
@@ -15,6 +16,10 @@
 #include "DynamicMesh/DynamicVertexSkinWeightsAttribute.h"
 #include "DynamicMesh/NonManifoldMappingSupport.h"
 #include "Misc/ScopedSlowTask.h"
+#include "SkeletalMeshOperations.h"
+#include "SkeletalMesh/RefSkeletonPoser.h"
+#include "SkeletalMesh/SkeletalMeshToolsHelper.h"
+#include "TargetInterfaces/SkeletonProvider.h"
 #include "TargetInterfaces/MeshDescriptionCommitter.h"
 #include "TargetInterfaces/MeshDescriptionProvider.h"
 #include "TargetInterfaces/PrimitiveComponentBackedTarget.h"
@@ -25,6 +30,8 @@ using namespace UE::Geometry;
 
 #define LOCTEXT_NAMESPACE "UDataflowEditorSkinWeightsPaintTool"
 
+//////////////////////////////////////////////////////////////////////////////////////////////////////////////
+
 void UDataflowEditorSkinWeightsPaintToolBuilder::GetSupportedConstructionViewModes(const UDataflowContextObject& ContextObject, TArray<const UE::Dataflow::IDataflowConstructionViewMode*>& Modes) const
 {}
 
@@ -46,17 +53,56 @@ bool UDataflowEditorSkinWeightsPaintToolBuilder::CanBuildTool(const FToolBuilder
 		return false;
 	};
 
+	auto HasSkinWeightsAndBonesData = [](const UActorComponent* Component) -> bool
+		{
+			if (const USkeletalMeshComponent* SkeletalMeshComponent = Cast<USkeletalMeshComponent>(Component))
+			{
+				if (const USkinnedAsset* SkinnedAsset = SkeletalMeshComponent->GetSkinnedAsset())
+				{
+					return (SkinnedAsset->GetRefSkeleton().GetNum() > 0);
+				}
+			}
+			else if (const UDynamicMeshComponent* DynamicMeshComponent = Cast<UDynamicMeshComponent>(Component))
+			{
+				if (const FDynamicMesh3* DynamicMesh = DynamicMeshComponent->GetMesh())
+				{
+					if (const FDynamicMeshAttributeSet* Attributes = DynamicMesh->Attributes())
+					{
+						const FName DefaultProfileName("Default");
+						const bool bHasSkinWeights = Attributes->HasSkinWeightsAttribute(DefaultProfileName);
+						const int32 NumBones = Attributes->GetNumBones();
+						return bHasSkinWeights && (NumBones > 0);
+					}
+				}
+			}
+			return false;
+		};
+
 	if (USkinWeightsPaintToolBuilder::CanBuildTool(SceneState))
 	{
-		if (SceneState.SelectedComponents.Num() == 1 && SceneState.SelectedComponents[0]->IsA<USkeletalMeshComponent>())
+		if (SceneState.SelectedComponents.Num() == 1)
 		{
-			if (UDataflowBaseContent* ContextObject = SceneState.ToolManager->GetContextObjectStore()->FindContext<UDataflowBaseContent>())
+			const bool bIsASkeletalMesh = SceneState.SelectedComponents[0]->IsA<USkeletalMeshComponent>();
+			const bool bIsADynamicMesh  = SceneState.SelectedComponents[0]->IsA<UDynamicMeshComponent>();
+			if (bIsASkeletalMesh || bIsADynamicMesh)
 			{
-				if (const TSharedPtr<UE::Dataflow::FEngineContext> EvaluationContext = ContextObject->GetDataflowContext())
+				if (UDataflowBaseContent* ContextObject = SceneState.ToolManager->GetContextObjectStore()->FindContext<UDataflowBaseContent>())
 				{
-					if (const FDataflowNode* PrimarySelection = ContextObject->GetSelectedNodeOfType<FDataflowCollectionEditSkinWeightsNode>())
+					if (const TSharedPtr<UE::Dataflow::FEngineContext> EvaluationContext = ContextObject->GetDataflowContext())
 					{
-						return HasManagedArrayCollection(PrimarySelection, EvaluationContext);
+						if (const FDataflowNode* SelectedNode = ContextObject->GetSelectedNodeOfType<FDataflowCollectionEditSkinWeightsNode>())
+						{
+							const bool bIsCollectionOutput = HasManagedArrayCollection(SelectedNode, EvaluationContext);
+							const bool bHasSkinWeightsAndBoneData = HasSkinWeightsAndBonesData(SceneState.SelectedComponents[0]);
+
+							if (!bHasSkinWeightsAndBoneData)
+							{
+								const FString Message = FString::Printf(TEXT("Cannot open the Edit Skin Weight Tool, the mesh is missing skin weights / bones data"));
+								EvaluationContext->Error(Message, SelectedNode);
+							}
+
+							return bIsCollectionOutput && bHasSkinWeightsAndBoneData;
+						}
 					}
 				}
 			}
@@ -90,47 +136,68 @@ UMeshSurfacePointTool* UDataflowEditorSkinWeightsPaintToolBuilder::CreateNewTool
 	return PaintTool;
 }
 
+//////////////////////////////////////////////////////////////////////////////////////////////////////////////
 
 void UDataflowEditorSkinWeightsPaintTool::SetSkinWeightNode(FDataflowCollectionEditSkinWeightsNode* Node)
 {
 	if (SkinWeightNode)
 	{
-		SkinWeightNode->OnBoneSelectionChanged.Remove(OnBoneSelectionChangedDelgateHandle);
+		SkinWeightNode->OnBoneSelectionChanged.Remove(OnBoneSelectionChangedDelegateHandle);
 	}
 
 	SkinWeightNode = Node;
 
 	if (SkinWeightNode)
 	{
-		OnBoneSelectionChangedDelgateHandle = SkinWeightNode->OnBoneSelectionChanged.AddUObject(this, &UDataflowEditorSkinWeightsPaintTool::HandleOnBoneSelectionChanged);
+		OnBoneSelectionChangedDelegateHandle = SkinWeightNode->OnBoneSelectionChanged.AddUObject(this, &UDataflowEditorSkinWeightsPaintTool::HandleOnBoneSelectionChanged);
 	}
 }
 
-void UDataflowEditorSkinWeightsPaintTool::HandleOnBoneSelectionChanged(const TArray<FName>& BoneNames)
+const TArray<FTransform>& UDataflowEditorSkinWeightsPaintTool::GetComponentSpaceBoneTransforms()
 {
-	GetNotifier()->HandleNotification(BoneNames, ESkeletalMeshNotifyType::BonesSelected);
+	if (BoneManipulator)
+	{
+		if (URefSkeletonPoser* Poser = BoneManipulator->GetRefSkeletonPoser())
+		{
+			return Poser->GetComponentSpaceTransforms();
+		}
+	}
+	return Super::GetComponentSpaceBoneTransforms();
 }
 
-void UDataflowEditorSkinWeightsPaintTool::SetDataflowEditorContextObject(TObjectPtr<UDataflowContextObject> InDataflowEditorContextObject)
+void UDataflowEditorSkinWeightsPaintTool::ToggleBoneManipulation(bool bEnable)
 {
-	DataflowEditorContextObject = InDataflowEditorContextObject;
+	if (BoneManipulator)
+	{
+		BoneManipulator->SetEnabled(bEnable);
+	}
+	Super::ToggleBoneManipulation(true);
 }
 
-FMeshDescription* UDataflowEditorSkinWeightsPaintTool::GetCurrentDescription(const int32 LODIndex) const 
+void UDataflowEditorSkinWeightsPaintTool::HandleOnBoneSelectionChanged(const TArray<FName>& BoneNames)
 {
-	if(const IPrimitiveComponentBackedTarget* TargetComponent = Cast<IPrimitiveComponentBackedTarget>(Target))
+	GetNotifier()->HandleNotification(BoneNames, ESkeletalMeshNotifyType::BonesSelected);
+	// this may not be called on the game thread so we need to make sure we post an update on the right thread
+	if (BoneManipulator)
 	{
-		if(USkeletalMeshComponent* SkeletalMeshComponent = Cast<USkeletalMeshComponent>(TargetComponent->GetOwnerComponent()))
-		{
-			if(USkeletalMesh* SkeletalMesh = SkeletalMeshComponent->GetSkeletalMeshAsset())
+		TWeakObjectPtr<UDataflowBoneManipulator> WeakManipulator = BoneManipulator;
+		const FName FirstSelectedName = BoneNames.IsEmpty() ? NAME_None : BoneNames[0];
+	
+		Async(EAsyncExecution::TaskGraphMainTick, 
+		[WeakManipulator, FirstSelectedName]
 			{
-				return SkeletalMesh->GetMeshDescription(LODIndex);
-			}
-		}
+				if (TStrongObjectPtr<UDataflowBoneManipulator> Manipulator = WeakManipulator.Pin())
+				{
+					Manipulator->SetSelectedBoneByName(FirstSelectedName);
+				}
+			});
 	}
-	return nullptr;
 }
 
+void UDataflowEditorSkinWeightsPaintTool::SetDataflowEditorContextObject(TObjectPtr<UDataflowContextObject> InDataflowEditorContextObject)
+{
+	DataflowEditorContextObject = InDataflowEditorContextObject;
+}
 
 int32 UDataflowEditorSkinWeightsPaintTool::GetVertexOffset() const
 {
@@ -143,30 +210,39 @@ int32 UDataflowEditorSkinWeightsPaintTool::GetVertexOffset() const
 				return SkinWeightNode->GetSkeletalMeshOffset(SkeletalMesh);
 			}
 		}
+		else if (UDynamicMeshComponent* DynamicMeshComponent = Cast<UDynamicMeshComponent>(TargetComponent->GetOwnerComponent()))
+		{
+			if (UDynamicMesh* DynamicMesh = DynamicMeshComponent->GetDynamicMesh())
+			{
+				// for dynamic meshes vertex offset is always 0, because the NonManifoldMappingSupport will return the right offset from the original mesh  
+				return 0;
+			}
+		}
 	}
 	return INDEX_NONE;
 }
 
-bool UDataflowEditorSkinWeightsPaintTool::ExtractSkinWeights(TArray<TArray<int32>>& CurrentIndices, TArray<TArray<float>>& CurrentWeights)
+
+bool UDataflowEditorSkinWeightsPaintTool::ExtractSkinWeights(TArray<TArray<int32>>& OutCurrentIndices, TArray<TArray<float>>& OutCurrentWeights, TArray<TArray<int32>>& OutSetupIndices, TArray<TArray<float>>& OutSetupWeights)
 {
 	// Setup DynamicMeshToWeight conversion and get Input weight map (if it exists)
-	if(DataflowEditorContextObject && SkinWeightNode)
+	if (DataflowEditorContextObject && SkinWeightNode)
 	{
 		// Find the map if it exists.
 		if (TSharedPtr<UE::Dataflow::FEngineContext> DataflowContext = DataflowEditorContextObject->GetDataflowContext())
 		{
-			if(DataflowEditorContextObject->GetSelectedCollection().IsValid())
+			if (DataflowEditorContextObject->GetSelectedCollection().IsValid())
 			{
 				// Fill the attribute values
 				SkinWeightNode->FillAttributeWeights(*DataflowEditorContextObject->GetSelectedCollection().Get(),
-					SkinWeightNode->GetBoneIndicesKey(*DataflowContext), SkinWeightNode->GetBoneWeightsKey(*DataflowContext), 
-					SetupIndices, SetupWeights);
-						
-				CurrentWeights.SetNumZeroed(SetupWeights.Num()); 
-				CurrentIndices.SetNumZeroed(SetupIndices.Num());
-						
-				SkinWeightNode->ExtractVertexWeights(*DataflowContext, SetupIndices, SetupWeights,
-					TArrayView<TArray<int32>>(CurrentIndices), TArrayView<TArray<float>>(CurrentWeights));
+					SkinWeightNode->GetBoneIndicesKey(*DataflowContext), SkinWeightNode->GetBoneWeightsKey(*DataflowContext),
+					OutSetupIndices, OutSetupWeights);
+
+				OutCurrentWeights.SetNumZeroed(OutSetupWeights.Num());
+				OutCurrentIndices.SetNumZeroed(OutSetupIndices.Num());
+
+				SkinWeightNode->ExtractVertexWeights(*DataflowContext, OutSetupIndices, OutSetupWeights,
+					MakeArrayView(OutCurrentIndices), MakeArrayView(OutCurrentWeights));
 
 				return true;
 			}
@@ -177,62 +253,47 @@ bool UDataflowEditorSkinWeightsPaintTool::ExtractSkinWeights(TArray<TArray<int32
 
 void UDataflowEditorSkinWeightsPaintTool::Setup()
 {
-	UE::Dataflow::FScopedProgressNotification ProgressNotification(LOCTEXT("SetupProgressMsg", "Starting Dataflow Skin Weight Paint Tool..."), 100, /* delay */ 0);
-	ProgressNotification.SetProgress(0);
+	USkinWeightsPaintTool::Setup();
 
-	const int32 VertexOffset = GetVertexOffset();
-	if(VertexOffset != INDEX_NONE)
-	{
-		if (FMeshDescription* MeshDescription = GetCurrentDescription(0))
-		{
-			const int32 NumVertices = MeshDescription->Vertices().Num();
-			
-			TArray<TArray<float>> CurrentWeights;
-			TArray<TArray<int32>> CurrentIndices;
-	
-			ProgressNotification.AddProgress(0.25f);
+	ISkeletonProvider* SkeletonProvider = CastChecked<ISkeletonProvider>(Target);
+	const FReferenceSkeleton& RefSkeleton = SkeletonProvider->GetSkeleton();
 
-			if(ExtractSkinWeights(CurrentIndices, CurrentWeights))
+	ensure(BoneManipulator == nullptr);
+	BoneManipulator = NewObject<UDataflowBoneManipulator>();
+	BoneManipulator->Setup(GetToolManager(), RefSkeleton);
+	BoneManipulator->OnReferenceSkeletonUpdated.AddWeakLambda(this, 
+		[this](UDataflowBoneManipulator& Manipulator) 
+		{
+			if (URefSkeletonPoser* Poser = Manipulator.GetRefSkeletonPoser())
 			{
-				FSkeletalMeshAttributes MeshAttribs(*MeshDescription);
-				FSkinWeightsVertexAttributesRef SkinWeights =
-					MeshAttribs.GetVertexSkinWeights(FSkeletalMeshAttributesShared::DefaultSkinWeightProfileName);
-			
-				UE::AnimationCore::FBoneWeightsSettings WeightsSettings;
-				WeightsSettings.SetNormalizeType(UE::AnimationCore::EBoneWeightNormalizeType::Always);
-
-				ProgressNotification.AddProgress(0.25f);
-				
-				// process vertices weights 
-				for (int32 VertexIndex = 0; VertexIndex < NumVertices; ++VertexIndex)
+				PoseMesh(Poser);
+				if (SkinWeightNode)
 				{
-					TArray<UE::AnimationCore::FBoneWeight> BoneWeights;
-					const int32 CurrentIndex = VertexOffset + VertexIndex;
-
-					if (CurrentIndex < CurrentWeights.Num() && CurrentIndex < CurrentIndices.Num() &&
-						CurrentWeights[CurrentIndex].Num() == CurrentIndices[CurrentIndex].Num())
+					if (TRefCountPtr<FDataflowDebugDrawSkeletonObject> DebugDrawSkeleton = SkinWeightNode->GetDebugDrawSkeleton())
 					{
-						BoneWeights.Reserve(CurrentWeights[CurrentIndex].Num());
-						for (int32 WeightIndex = 0, NumWeights = CurrentWeights[CurrentIndex].Num(); WeightIndex < NumWeights; ++WeightIndex)
-						{
-							BoneWeights.Add(UE::AnimationCore::FBoneWeight(CurrentIndices[CurrentIndex][WeightIndex],
-								CurrentWeights[CurrentIndex][WeightIndex]));
-						}
+						DebugDrawSkeleton->OverrideBoneTransforms(Poser->GetComponentSpaceTransforms());
 					}
-					SkinWeights.Set(FVertexID(VertexIndex), UE::AnimationCore::FBoneWeights::Create(BoneWeights, WeightsSettings));
 				}
 			}
-		}
-	}
+		});
 
-	ProgressNotification.AddProgress(0.25f);
-	USkinWeightsPaintTool::Setup();
+	if (WeightToolProperties)
+	{
+		ToggleBoneManipulation(WeightToolProperties->EditingMode == EWeightEditMode::Bones);
+	}
 
 	// enable back face option 
 	if (WeightToolProperties)
 	{
 		WeightToolProperties->bShowHitBackFaces = true;
 	}
+
+	// detect changes so that the accept button can be enabled only after a change has be made
+	OnWeightsChanged.AddWeakLambda(this,
+		[this]()
+		{
+			bWeightsHaveChanged = true;
+		});
 }
 
 void UDataflowEditorSkinWeightsPaintTool::OnShutdown(EToolShutdownType ShutdownType)
@@ -263,63 +324,95 @@ void UDataflowEditorSkinWeightsPaintTool::OnShutdown(EToolShutdownType ShutdownT
 
 			TArray<TArray<int32>> CurrentIndices;
 			TArray<TArray<float>> CurrentWeights;
+			TArray<TArray<int32>> SetupIndices;
+			TArray<TArray<float>> SetupWeights;
 
-			CurrentWeights.SetNumZeroed(SetupWeights.Num()); 
-			CurrentIndices.SetNumZeroed(SetupIndices.Num());
-					
-			SkinWeightNode->ExtractVertexWeights(*DataflowContext, SetupIndices, SetupWeights,
-				TArrayView<TArray<int32>>(CurrentIndices), TArrayView<TArray<float>>(CurrentWeights));
-			
-			// profile to edit
-			const FName ActiveProfile = WeightToolProperties->GetActiveSkinWeightProfile();
+			if (ExtractSkinWeights(CurrentIndices, CurrentWeights, SetupIndices, SetupWeights))
+			{
+				// profile to edit
+				const FName ActiveProfile = WeightToolProperties->GetActiveSkinWeightProfile();
 
-			FDynamicMeshVertexSkinWeightsAttribute* SkinWeights = EditedMesh.Attributes()->GetSkinWeightsAttribute(ActiveProfile);
+				FDynamicMeshVertexSkinWeightsAttribute* SkinWeights = EditedMesh.Attributes()->GetSkinWeightsAttribute(ActiveProfile);
 
-			FNonManifoldMappingSupport NonManifoldMappingSupport(EditedMesh);
-			for (int32 DynaMeshVertexIndex : EditedMesh.VertexIndicesItr())
-			{
-				int32 VertexIndex = NonManifoldMappingSupport.GetOriginalNonManifoldVertexID(DynaMeshVertexIndex);
-				
-				UE::AnimationCore::FBoneWeights BoneWeights;
-				SkinWeights->GetValue(DynaMeshVertexIndex, BoneWeights);
-				
-				const int32 CurrentIndex = VertexOffset+VertexIndex;
-			
-				if(CurrentIndex < CurrentWeights.Num() && CurrentIndex < CurrentIndices.Num() &&
-					CurrentWeights[CurrentIndex].Num() == CurrentIndices[CurrentIndex].Num())
+				FNonManifoldMappingSupport NonManifoldMappingSupport(EditedMesh);
+				for (int32 DynaMeshVertexIndex : EditedMesh.VertexIndicesItr())
 				{
-					CurrentIndices[CurrentIndex].Reset(BoneWeights.Num());
-					CurrentWeights[CurrentIndex].Reset(BoneWeights.Num());
+					int32 VertexIndex = NonManifoldMappingSupport.GetOriginalNonManifoldVertexID(DynaMeshVertexIndex);
 
-					TArray<TPair<int32,float>> SortedWeights;
-					SortedWeights.Reserve(BoneWeights.Num());
-						
-					for(int32 WeightIndex = 0, NumWeights = BoneWeights.Num(); WeightIndex < NumWeights; ++WeightIndex)
-					{
-						SortedWeights.Add({BoneWeights[WeightIndex].GetBoneIndex(), BoneWeights[WeightIndex].GetWeight()});
-					}
-					Algo::Sort(SortedWeights, [](const TPair<int32,float>& A, const TPair<int32,float>& B)
-					{
-						return A.Value > B.Value; 
-					});
+					UE::AnimationCore::FBoneWeights BoneWeights;
+					SkinWeights->GetValue(DynaMeshVertexIndex, BoneWeights);
 
-					for(int32 WeightIndex = 0, NumWeights = BoneWeights.Num(); WeightIndex < NumWeights; ++WeightIndex)
+					const int32 CurrentIndex = VertexOffset + VertexIndex;
+
+					if (CurrentIndex < CurrentWeights.Num() && CurrentIndex < CurrentIndices.Num() &&
+						CurrentWeights[CurrentIndex].Num() == CurrentIndices[CurrentIndex].Num())
 					{
-						CurrentIndices[CurrentIndex].Add(SortedWeights[WeightIndex].Key);
-						CurrentWeights[CurrentIndex].Add(SortedWeights[WeightIndex].Value);
+						CurrentIndices[CurrentIndex].Reset(BoneWeights.Num());
+						CurrentWeights[CurrentIndex].Reset(BoneWeights.Num());
+
+						TArray<TPair<int32, float>> SortedWeights;
+						SortedWeights.Reserve(BoneWeights.Num());
+
+						for (int32 WeightIndex = 0, NumWeights = BoneWeights.Num(); WeightIndex < NumWeights; ++WeightIndex)
+						{
+							SortedWeights.Add({ BoneWeights[WeightIndex].GetBoneIndex(), BoneWeights[WeightIndex].GetWeight() });
+						}
+						Algo::Sort(SortedWeights, [](const TPair<int32, float>& A, const TPair<int32, float>& B)
+							{
+								return A.Value > B.Value;
+							});
+
+						for (int32 WeightIndex = 0, NumWeights = BoneWeights.Num(); WeightIndex < NumWeights; ++WeightIndex)
+						{
+							CurrentIndices[CurrentIndex].Add(SortedWeights[WeightIndex].Key);
+							CurrentWeights[CurrentIndex].Add(SortedWeights[WeightIndex].Value);
+						}
 					}
-				}	
-			}
-			
-			SkinWeightNode->ReportVertexWeights(*DataflowContext, SetupIndices, SetupWeights, CurrentIndices, CurrentWeights);
-			SkinWeightNode->Invalidate();
+				}
 
-			// Avoid rebuilding the skeletal mesh after updating the skin weights
-			SkinWeightNode->ValidateSkeletalMeshes();
+				SkinWeightNode->ReportVertexWeights(*DataflowContext, SetupIndices, SetupWeights, CurrentIndices, CurrentWeights);
+				SkinWeightNode->Invalidate();
+
+				// Avoid rebuilding the skeletal mesh after updating the skin weights
+				SkinWeightNode->ValidateSkeletalMeshes();
+			}
 		}
 	}
 
 	SetSkinWeightNode(nullptr);
+
+	if (BoneManipulator)
+	{
+		BoneManipulator->Shutdown(GetToolManager());
+	}
+}
+
+void UDataflowEditorSkinWeightsPaintTool::PoseMesh(URefSkeletonPoser* Poser)
+{
+	const TArray<FTransform>& ComponentSpaceTransforms = Poser->GetComponentSpaceTransforms();
+	if (ComponentSpaceTransforms.IsEmpty())
+	{
+		return;
+	}
+
+	const TArray<FTransform>& ComponentSpaceTransformsRefPose = Poser->GetComponentSpaceTransformsRefPose();
+	TArray<FMatrix> BoneMatrices = SkeletalMeshToolsHelper::ComputeBoneMatrices(ComponentSpaceTransformsRefPose, ComponentSpaceTransforms);
+
+	auto DeformPreviewMeshFunc = [this, &BoneMatrices](FDynamicMesh3& VisualMesh)
+		{
+			auto WriteFunc = [&VisualMesh](int32 VertID, const FVector& PosedVertPos)
+				{
+					VisualMesh.SetVertex(VertID, PosedVertPos);
+				};
+
+			const TMap<FName, float> MorphTargetWeights;
+			SkeletalMeshToolsHelper::GetPosedMesh(WriteFunc, EditedMesh, BoneMatrices, NAME_None, MorphTargetWeights);
+			UE::Geometry::FMeshNormals::QuickRecomputeOverlayNormals(VisualMesh);
+		};
+
+	constexpr bool bRebuildSpatial = false;
+	PreviewMesh->DeferredEditMesh(DeformPreviewMeshFunc, bRebuildSpatial);
+	PreviewMesh->NotifyDeferredEditCompleted(UPreviewMesh::ERenderUpdateMode::FastUpdate, EMeshRenderAttributeFlags::Positions | EMeshRenderAttributeFlags::VertexNormals, bRebuildSpatial);
 }
 
 #undef LOCTEXT_NAMESPACE
diff --git a/Engine/Plugins/Experimental/Dataflow/Source/DataflowEditor/Public/DataflowEditorTools/DataflowEditorSkinWeightsPaintTool.h b/Engine/Plugins/Experimental/Dataflow/Source/DataflowEditor/Public/DataflowEditorTools/DataflowEditorSkinWeightsPaintTool.h
--- a/Engine/Plugins/Experimental/Dataflow/Source/DataflowEditor/Public/DataflowEditorTools/DataflowEditorSkinWeightsPaintTool.h
+++ b/Engine/Plugins/Experimental/Dataflow/Source/DataflowEditor/Public/DataflowEditorTools/DataflowEditorSkinWeightsPaintTool.h
@@ -5,6 +5,7 @@
 #include "CoreMinimal.h"
 #include "HAL/Platform.h"
 #include "DataflowEditorTools/DataflowEditorToolBuilder.h"
+#include "DataflowEditorTools/DataflowEditorBoneManipulator.h"
 #include "SkeletalMesh/SkinWeightsPaintTool.h"
 
 #include "DataflowEditorSkinWeightsPaintTool.generated.h"
@@ -15,6 +16,9 @@ class UDataflowContextObject;
 class UDataflowEditorMode;
 struct FDataflowCollectionEditSkinWeightsNode;
 
+class URefSkeletonPoser;
+struct FReferenceSkeleton;
+
 /**
  * Dataflow skin weights tool builder
  */
@@ -38,7 +42,6 @@ protected:
 	UE_API virtual const FToolTargetTypeRequirements& GetTargetRequirements() const override;
 };
 
-
 /**
  * Dataflow skin weights painting tool
  */
@@ -58,32 +61,33 @@ private:
 	virtual FEditorViewportClient* GetViewportClient() const override {return nullptr;}
 	UE_API virtual void Setup() override;
 	UE_API virtual void OnShutdown(EToolShutdownType ShutdownType) override;
+	UE_API virtual bool CanAccept() const override { return bWeightsHaveChanged; }
+	UE_API virtual void ToggleBoneManipulation(bool bEnable);
+	UE_API virtual const TArray<FTransform>& GetComponentSpaceBoneTransforms();
 	//~ End USkinWeightsPaintTool interface
 
 	void HandleOnBoneSelectionChanged(const TArray<FName>& BoneNames);
+	void HandleOnUpdateRefSkeletonToDraw(FReferenceSkeleton& InOutRefSkeletonToUpdate);
+
+	/** extract the initial weights and indices values as well as the delta values */
+	bool ExtractSkinWeights(TArray<TArray<int32>>& OutCurrentIndices, TArray<TArray<float>>& OutCurrentWeights, TArray<TArray<int32>>& OutSetupIndices, TArray<TArray<float>>& OutSetupWeights);
+
+	void PoseMesh(URefSkeletonPoser* Poser);
 
-	/** Extract the skin weights from the node */
-	UE_API bool ExtractSkinWeights(TArray<TArray<int32>>& CurrentIndices, TArray<TArray<float>>& CurrentWeights);
-	
 	/** Get the skeletal mesh vertex offset in the node collection */
-	UE_API int32 GetVertexOffset() const;
+	int32 GetVertexOffset() const;
 
-	/** Get the LOD description given a  LOD index */
-	UE_API FMeshDescription* GetCurrentDescription(const int32 LODIndex) const;
-	
 	/** Skin weight node associated with that tool */
 	FDataflowCollectionEditSkinWeightsNode* SkinWeightNode = nullptr;
 
-	/** Setup weights used to store the initial values */
-	TArray<TArray<float>> SetupWeights;
-
-	/** Setup indices used to store the initial values */
-	TArray<TArray<int32>> SetupIndices;
-
 	/** Dataflow context object to be used to retrieve the node context data*/
 	TObjectPtr<UDataflowContextObject> DataflowEditorContextObject;
 
-	FDelegateHandle OnBoneSelectionChangedDelgateHandle;
+	TObjectPtr<UDataflowBoneManipulator> BoneManipulator = nullptr;
+	
+	FDelegateHandle OnBoneSelectionChangedDelegateHandle;
+
+	bool bWeightsHaveChanged = false;
 };
 
 
diff --git a/Engine/Plugins/Experimental/Dataflow/Source/DataflowNodes/Private/Dataflow/DataflowCollectionEditSkinWeightsNode.cpp b/Engine/Plugins/Experimental/Dataflow/Source/DataflowNodes/Private/Dataflow/DataflowCollectionEditSkinWeightsNode.cpp
--- a/Engine/Plugins/Experimental/Dataflow/Source/DataflowNodes/Private/Dataflow/DataflowCollectionEditSkinWeightsNode.cpp
+++ b/Engine/Plugins/Experimental/Dataflow/Source/DataflowNodes/Private/Dataflow/DataflowCollectionEditSkinWeightsNode.cpp
@@ -30,7 +30,7 @@
 
 namespace UE::Dataflow::Private
 {
-	bool bEditSkinWeightsNodeGeneratesDynamicMesh = false;
+	bool bEditSkinWeightsNodeGeneratesDynamicMesh = true;
 	FAutoConsoleVariableRef CVarEditSkinWeightsNodeGeneratesDynamicMesh(TEXT("p.Dataflow.EditSkinWeightsNodeGeneratesDynamicMesh"), bEditSkinWeightsNodeGeneratesDynamicMesh, TEXT("When on the edit skin weight node will generate dynamic mesh component instead of the mode expensive skeletal mesh ones"));
 
 
@@ -287,7 +287,8 @@ void FDataflowCollectionEditSkinWeightsNode::AddPrimitiveComponents(UE::Dataflow
 		GeometryCollection::Facades::FRenderingFacade RenderingFacade(*RenderCollection);
 		const int32 NumGeometry = RenderingFacade.IsValid() ? RenderingFacade.NumGeometry() : 0;
 		
-		if (UE::Dataflow::Private::bEditSkinWeightsNodeGeneratesDynamicMesh)
+		const bool bHasSkinWeightsDataInRenderCollection = ((RenderingFacade.FindBoneIndices() != nullptr) && (RenderingFacade.FindBoneWeights() != nullptr));
+		if (UE::Dataflow::Private::bEditSkinWeightsNodeGeneratesDynamicMesh && bHasSkinWeightsDataInRenderCollection)
 		{
 			const bool bNeedsConstruction = (DynamicMeshes.Num() != NumGeometry);
 			if (bNeedsConstruction)
@@ -295,10 +296,18 @@ void FDataflowCollectionEditSkinWeightsNode::AddPrimitiveComponents(UE::Dataflow
 				TObjectPtr<USkeletalMesh> InSkeletalMesh = GetValue(Context, &SkeletalMesh, SkeletalMesh);
 
 				DynamicMeshes.SetNum(NumGeometry);
+				for (int32 MeshIndex = 0; MeshIndex < NumGeometry; ++MeshIndex)
+				{
+					DynamicMeshes[MeshIndex] = NewObject<UDynamicMesh>(RootActor);
+				}
+
 				ParallelFor(DynamicMeshes.Num(),
 					[this, &RenderingFacade, &InSkeletalMesh](int32 GeometryIndex)
 					{
-						FDynamicMesh3& DynamicMesh = DynamicMeshes[GeometryIndex];
+						TObjectPtr<UDynamicMesh> DynamicMeshObject = DynamicMeshes[GeometryIndex];
+						check(DynamicMeshObject && DynamicMeshObject->GetMeshPtr());
+
+						FDynamicMesh3& DynamicMesh = DynamicMeshObject->GetMeshRef();
 						UE::Dataflow::Conversion::RenderingFacadeToDynamicMesh(RenderingFacade, GeometryIndex, DynamicMesh);
 
 						if (InSkeletalMesh)
@@ -330,23 +339,13 @@ void FDataflowCollectionEditSkinWeightsNode::AddPrimitiveComponents(UE::Dataflow
 			}
 			if (!DynamicMeshes.IsEmpty())
 			{
-				// copy the meshes because the dynamic component needs to get the ownership and we don't want to loose the cached ones (DynamicMeshes)
-				TArray<FDynamicMesh3> DynamicMeshesToSet;
-				DynamicMeshesToSet.SetNum(NumGeometry);
-				ParallelFor(DynamicMeshes.Num(),
-					[this, &DynamicMeshesToSet](int32 GeometryIndex)
-					{
-						DynamicMeshesToSet[GeometryIndex] = DynamicMeshes[GeometryIndex];
-					});
-
 				PrimitiveComponents.Reserve(PrimitiveComponents.Num() + NumGeometry);
 				UMaterial* DefaultTwoSidedMaterial = Cast<UMaterial>(StaticLoadObject(UMaterial::StaticClass(), NULL, TEXT("/Engine/EditorMaterials/Dataflow/DataflowTwoSidedVertexMaterial")));
 				for (int32 GeometryIndex = 0; GeometryIndex < NumGeometry; ++GeometryIndex)
 				{
 					const FName DynamicMeshName(FString::Printf(TEXT("[%d]_DynamicMesh_(%s)"), GeometryIndex, *FGuid::NewGuid().ToString()));
 					UDynamicMeshComponent* DynamicMeshComponent = NewObject<UDynamicMeshComponent>(RootActor, DynamicMeshName);
-					FDynamicMesh3 DynamicMeshesToMove = DynamicMeshesToSet[GeometryIndex]; // need to copy to not loose the cached ones
-					DynamicMeshComponent->SetMesh(MoveTemp(DynamicMeshesToMove));
+					DynamicMeshComponent->SetDynamicMesh(DynamicMeshes[GeometryIndex]);
 					DynamicMeshComponent->SetOverrideRenderMaterial(DefaultTwoSidedMaterial);
 					PrimitiveComponents.Add(DynamicMeshComponent);
 				}
@@ -569,6 +568,7 @@ void FDataflowCollectionEditSkinWeightsNode::OnInvalidate()
 {
 	bValidSkeletalMeshes = false;
 	DynamicMeshes.Reset();
+	SkeletonObject = nullptr;
 }
 
 #if WITH_EDITOR
@@ -577,12 +577,14 @@ void FDataflowCollectionEditSkinWeightsNode::DebugDraw(UE::Dataflow::FContext& C
 		IDataflowDebugDrawInterface& DataflowRenderingInterface,
 		const FDebugDrawParameters& DebugDrawParameters) const
 {
+	SkeletonObject = nullptr;
+
 	TObjectPtr<USkeletalMesh> InSkeletalMesh = GetValue(Context, &SkeletalMesh, SkeletalMesh);
 	if(InSkeletalMesh)
 	{
-		TRefCountPtr<FDataflowDebugDrawSkeletonObject> SkeletonObject = MakeDebugDrawObject<FDataflowDebugDrawSkeletonObject>(
+		SkeletonObject = MakeDebugDrawObject<FDataflowDebugDrawSkeletonObject>(
 			DataflowRenderingInterface.ModifyDataflowElements(), InSkeletalMesh->GetRefSkeleton());
-	
+
 		DataflowRenderingInterface.DrawObject(TRefCountPtr<IDataflowDebugDrawObject>(SkeletonObject));
 
 		SkeletonObject->OnBoneSelectionChanged.AddLambda([this](const TArray<FName>& BoneNames)
@@ -618,7 +620,6 @@ int32 FDataflowCollectionEditSkinWeightsNode::GetSkeletalMeshOffset(const TObjec
 	return INDEX_NONE;
 }
 
-
 ////////////////////////////////////////////////////////////////////////////////////////////////////////////////////////////////////////////////////////////
 // 
 // Object encapsulating a change to the vertex skin weights values. Used for Undo/Redo.
diff --git a/Engine/Plugins/Experimental/Dataflow/Source/DataflowNodes/Public/Dataflow/DataflowCollectionEditSkinWeightsNode.h b/Engine/Plugins/Experimental/Dataflow/Source/DataflowNodes/Public/Dataflow/DataflowCollectionEditSkinWeightsNode.h
--- a/Engine/Plugins/Experimental/Dataflow/Source/DataflowNodes/Public/Dataflow/DataflowCollectionEditSkinWeightsNode.h
+++ b/Engine/Plugins/Experimental/Dataflow/Source/DataflowNodes/Public/Dataflow/DataflowCollectionEditSkinWeightsNode.h
@@ -8,7 +8,6 @@
 #include "DataflowSkeletalMeshNodes.h"
 #include "Dataflow/DataflowCollectionAttributeKeyNodes.h"
 #include "Dataflow/DataflowNode.h"
-#include "DynamicMesh/DynamicMesh3.h"
 #include "GeometryCollection/GeometryCollection.h"
 #include "GeometryCollection/ManagedArrayCollection.h"
 #include "GPUSkinPublicDefs.h"
@@ -17,6 +16,8 @@
 
 #define UE_API DATAFLOWNODES_API
 
+class UDynamicMesh;
+
 /** Dataflow skin weights data */
 USTRUCT()
 struct FDataflowSkinWeightData
@@ -158,12 +159,14 @@ public:
 	/** Get the weights attribute key to retrieve/set the weight values*/
     UE_API FCollectionAttributeKey GetBoneWeightsKey(UE::Dataflow::FContext& Context) const;
 
+	TRefCountPtr<FDataflowDebugDrawSkeletonObject> GetDebugDrawSkeleton() const { return SkeletonObject; }
+
 	/** Validate the skeletal mesh construction */
 	void ValidateSkeletalMeshes() {bValidSkeletalMeshes = true;}
 
 	/** Return the collection offset given a skeletal mesh*/
 	UE_API int32 GetSkeletalMeshOffset(const TObjectPtr<USkeletalMesh>& SkeletalMesh) const;
-	
+
 	class FEditNodeToolChange;
 	UE_API TUniquePtr<class FToolCommandChange> MakeEditNodeToolChange();
 
@@ -197,7 +200,14 @@ private:
 	bool bValidSkeletalMeshes = false;
 
 	/** Transient dynamic meshes built from dataflow render collection */
-	TArray<UE::Geometry::FDynamicMesh3> DynamicMeshes;
+	UPROPERTY(Transient)
+	TArray<TObjectPtr<UDynamicMesh>> DynamicMeshes;
+
+	/** 
+	* not ideal but for now this is the only way we can update the displayed skeleton 
+	* todo(dataflow) : find a better way to handle skeleton drawing and manipulation post 5.7
+	*/
+	mutable TRefCountPtr<FDataflowDebugDrawSkeletonObject> SkeletonObject;
 };
 
 #undef UE_API
diff --git a/Engine/Plugins/Runtime/MeshModelingToolset/Source/MeshModelingToolsEditorOnly/Private/SkeletalMesh/SkinWeightsPaintTool.cpp b/Engine/Plugins/Runtime/MeshModelingToolset/Source/MeshModelingToolsEditorOnly/Private/SkeletalMesh/SkinWeightsPaintTool.cpp
--- a/Engine/Plugins/Runtime/MeshModelingToolset/Source/MeshModelingToolsEditorOnly/Private/SkeletalMesh/SkinWeightsPaintTool.cpp
+++ b/Engine/Plugins/Runtime/MeshModelingToolset/Source/MeshModelingToolsEditorOnly/Private/SkeletalMesh/SkinWeightsPaintTool.cpp
@@ -1893,7 +1893,7 @@ void USkinWeightsPaintTool::Setup()
 	RecalculateBrushRadius();
 
 	// default to the root bone as current bone
-	PendingCurrentBone = CurrentBone = RefSkeleton.GetBoneName(0);
+	PendingCurrentBone = CurrentBone = RefSkeleton.IsValidIndex(0)? RefSkeleton.GetBoneName(0) : NAME_None;
 
 	// configure preview mesh
 	PreviewMesh->SetTangentsMode(EDynamicMeshComponentTangentsMode::AutoCalculated);
