Dataflow : fix hardcoded value sent to base class

diff --git a/Engine/Plugins/Experimental/Dataflow/Source/DataflowEditor/Private/DataflowEditorTools/DataflowEditorSkinWeightsPaintTool.cpp b/Engine/Plugins/Experimental/Dataflow/Source/DataflowEditor/Private/DataflowEditorTools/DataflowEditorSkinWeightsPaintTool.cpp
--- a/Engine/Plugins/Experimental/Dataflow/Source/DataflowEditor/Private/DataflowEditorTools/DataflowEditorSkinWeightsPaintTool.cpp
+++ b/Engine/Plugins/Experimental/Dataflow/Source/DataflowEditor/Private/DataflowEditorTools/DataflowEditorSkinWeightsPaintTool.cpp
@@ -171,7 +171,7 @@ void UDataflowEditorSkinWeightsPaintTool::ToggleBoneManipulation(bool bEnable)
 	{
 		BoneManipulator->SetEnabled(bEnable);
 	}
-	Super::ToggleBoneManipulation(true);
+	Super::ToggleBoneManipulation(bEnable);
 }
 
 void UDataflowEditorSkinWeightsPaintTool::HandleOnBoneSelectionChanged(const TArray<FName>& BoneNames)
