Add a way fro each dataflow editors to hide geometry cache properties in the preview menu based on the type of asset
-enable it only for flesh where it is supported

diff --git a/Engine/Plugins/Experimental/ChaosFlesh/Source/ChaosFleshEditor/Private/ChaosFlesh/Asset/AssetDefinition_FleshAsset.cpp b/Engine/Plugins/Experimental/ChaosFlesh/Source/ChaosFleshEditor/Private/ChaosFlesh/Asset/AssetDefinition_FleshAsset.cpp
--- a/Engine/Plugins/Experimental/ChaosFlesh/Source/ChaosFleshEditor/Private/ChaosFlesh/Asset/AssetDefinition_FleshAsset.cpp
+++ b/Engine/Plugins/Experimental/ChaosFlesh/Source/ChaosFleshEditor/Private/ChaosFlesh/Asset/AssetDefinition_FleshAsset.cpp
@@ -80,6 +80,7 @@ EAssetCommandResult UAssetDefinition_FleshAsset::OpenAssets(const FAssetOpenArgs
 					SimulationSettings->bIsSimulationPlayingByDefault = false;
 					SimulationSettings->bIsAsyncCachingSupported = true;
 					SimulationSettings->bIsAsyncCachingEnabledByDefault = true;
+					SimulationSettings->bIsGeometryCacheOutputSupported = true;
 					AssetEditor->AddEditorSettings(SimulationSettings);
 				}
 				
diff --git a/Engine/Plugins/Experimental/Dataflow/Source/DataflowEditor/Private/Dataflow/DataflowSimulationScene.cpp b/Engine/Plugins/Experimental/Dataflow/Source/DataflowEditor/Private/Dataflow/DataflowSimulationScene.cpp
--- a/Engine/Plugins/Experimental/Dataflow/Source/DataflowEditor/Private/Dataflow/DataflowSimulationScene.cpp
+++ b/Engine/Plugins/Experimental/Dataflow/Source/DataflowEditor/Private/Dataflow/DataflowSimulationScene.cpp
@@ -25,6 +25,10 @@
 #include "GeometryCache.h"
 #include "Dataflow/DataflowElement.h"
 
+#include "Framework/Notifications/NotificationManager.h"
+#include "Widgets/Notifications/SNotificationList.h"
+
+
 #if WITH_EDITOR
 #include "Misc/FileHelper.h"
 #endif
@@ -32,6 +36,23 @@
 #include UE_INLINE_GENERATED_CPP_BY_NAME(DataflowSimulationScene)
 #define LOCTEXT_NAMESPACE "FDataflowSimulationScene"
 
+namespace DataflowSimulationScene::Private
+{
+	void ShowErrorToastMessage(const FText& Text, const FText& SubText, const float Duration)
+	{
+		FNotificationInfo ToastInfo(Text);
+		ToastInfo.ExpireDuration = Duration;
+		ToastInfo.bFireAndForget = true;
+		ToastInfo.SubText = SubText;
+		ToastInfo.bUseLargeFont = true;
+
+		TSharedPtr<SNotificationItem> Notification = FSlateNotificationManager::Get().AddNotification(ToastInfo);
+		if (Notification.IsValid())
+		{
+			Notification->SetCompletionState(SNotificationItem::ECompletionState::CS_Fail);
+		}
+	}
+}
 //
 // Simulation Scene
 //
@@ -92,6 +113,8 @@ FDataflowSimulationScene::FDataflowSimulationScene(FPreviewScene::ConstructionVa
 		{
 			SceneDescription->CacheParams.bAsyncCaching = false;
 		}
+
+		SceneDescription->bIsGeometryCacheOutputSupported = Settings->bIsGeometryCacheOutputSupported;
 	}
 }
 
@@ -537,6 +560,16 @@ void FDataflowSimulationScene::SceneDescriptionPropertyChanged(const FName& Prop
 
 void UDataflowSimulationSceneDescription::GenerateGeometryCache()
 {
+	if (!bIsGeometryCacheOutputSupported)
+	{
+		const FText Title = LOCTEXT("DataflowEditor_GenerateGeometryCache_Title", "Dataflow Editor");
+		const FText Message = LOCTEXT("DataflowEditor_GenerateGeometryCache_Message", "Writing to a geometry cache is not supported for this type of asset");
+		constexpr float MessageDuration = 5.f;
+		DataflowSimulationScene::Private::ShowErrorToastMessage(Title, Message, MessageDuration);
+
+		return;
+	}
+
 	SimulationScene->ResetSimulationScene();
 	SimulationScene->CreateSimulationScene();
 	const int32 NumFrames = SimulationScene->GetNumFrames();
@@ -658,6 +691,16 @@ namespace UE::Dataflow::Private
 
 void UDataflowSimulationSceneDescription::NewGeometryCache()
 {
+	if (!bIsGeometryCacheOutputSupported)
+	{
+		const FText Title = LOCTEXT("DataflowEditor_NewGeometryCache_Title", "Dataflow Editor");
+		const FText Message = LOCTEXT("DataflowEditor_NewGeometryCache_Message", "Writing to a geometry cache is not supported for this type of asset");
+		constexpr float MessageDuration = 5.f;
+		DataflowSimulationScene::Private::ShowErrorToastMessage(Title, Message, MessageDuration);
+
+		return;
+	}
+
 	const UObject* const NamingAsset = CacheAsset? CacheAsset.Get() : nullptr;
 	GeometryCacheAsset = UE::Dataflow::Private::NewGeometryCacheDialog(NamingAsset);
 }
diff --git a/Engine/Plugins/Experimental/Dataflow/Source/DataflowEditor/Public/Dataflow/DataflowSimulationScene.h b/Engine/Plugins/Experimental/Dataflow/Source/DataflowEditor/Public/Dataflow/DataflowSimulationScene.h
--- a/Engine/Plugins/Experimental/Dataflow/Source/DataflowEditor/Public/Dataflow/DataflowSimulationScene.h
+++ b/Engine/Plugins/Experimental/Dataflow/Source/DataflowEditor/Public/Dataflow/DataflowSimulationScene.h
@@ -34,6 +34,9 @@ public:
 
 	UPROPERTY()
 	bool bIsAsyncCachingEnabledByDefault = false;
+
+	UPROPERTY()
+	bool bIsGeometryCacheOutputSupported = false;
 };
 
 UCLASS(MinimalAPI)
@@ -74,16 +77,19 @@ public:
 	UPROPERTY(EditAnywhere, Category="Caching")
 	FDataflowPreviewCacheParams CacheParams;
 
+	UPROPERTY()
+	bool bIsGeometryCacheOutputSupported = false;
+
 	/** Geometry cache asset used to extract skeletal mesh results from simulation */
-	UPROPERTY(EditAnywhere, Category = "Geometry Cache", DisplayName="Geometry Cache", meta=(EditCondition = "CacheAsset != nullptr"))
+	UPROPERTY(EditAnywhere, Category = "Geometry Cache", DisplayName="Geometry Cache", meta=(EditCondition = "bIsGeometryCacheOutputSupported && CacheAsset != nullptr", EditConditionHides))
 	TObjectPtr<UGeometryCache> GeometryCacheAsset = nullptr;
 
 	/** Skeletal mesh interpolated from simulation. This should match the SkeletalMesh used in GenerateSurfaceBindings node */
-	UPROPERTY(EditAnywhere, Category = "Geometry Cache", meta = (EditCondition = "CacheAsset != nullptr && EmbeddedStaticMesh == nullptr"))
+	UPROPERTY(EditAnywhere, Category = "Geometry Cache", meta = (EditCondition = "bIsGeometryCacheOutputSupported && CacheAsset != nullptr && EmbeddedStaticMesh == nullptr", EditConditionHides))
 	TObjectPtr<USkeletalMesh> EmbeddedSkeletalMesh = nullptr;
 
 	/** Static mesh interpolated from simulation. This should match the Static mesh used in GenerateSurfaceBindings node */
-	UPROPERTY(EditAnywhere, Category = "Geometry Cache", meta = (EditCondition = "CacheAsset != nullptr && EmbeddedSkeletalMesh == nullptr"))
+	UPROPERTY(EditAnywhere, Category = "Geometry Cache", meta = (EditCondition = "bIsGeometryCacheOutputSupported && CacheAsset != nullptr && EmbeddedSkeletalMesh == nullptr", EditConditionHides))
 	TObjectPtr<UStaticMesh> EmbeddedStaticMesh = nullptr;
 
 	/** Interpolates and saves geometry cache from Chaos cache */
