Disable the dataflow slow task progress notification for now as this is causing UI focus issues ( like for renaming nodes for example )

diff --git a/Engine/Plugins/Experimental/Dataflow/Source/DataflowEditor/Private/Dataflow/DataflowEditorUtil.cpp b/Engine/Plugins/Experimental/Dataflow/Source/DataflowEditor/Private/Dataflow/DataflowEditorUtil.cpp
--- a/Engine/Plugins/Experimental/Dataflow/Source/DataflowEditor/Private/Dataflow/DataflowEditorUtil.cpp
+++ b/Engine/Plugins/Experimental/Dataflow/Source/DataflowEditor/Private/Dataflow/DataflowEditorUtil.cpp
@@ -22,6 +22,12 @@ using namespace UE::Geometry;
 
 namespace UE::Dataflow::Private
 {
+	// disable slow task for the progress tracking as this keep on stealing focus from other operation 
+	// we need to shift to a true asynchronous model for the Dataflow editor so that the game thread is free to update progress toast notification
+	static bool bUseSlowTaskForProgressNotification = false;
+	static FAutoConsoleVariableRef CVarUseSlowTaskForProgressNotification(TEXT("p.Dataflow.UseSlowTaskForProgressNotification"), bUseSlowTaskForProgressNotification, TEXT("When enabled the progress notification class will use the SlowTask system to diaply progress"));
+
+
 	bool HasSkeletalMesh(UObject* InObject)
 	{
 		if (UClass* Class = InObject->GetClass())
@@ -190,21 +196,27 @@ namespace UE::Dataflow
 	FScopedProgressNotification::FScopedProgressNotification(const FText& Title, int32 InMaxProgress, float DelayInSeconds)
 		: MaxProgress(InMaxProgress)
 	{
-		SlowTask = MakeUnique<FSlowTask>(static_cast<float>(MaxProgress), Title);
-		SlowTask->Initialize();
-		if (DelayInSeconds != 0)
-		{
-			SlowTask->MakeDialogDelayed(DelayInSeconds);
-		}
-		else
+		if (Private::bUseSlowTaskForProgressNotification)
 		{
-			SlowTask->MakeDialog(false);
+			SlowTask = MakeUnique<FSlowTask>(static_cast<float>(MaxProgress), Title);
+			SlowTask->Initialize();
+			if (DelayInSeconds != 0)
+			{
+				SlowTask->MakeDialogDelayed(DelayInSeconds);
+			}
+			else
+			{
+				SlowTask->MakeDialog(false);
+			}
 		}
 	}
 
 	FScopedProgressNotification::~FScopedProgressNotification()
 	{
-		SlowTask->Destroy();
+		if (SlowTask)
+		{
+			SlowTask->Destroy();
+		}
 	}
 
 	void FScopedProgressNotification::SetProgress(float InProgress, const FText& Message)
@@ -215,7 +227,10 @@ namespace UE::Dataflow
 		if (UpdateSteps <= 0 || DeltaProgressSinceLastUpdate > UpdateSteps)
 		{
 			const float DeltaProgress = FMath::Max(0, InProgress - CurrentProgress);
-			SlowTask->EnterProgressFrame(DeltaProgress, Message);
+			if (SlowTask)
+			{
+				SlowTask->EnterProgressFrame(DeltaProgress, Message);
+			}
 			LastUpdateProgress = ClampedProgress;
 		}
 		CurrentProgress = ClampedProgress;
