Dataflow : fix cancelled close request causing the preview actor to be deleted and subsequent calls to potential crash the editor
- moved the activate default mode from OnRequestClose to OnClose as OnrequestClose should only answer if the editor can be close and should not take any actions until OnClose is called

diff --git a/Engine/Plugins/Experimental/Dataflow/Source/DataflowEditor/Private/Dataflow/DataflowEditorToolkit.cpp b/Engine/Plugins/Experimental/Dataflow/Source/DataflowEditor/Private/Dataflow/DataflowEditorToolkit.cpp
--- a/Engine/Plugins/Experimental/Dataflow/Source/DataflowEditor/Private/Dataflow/DataflowEditorToolkit.cpp
+++ b/Engine/Plugins/Experimental/Dataflow/Source/DataflowEditor/Private/Dataflow/DataflowEditorToolkit.cpp
@@ -706,14 +706,20 @@ bool FDataflowEditorToolkit::OnRequestClose(EAssetEditorCloseReason InCloseReaso
 		// simply return true because there's nothing left to do.
 		return true;
 	}
+	
+	return FAssetEditorToolkit::OnRequestClose(InCloseReason);
+}
 
+/** Called when this toolkit is being closed */
+void FDataflowEditorToolkit::OnClose()
+{
 	// Give any active modes a chance to shutdown while the toolkit host is still alive
 	// This is super important to do, otherwise currently opened tabs won't be marked as "closed".
 	// This results in tabs not being properly recycled upon reopening the editor and tab
 	// duplication for each opening event.
 	GetEditorModeManager().ActivateDefaultMode();
 
-	return FAssetEditorToolkit::OnRequestClose(InCloseReason);
+	FAssetEditorToolkit::OnClose();
 }
 
 void FDataflowEditorToolkit::PostInitAssetEditor()
diff --git a/Engine/Plugins/Experimental/Dataflow/Source/DataflowEditor/Public/Dataflow/DataflowEditorToolkit.h b/Engine/Plugins/Experimental/Dataflow/Source/DataflowEditor/Public/Dataflow/DataflowEditorToolkit.h
--- a/Engine/Plugins/Experimental/Dataflow/Source/DataflowEditor/Public/Dataflow/DataflowEditorToolkit.h
+++ b/Engine/Plugins/Experimental/Dataflow/Source/DataflowEditor/Public/Dataflow/DataflowEditorToolkit.h
@@ -217,6 +217,7 @@ private:
 
 	// FAssetEditorToolkit interface
 	UE_API virtual bool OnRequestClose(EAssetEditorCloseReason InCloseReason) override;
+	UE_API virtual void OnClose() override;
 	UE_API virtual void PostInitAssetEditor() override;
 	UE_API virtual void GetSaveableObjects(TArray<UObject*>& OutObjects) const override;
 	UE_API virtual void OnAssetsSavedAs(const TArray<UObject*>& SavedObjects) override;
