Unshelved from pending changelist '46933319':

Fix crash when setting source in skin weight tool and when using the edting bones tool

diff --git a/Engine/Plugins/Animation/SkeletalMeshModelingTools/Source/SkeletalMeshModelingTools/Private/SkeletalMeshEditorUtils.cpp b/Engine/Plugins/Animation/SkeletalMeshModelingTools/Source/SkeletalMeshModelingTools/Private/SkeletalMeshEditorUtils.cpp
--- a/Engine/Plugins/Animation/SkeletalMeshModelingTools/Source/SkeletalMeshModelingTools/Private/SkeletalMeshEditorUtils.cpp
+++ b/Engine/Plugins/Animation/SkeletalMeshModelingTools/Source/SkeletalMeshModelingTools/Private/SkeletalMeshEditorUtils.cpp
@@ -128,21 +128,37 @@ void USkeletalMeshEditorContextObject::Init(USkeletalMeshModelingToolsEditorMode
 
 EMeshLODIdentifier USkeletalMeshEditorContextObject::GetEditingLOD()
 {
+	if (!EditorMode.IsValid())
+	{
+		return EMeshLODIdentifier::Default;
+	}
 	return EditorMode->GetEditingLOD();
 }
 
 void USkeletalMeshEditorContextObject::HideSkeleton()
 {
+	if (!EditorMode.IsValid())
+	{
+		return;
+	}
 	EditorMode->HideSkeletonForTool();
 }
 
 void USkeletalMeshEditorContextObject::ShowSkeleton()
 {
+	if (!EditorMode.IsValid())
+	{
+		return;
+	}
 	EditorMode->ShowSkeletonForTool();
 }
 
 void USkeletalMeshEditorContextObject::ToggleBoneManipulation(bool bEnable)
 {
+	if (!EditorMode.IsValid())
+	{
+		return;
+	}
 	EditorMode->ToggleBoneManipulation(bEnable);
 }
 
diff --git a/Engine/Plugins/Animation/SkeletalMeshModelingTools/Source/SkeletalMeshModelingTools/Private/SkeletalMeshGizmoUtils.cpp b/Engine/Plugins/Animation/SkeletalMeshModelingTools/Source/SkeletalMeshModelingTools/Private/SkeletalMeshGizmoUtils.cpp
--- a/Engine/Plugins/Animation/SkeletalMeshModelingTools/Source/SkeletalMeshModelingTools/Private/SkeletalMeshGizmoUtils.cpp
+++ b/Engine/Plugins/Animation/SkeletalMeshModelingTools/Source/SkeletalMeshModelingTools/Private/SkeletalMeshGizmoUtils.cpp
@@ -123,7 +123,7 @@ void USkeletalMeshGizmoContextObject::UnregisterGizmosWithManager(UInteractiveTo
 	bGizmosRegistered = false;
 }
 
-USkeletalMeshGizmoWrapperBase* USkeletalMeshGizmoContextObject::GetNewWrapper(UInteractiveToolManager* InToolManager, UObject* Outer, IGizmoStateTarget* InStateTarget)
+USkeletalMeshGizmoWrapperBase* USkeletalMeshGizmoContextObject::GetNewWrapper(UInteractiveToolManager* InToolManager, UObject* Outer, IGizmoStateTarget* InStateTarget, const TScriptInterface<ITransformGizmoSource>& GizmoSource)
 {
 	if (!bGizmosRegistered)
 	{
@@ -135,8 +135,8 @@ USkeletalMeshGizmoWrapperBase* USkeletalMeshGizmoContextObject::GetNewWrapper(UI
 	{
 		return nullptr;
 	}
-
-	Gizmo->TransformGizmoSource = nullptr;
+	
+	Gizmo->TransformGizmoSource = GizmoSource;
 	Gizmo->SetVisibility(true);
 
 	USkeletonTransformProxy* SkeletonTransformProxy = NewObject<USkeletonTransformProxy>(Outer);
diff --git a/Engine/Plugins/Animation/SkeletalMeshModelingTools/Source/SkeletalMeshModelingTools/Private/SkeletalMeshGizmoUtils.h b/Engine/Plugins/Animation/SkeletalMeshModelingTools/Source/SkeletalMeshModelingTools/Private/SkeletalMeshGizmoUtils.h
--- a/Engine/Plugins/Animation/SkeletalMeshModelingTools/Source/SkeletalMeshModelingTools/Private/SkeletalMeshGizmoUtils.h
+++ b/Engine/Plugins/Animation/SkeletalMeshModelingTools/Source/SkeletalMeshModelingTools/Private/SkeletalMeshGizmoUtils.h
@@ -4,6 +4,7 @@
 
 #include "UObject/Object.h"
 #include "Containers/UnrealString.h"
+#include "EditorGizmos/TransformGizmoInterfaces.h"
 #include "SkeletalMesh/SkeletalMeshEditingInterface.h"
 
 #include "SkeletalMeshGizmoUtils.generated.h"
@@ -88,7 +89,7 @@ public:
 	UE_API void RegisterGizmosWithManager(UInteractiveToolManager* InToolManager);
 	UE_API void UnregisterGizmosWithManager(UInteractiveToolManager* InToolManager);
 
-	UE_API virtual USkeletalMeshGizmoWrapperBase* GetNewWrapper(UInteractiveToolManager* InToolManager, UObject* Outer = nullptr, IGizmoStateTarget* InStateTarget = nullptr) override;
+	UE_API virtual USkeletalMeshGizmoWrapperBase* GetNewWrapper(UInteractiveToolManager* InToolManager, UObject* Outer = nullptr, IGizmoStateTarget* InStateTarget = nullptr, const TScriptInterface<ITransformGizmoSource>& GizmoSource = nullptr) override;
 
 private:
 	
diff --git a/Engine/Plugins/ChaosClothAssetEditor/Source/ChaosClothAssetDataflowNodes/Private/ChaosClothAsset/ChaosClothAssetDataflowNodesModule.cpp b/Engine/Plugins/ChaosClothAssetEditor/Source/ChaosClothAssetDataflowNodes/Private/ChaosClothAsset/ChaosClothAssetDataflowNodesModule.cpp
--- a/Engine/Plugins/ChaosClothAssetEditor/Source/ChaosClothAssetDataflowNodes/Private/ChaosClothAsset/ChaosClothAssetDataflowNodesModule.cpp
+++ b/Engine/Plugins/ChaosClothAssetEditor/Source/ChaosClothAssetDataflowNodes/Private/ChaosClothAsset/ChaosClothAssetDataflowNodesModule.cpp
@@ -292,6 +292,8 @@ namespace UE::Chaos::ClothAsset
 					TArray<FIntVector> Indices;
 					TArray<TArray<FVector2f>> UVs;
 					TArray<int32> RenderMaterialIDs;
+					TArray<TArray<float>> BoneWeights;
+					TArray<TArray<int32>> BoneIndices;
 
 
 					const FName& ViewModeName = State.GetViewMode().GetName();
@@ -303,6 +305,9 @@ namespace UE::Chaos::ClothAsset
 						Vertices = ClothFacade.GetSimPosition3D();
 						Indices = ClothFacade.GetSimIndices3D();
 						Normals = ClothFacade.GetSimNormal();
+
+						BoneWeights = TArray<TArray<float>>(ClothFacade.GetSimBoneWeights());
+						BoneIndices = TArray<TArray<int32>>(ClothFacade.GetSimBoneIndices());
 					}
 					else if (bHasSimData && (ViewModeName == FCloth2DSimViewMode::Name || ViewModeName == Dataflow::FDataflowConstruction2DViewMode::Name))
 					{
@@ -343,6 +348,9 @@ namespace UE::Chaos::ClothAsset
 						Vertices = ClothFacade.GetRenderPosition();
 						Indices = ClothFacade.GetRenderIndices();
 						Normals = ClothFacade.GetRenderNormal();
+
+						BoneWeights = TArray<TArray<float>>(ClothFacade.GetRenderBoneWeights());
+						BoneIndices = TArray<TArray<int32>>(ClothFacade.GetRenderBoneIndices());
 						
 						const TConstArrayView<TArray<FVector2f>> AllClothUVs = ClothFacade.GetRenderUVs();
 						check(AllClothUVs.Num() == Vertices.Num());
@@ -382,6 +390,11 @@ namespace UE::Chaos::ClothAsset
 						RenderCollection.AddSurface(MoveTemp(Vertices), MoveTemp(Indices), MoveTemp(Normals), MoveTemp(Colors));
 					}
 
+					if (!BoneWeights.IsEmpty() && !BoneIndices.IsEmpty())
+					{
+						RenderCollection.AddSurfaceBoneWeightsAndIndices(MoveTemp(BoneWeights), MoveTemp(BoneIndices));
+					}
+
 					RenderCollection.EndGeometryGroup(GeometryIndex);
 				}
 			}
diff --git a/Engine/Plugins/Experimental/Dataflow/Source/DataflowEditor/Private/Dataflow/DataflowConstructionViewportClient.cpp b/Engine/Plugins/Experimental/Dataflow/Source/DataflowEditor/Private/Dataflow/DataflowConstructionViewportClient.cpp
--- a/Engine/Plugins/Experimental/Dataflow/Source/DataflowEditor/Private/Dataflow/DataflowConstructionViewportClient.cpp
+++ b/Engine/Plugins/Experimental/Dataflow/Source/DataflowEditor/Private/Dataflow/DataflowConstructionViewportClient.cpp
@@ -87,6 +87,9 @@ FDataflowConstructionViewportClient::FDataflowConstructionViewportClient(FEditor
 		ViewFOV = FOVAngle;
 		ExposureSettings.bFixed = Options->bConstructionViewFixedExposure;
 	}
+	WidgetMode = UE::Widget::WM_None;
+	ModeTools->SetWidgetMode(WidgetMode);
+	ModeTools->SetCoordSystem(COORD_Local);
 }
 
 FDataflowConstructionViewportClient::~FDataflowConstructionViewportClient()
@@ -99,6 +102,25 @@ FDataflowConstructionViewportClient::~FDataflowConstructionViewportClient()
 	}
 }
 
+UE::Widget::EWidgetMode FDataflowConstructionViewportClient::GetWidgetMode() const
+{
+	// No transform gizmo yet available in the viewport
+	return UE::Widget::WM_None;
+}
+
+bool FDataflowConstructionViewportClient::CanSetWidgetMode(UE::Widget::EWidgetMode NewMode) const
+{
+	return	NewMode == UE::Widget::EWidgetMode::WM_Translate
+			|| NewMode == UE::Widget::EWidgetMode::WM_Scale
+			|| NewMode == UE::Widget::EWidgetMode::WM_Rotate;
+}
+
+void FDataflowConstructionViewportClient::SetWidgetMode(UE::Widget::EWidgetMode NewMode)
+{
+	ModeTools->SetWidgetMode(NewMode);
+	WidgetMode = NewMode;
+}
+
 void FDataflowConstructionViewportClient::SetDataflowEditorToolkit(TWeakPtr<FDataflowEditorToolkit> InDataflowEditorToolkitPtr)
 {
 	DataflowEditorToolkitPtr = InDataflowEditorToolkitPtr;
@@ -118,7 +140,7 @@ void FDataflowConstructionViewportClient::Tick(float DeltaSeconds)
 	}
 }
 
-USelection* FDataflowConstructionViewportClient::GetSelectedComponents() 
+USelection* FDataflowConstructionViewportClient::GetSelectedComponents() const 
 {
 	if (const TSharedPtr<FDataflowConstructionScene> DataflowConstructionScenePtr = ConstructionScene.Pin())
 	{
diff --git a/Engine/Plugins/Experimental/Dataflow/Source/DataflowEditor/Private/DataflowEditorTools/DataflowEditorEditSkeletonBonesTool.cpp b/Engine/Plugins/Experimental/Dataflow/Source/DataflowEditor/Private/DataflowEditorTools/DataflowEditorEditSkeletonBonesTool.cpp
--- a/Engine/Plugins/Experimental/Dataflow/Source/DataflowEditor/Private/DataflowEditorTools/DataflowEditorEditSkeletonBonesTool.cpp
+++ b/Engine/Plugins/Experimental/Dataflow/Source/DataflowEditor/Private/DataflowEditorTools/DataflowEditorEditSkeletonBonesTool.cpp
@@ -3,6 +3,8 @@
 #include "DataflowEditorTools/DataflowEditorEditSkeletonBonesTool.h"
 
 #include "ContextObjectStore.h"
+#include "EditorModeManager.h"
+#include "InteractiveGizmoManager.h"
 #include "Dataflow/DataflowCollectionEditSkeletonBonesNode.h"
 #include "Dataflow/DataflowContent.h"
 #include "Dataflow/DataflowObjectInterface.h"
@@ -12,7 +14,10 @@
 #include "ToolTargetManager.h"
 #include "ModelingToolTargetUtil.h"
 #include "Components/SkeletalMeshComponent.h"
+#include "Dataflow/DataflowConstructionViewportClient.h"
 #include "DynamicMesh/DynamicVertexSkinWeightsAttribute.h"
+#include "EditorGizmos/EditorTransformGizmoSource.h"
+#include "EditorGizmos/EditorTransformGizmoUtil.h"
 #include "TargetInterfaces/MeshDescriptionCommitter.h"
 #include "TargetInterfaces/MeshDescriptionProvider.h"
 #include "TargetInterfaces/PrimitiveComponentBackedTarget.h"
@@ -118,6 +123,29 @@ void UDataflowEditorEditSkeletonBonesTool::Shutdown(EToolShutdownType ShutdownTy
 	}
 }
 
+TScriptInterface<ITransformGizmoSource> UDataflowEditorEditSkeletonBonesTool::BuildTransformSource()
+{
+	const UEditorTransformGizmoContextObject* ToolGizmoContext = nullptr;
+	if (GetToolManager() && GetToolManager()->GetPairedGizmoManager())
+	{
+		if (IToolsContextQueriesAPI* QueriesAPI = GetToolManager()->GetPairedGizmoManager()->GetContextQueriesAPI())
+		{
+			if (FViewport* Viewport = QueriesAPI->GetFocusedViewport())
+			{
+				if (FViewportClient* ViewportClient = Viewport->GetClient())
+				{
+					if (FEditorModeTools*EditorModeTools = static_cast<FEditorViewportClient*>(ViewportClient)->GetModeTools())
+					{
+						ToolGizmoContext = EditorModeTools->GetGizmoContext();
+					}
+				}
+			}
+		}
+	}
+	
+	return UDataflowTransformGizmoSource::CreateNew(this, ToolGizmoContext);
+}
+
 FDataflowSkeletalMeshEditorBinding::FDataflowSkeletalMeshEditorBinding()
 	: ISkeletalMeshEditorBinding()
 {}
@@ -157,4 +185,44 @@ FDataflowSkeletalMeshEditorNotifier::FDataflowSkeletalMeshEditorNotifier()
 void FDataflowSkeletalMeshEditorNotifier::HandleNotification(const TArray<FName>& BoneNames, const ESkeletalMeshNotifyType InNotifyType)
 {}
 
+EGizmoTransformMode UDataflowTransformGizmoSource::GetGizmoMode() const
+{
+	if (const FEditorViewportClient* ViewportClient = GetViewportClient())
+	{
+		const FDataflowConstructionViewportClient* DataflowClient = static_cast<const FDataflowConstructionViewportClient*>(ViewportClient);
+
+		switch (DataflowClient->GetToolWidgetMode())
+		{
+			case UE::Widget::EWidgetMode::WM_Translate: return EGizmoTransformMode::Translate;
+			case UE::Widget::EWidgetMode::WM_Rotate: return EGizmoTransformMode::Rotate;
+			case UE::Widget::EWidgetMode::WM_Scale: return EGizmoTransformMode::Scale;
+			default: return EGizmoTransformMode::None;
+		}
+	}
+	return GizmoMode;
+}
+bool UDataflowTransformGizmoSource::GetVisible(const EViewportContext InViewportContext) const
+{
+	return CanInteract(InViewportContext);
+}
+
+bool UDataflowTransformGizmoSource::CanInteract(const EViewportContext InViewportContext) const
+{
+	return true;
+}
+
+const FRotationContext& UDataflowTransformGizmoSource::GetRotationContext() const
+{
+	static const FRotationContext DefaultContext;
+	return WeakContext.IsValid() ? WeakContext->RotationContext : DefaultContext;
+}
+
+UDataflowTransformGizmoSource* UDataflowTransformGizmoSource::CreateNew(UObject* Outer, 
+		const UEditorTransformGizmoContextObject* ContextObject)
+{
+	UDataflowTransformGizmoSource* NewSource = NewObject<UDataflowTransformGizmoSource>(Outer);
+	NewSource->WeakContext = ContextObject;
+	return NewSource;
+}
+
 #undef LOCTEXT_NAMESPACE
diff --git a/Engine/Plugins/Experimental/Dataflow/Source/DataflowEditor/Private/DataflowEditorTools/DataflowEditorSkinWeightsPaintTool.cpp b/Engine/Plugins/Experimental/Dataflow/Source/DataflowEditor/Private/DataflowEditorTools/DataflowEditorSkinWeightsPaintTool.cpp
--- a/Engine/Plugins/Experimental/Dataflow/Source/DataflowEditor/Private/DataflowEditorTools/DataflowEditorSkinWeightsPaintTool.cpp
+++ b/Engine/Plugins/Experimental/Dataflow/Source/DataflowEditor/Private/DataflowEditorTools/DataflowEditorSkinWeightsPaintTool.cpp
@@ -123,7 +123,8 @@ const FToolTargetTypeRequirements& UDataflowEditorSkinWeightsPaintToolBuilder::G
 UMeshSurfacePointTool* UDataflowEditorSkinWeightsPaintToolBuilder::CreateNewTool(const FToolBuilderState& SceneState) const
 {
 	UDataflowEditorSkinWeightsPaintTool* PaintTool = NewObject<UDataflowEditorSkinWeightsPaintTool>(SceneState.ToolManager);
-
+	PaintTool->SetTargetManager(SceneState.TargetManager);
+	
 	if (UDataflowBaseContent* ContextObject = SceneState.ToolManager->GetContextObjectStore()->FindContext<UDataflowBaseContent>())
 	{
 		if (FDataflowNode* PrimarySelection = ContextObject->GetSelectedNodeOfType<FDataflowCollectionEditSkinWeightsNode>())
diff --git a/Engine/Plugins/Experimental/Dataflow/Source/DataflowEditor/Public/Dataflow/DataflowConstructionViewportClient.h b/Engine/Plugins/Experimental/Dataflow/Source/DataflowEditor/Public/Dataflow/DataflowConstructionViewportClient.h
--- a/Engine/Plugins/Experimental/Dataflow/Source/DataflowEditor/Public/Dataflow/DataflowConstructionViewportClient.h
+++ b/Engine/Plugins/Experimental/Dataflow/Source/DataflowEditor/Public/Dataflow/DataflowConstructionViewportClient.h
@@ -29,7 +29,7 @@ public:
 
 	UE_API void SetConstructionViewMode(const UE::Dataflow::IDataflowConstructionViewMode* InViewMode);
 
-	UE_API USelection* GetSelectedComponents();
+	UE_API USelection* GetSelectedComponents() const;
 
 	/** Set the data flow toolkit used to create the client*/
 	UE_API void SetDataflowEditorToolkit(TWeakPtr<FDataflowEditorToolkit> DataflowToolkit);
@@ -48,6 +48,9 @@ public:
 
 	UE_API FString GetOverlayString() const;
 
+	/** Return the widget mode only used for tooling */
+	UE_API UE::Widget::EWidgetMode GetToolWidgetMode() const {return WidgetMode;}
+
 protected:
 
 	// FGCObject
@@ -65,6 +68,9 @@ private:
 	UE_API virtual void Draw(const FSceneView* View, FPrimitiveDrawInterface* PDI) override;
 	UE_API virtual void DrawCanvas(FViewport& InViewport, FSceneView& View, FCanvas& Canvas) override;
 	UE_API virtual float GetMinimumOrthoZoom() const override;
+	UE_API virtual bool CanSetWidgetMode(UE::Widget::EWidgetMode NewMode) const override;
+	UE_API virtual UE::Widget::EWidgetMode GetWidgetMode() const override;
+	UE_API virtual void SetWidgetMode(UE::Widget::EWidgetMode NewMode) override;
 
 	TUniquePtr<IClickDragBehaviorTarget> OrthoScrollBehaviorTarget;
 	TArray<TObjectPtr<UInputBehavior>> BehaviorsFor2DMode;
@@ -86,6 +92,9 @@ private:
 
 	// Saved view transforms for the currently inactive view modes (e.g. store the 3D camera here while in 2D mode and vice-versa)
 	TMap<FName, FViewportCameraTransform> SavedInactiveViewTransforms;
+
+	/** Currently active transform widget type. */
+	UE::Widget::EWidgetMode WidgetMode = UE::Widget::WM_Translate;
 };
 
 #undef UE_API
diff --git a/Engine/Plugins/Experimental/Dataflow/Source/DataflowEditor/Public/DataflowEditorTools/DataflowEditorEditSkeletonBonesTool.h b/Engine/Plugins/Experimental/Dataflow/Source/DataflowEditor/Public/DataflowEditorTools/DataflowEditorEditSkeletonBonesTool.h
--- a/Engine/Plugins/Experimental/Dataflow/Source/DataflowEditor/Public/DataflowEditorTools/DataflowEditorEditSkeletonBonesTool.h
+++ b/Engine/Plugins/Experimental/Dataflow/Source/DataflowEditor/Public/DataflowEditorTools/DataflowEditorEditSkeletonBonesTool.h
@@ -1,10 +1,12 @@
-﻿// Copyright Epic Games, Inc. All Rights Reserved.
+// Copyright Epic Games, Inc. All Rights Reserved.
 
 #pragma once
 
 #include "CoreMinimal.h"
 #include "HAL/Platform.h"
 #include "DataflowEditorTools/DataflowEditorToolBuilder.h"
+#include "EditorGizmos/EditorTransformGizmoSource.h"
+#include "EditorGizmos/TransformGizmoInterfaces.h"
 #include "SkeletalMesh/SkeletonEditingTool.h"
 
 #include "DataflowEditorEditSkeletonBonesTool.generated.h"
@@ -48,6 +50,7 @@ class UDataflowEditorEditSkeletonBonesTool : public USkeletonEditingTool
 	//~ Begin USkeletonEditingTool interface
 	UE_API virtual void Setup() override;
 	UE_API virtual void Shutdown(EToolShutdownType ShutdownType) override;
+	UE_API virtual TScriptInterface<ITransformGizmoSource> BuildTransformSource() override;
 	//~ End USkeletonEditingTool interface
 
 	friend class UDataflowEditorEditSkeletonBonesToolBuilder;
@@ -100,5 +103,36 @@ private :
 	TSharedPtr<FDataflowSkeletalMeshEditorNotifier> DataflowNotifier;
 };
 
+/**
+ * UDataflowTransformGizmoSource is an ITransformGizmoSource implementation that provides
+ * current state information used to configure the Editor transform gizmo.
+ */
+UCLASS(MinimalAPI)
+class UDataflowTransformGizmoSource : public UEditorTransformGizmoSource
+{
+	GENERATED_BODY()
+public:
+
+	//~ Begin ITransformGizmoSource interface
+	UE_API virtual EGizmoTransformMode GetGizmoMode() const override;
+	virtual EToolContextCoordinateSystem GetGizmoCoordSystemSpace() const override {return EToolContextCoordinateSystem::Local;}
+	virtual float GetGizmoScale() const override {return 1.0f;}
+	UE_API virtual bool GetVisible(const EViewportContext InViewportContext = EViewportContext::Focused) const override;
+	UE_API virtual bool CanInteract(const EViewportContext InViewportContext = EViewportContext::Focused) const override;
+	virtual EGizmoTransformScaleType GetScaleType() const override {return EGizmoTransformScaleType::Default;}
+	UE_API virtual const FRotationContext& GetRotationContext() const override;
+	//~ End ITransformGizmoSource interface
+
+	/** Transform source creation */
+	static UE_API UDataflowTransformGizmoSource* CreateNew(
+		UObject* Outer = (UObject*)GetTransientPackage(),
+		const class UEditorTransformGizmoContextObject* ContextObject = nullptr);
+
+private:
+
+	/** Default transform mode to be used when moving bones */
+	EGizmoTransformMode GizmoMode = EGizmoTransformMode::Translate;
+};
+
 
 #undef UE_API
diff --git a/Engine/Plugins/Runtime/MeshModelingToolset/Source/MeshModelingToolsEditorOnly/Private/SkeletalMesh/SkeletonEditingTool.cpp b/Engine/Plugins/Runtime/MeshModelingToolset/Source/MeshModelingToolsEditorOnly/Private/SkeletalMesh/SkeletonEditingTool.cpp
--- a/Engine/Plugins/Runtime/MeshModelingToolset/Source/MeshModelingToolsEditorOnly/Private/SkeletalMesh/SkeletonEditingTool.cpp
+++ b/Engine/Plugins/Runtime/MeshModelingToolset/Source/MeshModelingToolsEditorOnly/Private/SkeletalMesh/SkeletonEditingTool.cpp
@@ -210,6 +210,11 @@ void USkeletonEditingTool::SetupBehaviors()
 	AddInputBehavior(ClickDragBehavior);
 }
 
+TScriptInterface<ITransformGizmoSource> USkeletonEditingTool::BuildTransformSource() 
+{
+	return nullptr;
+}
+
 void USkeletonEditingTool::SetupGizmo(USceneComponent* InComponent)
 {
 	if (!GizmoContext.IsValid() || !InComponent)
@@ -244,7 +249,7 @@ void USkeletonEditingTool::SetupGizmo(USceneComponent* InComponent)
 	};
 
 	// create gizmo
-	GizmoWrapper = GizmoContext->GetNewWrapper(GetToolManager(), this, NewTarget);
+	GizmoWrapper = GizmoContext->GetNewWrapper(GetToolManager(), this, NewTarget, BuildTransformSource());
 	if (GizmoWrapper)
 	{
 		GizmoWrapper->Component = InComponent;
diff --git a/Engine/Plugins/Runtime/MeshModelingToolset/Source/MeshModelingToolsEditorOnly/Public/SkeletalMesh/SkeletalMeshEditingInterface.h b/Engine/Plugins/Runtime/MeshModelingToolset/Source/MeshModelingToolsEditorOnly/Public/SkeletalMesh/SkeletalMeshEditingInterface.h
--- a/Engine/Plugins/Runtime/MeshModelingToolset/Source/MeshModelingToolsEditorOnly/Public/SkeletalMesh/SkeletalMeshEditingInterface.h
+++ b/Engine/Plugins/Runtime/MeshModelingToolset/Source/MeshModelingToolsEditorOnly/Public/SkeletalMesh/SkeletalMeshEditingInterface.h
@@ -22,6 +22,7 @@ class IGizmoStateTarget;
 class USkeletonModifier;
 struct FInputDeviceRay;
 class UToolTarget;
+class ITransformGizmoSource;
 
 /**
  * USkeletalMeshEditingInterface
@@ -100,7 +101,8 @@ class USkeletalMeshGizmoContextObjectBase : public UObject
 	GENERATED_BODY()
 	
 public:
-	virtual USkeletalMeshGizmoWrapperBase* GetNewWrapper(UInteractiveToolManager* InToolManager, UObject* Outer = nullptr, IGizmoStateTarget* InStateTarget = nullptr)
+	virtual USkeletalMeshGizmoWrapperBase* GetNewWrapper(UInteractiveToolManager* InToolManager, UObject* Outer = nullptr,
+		IGizmoStateTarget* InStateTarget = nullptr, const TScriptInterface<ITransformGizmoSource>& GizmoSource = nullptr)
 	PURE_VIRTUAL(UGizmoContextObject::GetNewWrapper, return nullptr;);
 };
 
diff --git a/Engine/Plugins/Runtime/MeshModelingToolset/Source/MeshModelingToolsEditorOnly/Public/SkeletalMesh/SkeletonEditingTool.h b/Engine/Plugins/Runtime/MeshModelingToolset/Source/MeshModelingToolsEditorOnly/Public/SkeletalMesh/SkeletonEditingTool.h
--- a/Engine/Plugins/Runtime/MeshModelingToolset/Source/MeshModelingToolsEditorOnly/Public/SkeletalMesh/SkeletonEditingTool.h
+++ b/Engine/Plugins/Runtime/MeshModelingToolset/Source/MeshModelingToolsEditorOnly/Public/SkeletalMesh/SkeletonEditingTool.h
@@ -170,6 +170,9 @@ protected:
 	// ISkeletalMeshEditingInterface overrides
 	UE_API virtual void HandleSkeletalMeshModified(const TArray<FName>& BoneNames, const ESkeletalMeshNotifyType InNotifyType) override;
 
+	/** Build a transform source to override the mode tools functionality */
+	UE_API virtual TScriptInterface<ITransformGizmoSource> BuildTransformSource();
+
 	// Modifier functions
 	UE_API void CreateNewBone();
 	UE_API void ParentBones(const FName& InParentName);
diff --git a/Engine/Plugins/Runtime/MeshModelingToolset/Source/MeshModelingToolsEditorOnly/Public/SkeletalMesh/SkinWeightsPaintTool.h b/Engine/Plugins/Runtime/MeshModelingToolset/Source/MeshModelingToolsEditorOnly/Public/SkeletalMesh/SkinWeightsPaintTool.h
--- a/Engine/Plugins/Runtime/MeshModelingToolset/Source/MeshModelingToolsEditorOnly/Public/SkeletalMesh/SkinWeightsPaintTool.h
+++ b/Engine/Plugins/Runtime/MeshModelingToolset/Source/MeshModelingToolsEditorOnly/Public/SkeletalMesh/SkinWeightsPaintTool.h
@@ -18,6 +18,7 @@
 #include "Selections/GeometrySelection.h"
 #include "TargetInterfaces/MeshTargetInterfaceTypes.h"
 #include "ToolMeshSelector.h"
+#include "ToolTargetManager.h"
 
 #include "SkinWeightsPaintTool.generated.h"
 
@@ -724,6 +725,9 @@ public:
 	// get the target manager (cached from Setup)
 	UToolTargetManager* GetTargetManager() const { return TargetManager.Get(); };
 
+	// Set the target manager 
+	void SetTargetManager(UToolTargetManager* InTargetManager) { TargetManager = InTargetManager; };
+
 	// allows outside systems to access the weight data
 	SkinPaintTool::FSkinToolWeights& GetWeights() { return Weights; };
 
