Dataflow : fix super class call for OnClose

diff --git a/Engine/Plugins/Experimental/Dataflow/Source/DataflowEditor/Private/Dataflow/DataflowEditorToolkit.cpp b/Engine/Plugins/Experimental/Dataflow/Source/DataflowEditor/Private/Dataflow/DataflowEditorToolkit.cpp
--- a/Engine/Plugins/Experimental/Dataflow/Source/DataflowEditor/Private/Dataflow/DataflowEditorToolkit.cpp
+++ b/Engine/Plugins/Experimental/Dataflow/Source/DataflowEditor/Private/Dataflow/DataflowEditorToolkit.cpp
@@ -719,7 +719,7 @@ void FDataflowEditorToolkit::OnClose()
 	// duplication for each opening event.
 	GetEditorModeManager().ActivateDefaultMode();
 
-	FAssetEditorToolkit::OnClose();
+	FBaseCharacterFXEditorToolkit::OnClose();
 }
 
 void FDataflowEditorToolkit::PostInitAssetEditor()
