Dataflow - Fixed the node detail panel when undoing a reset property.
[RN] Dataflow - Fixed the refresh of the node detail panel when undoing/redoing a property change or property reset.

diff --git a/Engine/Plugins/Experimental/Dataflow/Source/DataflowEditor/Private/Dataflow/DataflowEditorCommands.cpp b/Engine/Plugins/Experimental/Dataflow/Source/DataflowEditor/Private/Dataflow/DataflowEditorCommands.cpp
--- a/Engine/Plugins/Experimental/Dataflow/Source/DataflowEditor/Private/Dataflow/DataflowEditorCommands.cpp
+++ b/Engine/Plugins/Experimental/Dataflow/Source/DataflowEditor/Private/Dataflow/DataflowEditorCommands.cpp
@@ -332,6 +332,8 @@ void FDataflowEditorCommands::OnNotifyPropertyPreChange(TSharedPtr<IStructureDet
 	// Find the associated UDataflowEdNode(s) and call Modify on them for Undo/Redo.
 	if (PropertiesEditor && Graph)
 	{
+		Graph->Modify();  // Modify must be called even without Dataflow nodes because of comment nodes' properties
+
 		if (TSharedPtr<const IStructureDataProvider> StructProvider = PropertiesEditor->GetStructureProvider())
 		{
 			const UStruct* const BaseStruct = StructProvider->GetBaseStructure();
@@ -396,8 +398,6 @@ void FDataflowEditorCommands::OnPropertyValueChanged(UDataflow* OutDataflow, TSh
 	{
 		if (ensure(OutDataflow && InPropertyChangedEvent.Property && InPropertyChangedEvent.Property->GetOwnerUObject()))
 		{
-			OutDataflow->Modify();  // Modify must be called even if SelectedNodes is empty because comment nodes aren't part of the selection set but still have properties
-
 			for (UDataflowEdNode* const SelectedNode : SelectedNodes)
 			{
 				if (TSharedPtr<FDataflowNode> DataflowNode = SelectedNode->GetDataflowNode())
diff --git a/Engine/Plugins/Experimental/Dataflow/Source/DataflowEditor/Private/Dataflow/DataflowEditorToolkit.cpp b/Engine/Plugins/Experimental/Dataflow/Source/DataflowEditor/Private/Dataflow/DataflowEditorToolkit.cpp
--- a/Engine/Plugins/Experimental/Dataflow/Source/DataflowEditor/Private/Dataflow/DataflowEditorToolkit.cpp
+++ b/Engine/Plugins/Experimental/Dataflow/Source/DataflowEditor/Private/Dataflow/DataflowEditorToolkit.cpp
@@ -4,6 +4,7 @@
 
 #include "AdvancedPreviewScene.h"
 #include "AdvancedPreviewSceneModule.h"
+#include "Algo/Transform.h"
 #include "AssetEditorModeManager.h"
 #include "AssetViewerSettings.h"
 #include "DataflowEditorStyle.h"
@@ -557,6 +558,16 @@ void FDataflowEditorToolkit::RegisterContextHandlers()
 	}
 
 	FDataflowAssetDelegates::OnNodeInvalidated.AddSP(this, &FDataflowEditorToolkit::OnNodeInvalidated);
+	FDataflowAssetDelegates::OnDataflowPostEditUndo.AddSPLambda(this, [this](UDataflow& DataflowAsset)
+		{
+			if (NodeDetailsEditor.IsValid())
+			{
+				// Bug in SStructureDetailsView (UE-349909) causes the UI to get out of sync and requires the view to be recreated
+				TSet<TObjectPtr<UObject>> SelectionSet;
+				Algo::Transform(SelectedDataflowNodes, SelectionSet, [](const TObjectPtr<UDataflowEdNode>& EdNode) { return Cast<UObject>(EdNode); });
+				FDataflowEditorCommands::OnSelectedNodesChanged(NodeDetailsEditor, &DataflowAsset, &DataflowAsset, SelectionSet);
+			}
+		});
 }
 
 void FDataflowEditorToolkit::UnregisterContextHandlers()
@@ -575,6 +586,7 @@ void FDataflowEditorToolkit::UnregisterContextHandlers()
 	}
 
 	FDataflowAssetDelegates::OnNodeInvalidated.RemoveAll(this);
+	FDataflowAssetDelegates::OnDataflowPostEditUndo.RemoveAll(this);
 }
 
 void FDataflowEditorToolkit::CreateEditorModeManager()
@@ -1290,7 +1302,7 @@ void FDataflowEditorToolkit::CreateWidgets()
 	{
 		if (UDataflow* DataflowAsset = EditorContent->GetDataflowAsset())
 		{
-			NodeDetailsEditor = CreateNodeDetailsEditorWidget(EditorContent->GetDataflowOwner());
+			NodeDetailsEditor = CreateNodeDetailsEditorWidget(EditorContent->GetDataflowAsset());
 			if(EditorContent->GetDataflowOwner() != EditorContent->GetDataflowAsset())
 			{
 				AssetDetailsEditor = CreateAssetDetailsEditorWidget({EditorContent->GetDataflowOwner(),EditorContent->GetDataflowAsset()});
diff --git a/Engine/Source/Runtime/Experimental/Dataflow/Engine/Private/Dataflow/DataflowObject.cpp b/Engine/Source/Runtime/Experimental/Dataflow/Engine/Private/Dataflow/DataflowObject.cpp
--- a/Engine/Source/Runtime/Experimental/Dataflow/Engine/Private/Dataflow/DataflowObject.cpp
+++ b/Engine/Source/Runtime/Experimental/Dataflow/Engine/Private/Dataflow/DataflowObject.cpp
@@ -17,6 +17,9 @@
 FDataflowAssetDelegates::FOnVariablesChanged FDataflowAssetDelegates::OnVariablesChanged;
 FDataflowAssetDelegates::FOnSubGraphsChanged FDataflowAssetDelegates::OnSubGraphsChanged;
 FDataflowAssetDelegates::FOnNodeInvalidated FDataflowAssetDelegates::OnNodeInvalidated;
+#if WITH_EDITOR
+FDataflowAssetDelegates::FOnDataflowPostEditUndo FDataflowAssetDelegates::OnDataflowPostEditUndo;
+#endif
 FDataflowAssetDelegates::FOnVariablesOverrideStateChanged FDataflowAssetDelegates::OnVariablesOverrideStateChanged;
 
 namespace UE::Dataflow::CVars
@@ -173,6 +176,12 @@ void UDataflow::PostEditChangeProperty(struct FPropertyChangedEvent& PropertyCha
 	Super::PostEditChangeProperty(PropertyChangedEvent);
 }
 
+void UDataflow::PostEditUndo()
+{
+	Super::PostEditUndo();
+	FDataflowAssetDelegates::OnDataflowPostEditUndo.Broadcast(*this);
+}
+
 #endif
 
 void UDataflow::PostLoad()
diff --git a/Engine/Source/Runtime/Experimental/Dataflow/Engine/Public/Dataflow/DataflowObject.h b/Engine/Source/Runtime/Experimental/Dataflow/Engine/Public/Dataflow/DataflowObject.h
--- a/Engine/Source/Runtime/Experimental/Dataflow/Engine/Public/Dataflow/DataflowObject.h
+++ b/Engine/Source/Runtime/Experimental/Dataflow/Engine/Public/Dataflow/DataflowObject.h
@@ -56,6 +56,11 @@ struct FDataflowAssetDelegates
 
 	DECLARE_MULTICAST_DELEGATE_TwoParams(FOnNodeInvalidated, UDataflow&, FDataflowNode&);
 	static DATAFLOWENGINE_API FOnNodeInvalidated OnNodeInvalidated;
+
+#if WITH_EDITOR
+	DECLARE_MULTICAST_DELEGATE_OneParam(FOnDataflowPostEditUndo, UDataflow&);
+	static DATAFLOWENGINE_API FOnDataflowPostEditUndo OnDataflowPostEditUndo;
+#endif
 };
 
 /**
@@ -161,6 +166,7 @@ public:
 #if WITH_EDITOR
 	DATAFLOWENGINE_API virtual bool CanEditChange(const FProperty* InProperty) const override;
 	DATAFLOWENGINE_API virtual void PostEditChangeProperty(struct FPropertyChangedEvent& PropertyChangedEvent) override;
+	virtual void PostEditUndo() override;
 #endif
 	DATAFLOWENGINE_API virtual void PostLoad() override;
 	/** End UObject Interface */
