Dataflow : fix collection spreadsheet crashing when encountering a attribute with an empty name

diff --git a/Engine/Plugins/Experimental/Dataflow/Source/DataflowEditor/Private/Dataflow/CollectionSpreadSheetWidget.cpp b/Engine/Plugins/Experimental/Dataflow/Source/DataflowEditor/Private/Dataflow/CollectionSpreadSheetWidget.cpp
--- a/Engine/Plugins/Experimental/Dataflow/Source/DataflowEditor/Private/Dataflow/CollectionSpreadSheetWidget.cpp
+++ b/Engine/Plugins/Experimental/Dataflow/Source/DataflowEditor/Private/Dataflow/CollectionSpreadSheetWidget.cpp
@@ -17,6 +17,7 @@
 #define LOCTEXT_NAMESPACE "CollectionSpreadSheet"
 
 const FName FCollectionSpreadSheetHeader::IndexColumnName = FName("Index");
+const FName FCollectionSpreadSheetHeader::NoColumnName = FName("[NoName]");
 
 void SCollectionSpreadSheetRow::Construct(const FArguments& InArgs, TSharedRef<STableViewBase> OwnerTableView, const TSharedPtr<const FCollectionSpreadSheetHeader>& InHeader, const TSharedPtr<const FCollectionSpreadSheetItem>& InItem)
 {
@@ -134,12 +135,19 @@ void SCollectionSpreadSheet::RegenerateHeader()
 	{
 		Header->ColumnNames.Add(FCollectionSpreadSheetHeader::IndexColumnName);
 
-		if (TSharedPtr<const FManagedArrayCollection> CollectionPtr = CollectionInfoMap[SelectedOutput.ToString()].Collection)
+		if (const SCollectionInfo* CollectionInfo = CollectionInfoMap.Find(SelectedOutput.ToString()))
 		{
-			for (FName Attr : CollectionPtr->AttributeNames(SelectedGroup))
+			if (CollectionInfo->Collection.IsValid())
 			{
-				Header->ColumnNames.Add(Attr);
-				AttrTypes.Add(UE::Dataflow::CollectionSpreadSheetHelpers::GetArrayTypeString(CollectionPtr->GetAttributeType(Attr, SelectedGroup)).ToString());
+				const TArray<FName> AttributeNames = CollectionInfo->Collection->AttributeNames(SelectedGroup);
+				const int32 NumAttributes = AttributeNames.Num();
+				for (int32 AttributeIndex = 0; AttributeIndex < NumAttributes; ++AttributeIndex)
+				{
+					// column cannot have empty names
+					const FName AttributeName = AttributeNames[AttributeIndex].IsNone() ? FCollectionSpreadSheetHeader::NoColumnName : AttributeNames[AttributeIndex];
+					Header->ColumnNames.Add(AttributeName);
+					AttrTypes.Add(UE::Dataflow::CollectionSpreadSheetHelpers::GetArrayTypeString(CollectionInfo->Collection->GetAttributeType(AttributeName, SelectedGroup)).ToString());
+				}
 			}
 		}
 	}
@@ -200,37 +208,43 @@ void SCollectionSpreadSheet::RepopulateListView()
 		!SelectedOutput.IsNone() && SelectedOutput.ToString() != "" &&
 		!SelectedGroup.IsNone() && SelectedGroup.ToString() != "")
 	{
-		TSharedPtr<const FManagedArrayCollection> CollectionPtr = CollectionInfoMap[SelectedOutput.ToString()].Collection;
-		if (CollectionPtr)
+		if (const SCollectionInfo* CollectionInfo = CollectionInfoMap.Find(SelectedOutput.ToString()))
 		{
-			const int32 NumElems = CollectionPtr->NumElements(SelectedGroup);
-			ListItems.SetNum(NumElems);
-
-			constexpr int32 BatchSize = 50;
-			ParallelFor(TEXT("ScollectionSpreadSheet_GenerateItems"), NumElems, BatchSize,
-				[this, &CollectionPtr](int32 IdxElem)
-				{
-					const TSharedPtr<FCollectionSpreadSheetItem> NewItem = MakeShared<FCollectionSpreadSheetItem>();
-					NewItem->Values.SetNum(Header->ColumnNames.Num());
+			if (CollectionInfo->Collection.IsValid())
+			{
+				const int32 NumElems = CollectionInfo->Collection->NumElements(SelectedGroup);
+				ListItems.SetNum(NumElems);
 
-					for (int32 IdxColumn = 0; IdxColumn < Header->ColumnNames.Num(); ++IdxColumn)
+				constexpr int32 BatchSize = 50;
+				ParallelFor(TEXT("ScollectionSpreadSheet_GenerateItems"), NumElems, BatchSize,
+					[this, CollectionPtr = CollectionInfo->Collection](int32 IdxElem)
 					{
-						const FName& ColumnName = Header->ColumnNames[IdxColumn];
-						if (ColumnName == FCollectionSpreadSheetHeader::IndexColumnName)
-						{
-							NewItem->Values[IdxColumn] = FString::FromInt(IdxElem);
-						}
-						else
+						const TSharedPtr<FCollectionSpreadSheetItem> NewItem = MakeShared<FCollectionSpreadSheetItem>();
+						NewItem->Values.SetNum(Header->ColumnNames.Num());
+
+						for (int32 IdxColumn = 0; IdxColumn < Header->ColumnNames.Num(); ++IdxColumn)
 						{
-							NewItem->Values[IdxColumn] = UE::Dataflow::CollectionSpreadSheetHelpers::AttributeValueToString(*CollectionPtr, ColumnName, SelectedGroup, IdxElem);
+							const FName& ColumnName = Header->ColumnNames[IdxColumn];
+							if (ColumnName == FCollectionSpreadSheetHeader::IndexColumnName)
+							{
+								NewItem->Values[IdxColumn] = FString::FromInt(IdxElem);
+							}
+							else
+							{
+								const FName AttributeName = (ColumnName == FCollectionSpreadSheetHeader::NoColumnName) ? NAME_None : ColumnName;
+								NewItem->Values[IdxColumn] = UE::Dataflow::CollectionSpreadSheetHelpers::AttributeValueToString(*CollectionPtr, AttributeName, SelectedGroup, IdxElem);
+							}
 						}
-					}
-					ListItems[IdxElem] = NewItem;
-				});
+						ListItems[IdxElem] = NewItem;
+					});
+			}
 		}
 	}
 
-	ListView->RequestListRefresh();
+	if (ListView)
+	{
+		ListView->RequestListRefresh();
+	}
 }
 
 TSharedRef<ITableRow> SCollectionSpreadSheet::GenerateRow(TSharedPtr<const FCollectionSpreadSheetItem> InItem, const TSharedRef<STableViewBase>& OwnerTable)
@@ -253,10 +267,9 @@ void SCollectionSpreadSheetWidget::UpdateTransformOutliner()
 		if (TransformOutliner && bIsTransformGroupSelected)
 		{
 			const FString OutputStr = CollectionTable->GetSelectedOutputStr();
-			if (CollectionTable->GetCollectionInfoMap().Contains(OutputStr))
+			if (const SCollectionInfo* CollectionInfo = CollectionTable->GetCollectionInfoMap().Find(OutputStr))
 			{
-				TSharedPtr<const FManagedArrayCollection> CollectionPtr = CollectionTable->GetCollectionInfoMap()[OutputStr].Collection;
-				TransformOutliner->SetCollection(CollectionPtr, CollectionTable->GetSelectedOutput());
+				TransformOutliner->SetCollection(CollectionInfo->Collection, CollectionTable->GetSelectedOutput());
 				TransformOutliner->RegenerateHeader();
 			}
 		}
@@ -271,10 +284,9 @@ void SCollectionSpreadSheetWidget::UpdateVertexOutliner()
 		if (VerticesOutliner && bIsVerticesGroupSelected)
 		{
 			const FString OutputStr = CollectionTable->GetSelectedOutputStr();
-			if (CollectionTable->GetCollectionInfoMap().Contains(OutputStr))
+			if (const SCollectionInfo* CollectionInfo = CollectionTable->GetCollectionInfoMap().Find(OutputStr))
 			{
-				TSharedPtr<const FManagedArrayCollection> CollectionPtr = CollectionTable->GetCollectionInfoMap()[OutputStr].Collection;
-				VerticesOutliner->SetCollection(CollectionPtr);
+				VerticesOutliner->SetCollection(CollectionInfo->Collection);
 				VerticesOutliner->RegenerateHeader();
 			}
 		}
@@ -289,10 +301,9 @@ void SCollectionSpreadSheetWidget::UpdateFaceOutliner()
 		if (FacesOutliner && bIsFacesGroupSelected)
 		{
 			const FString OutputStr = CollectionTable->GetSelectedOutputStr();
-			if (CollectionTable->GetCollectionInfoMap().Contains(OutputStr))
+			if (const SCollectionInfo* CollectionInfo = CollectionTable->GetCollectionInfoMap().Find(OutputStr))
 			{
-				TSharedPtr<const FManagedArrayCollection> CollectionPtr = CollectionTable->GetCollectionInfoMap()[OutputStr].Collection;
-				FacesOutliner->SetCollection(CollectionPtr);
+				FacesOutliner->SetCollection(CollectionInfo->Collection);
 				FacesOutliner->RegenerateHeader();
 			}
 		}
@@ -352,8 +363,11 @@ void SCollectionSpreadSheetWidget::CollectionGroupsComboBoxSelectionChanged(FNam
 
 			if (!InSelectedGroup.IsNone())
 			{
-				int32 NumElems = CollectionTable->GetCollectionInfoMap()[CollectionTable->GetSelectedOutput().ToString()].Collection->NumElements(InSelectedGroup);
-				CollectionTable->SetNumItems(NumElems);
+				if (const SCollectionInfo* CollectionInfo = CollectionTable->GetCollectionInfoMap().Find(CollectionTable->GetSelectedOutput().ToString()))
+				{
+					const int32 NumElems = CollectionInfo->Collection.IsValid() ? CollectionInfo->Collection->NumElements(InSelectedGroup): 0;
+					CollectionTable->SetNumItems(NumElems);
+				}
 			}
 		}
 	}
@@ -791,11 +805,14 @@ void SCollectionSpreadSheetWidget::UpdateCollectionGroups(const FName& InOutputN
 
 		if (CollectionTable)
 		{
-			if (TSharedPtr<const FManagedArrayCollection> CollectionPtr = CollectionTable->GetCollectionInfoMap()[InOutputName.ToString()].Collection)
+			if (const SCollectionInfo* CollectionInfo = CollectionTable->GetCollectionInfoMap().Find(InOutputName.ToString()))
 			{
-				for (FName Group : CollectionPtr->GroupNames())
+				if (CollectionInfo->Collection.IsValid())
 				{
-					CollectionGroups.Add(Group);
+					for (const FName Group : CollectionInfo->Collection->GroupNames())
+					{
+						CollectionGroups.Add(Group);
+					}
 				}
 			}
 		}
diff --git a/Engine/Plugins/Experimental/Dataflow/Source/DataflowEditor/Public/Dataflow/DataflowCollectionSpreadSheetWidget.h b/Engine/Plugins/Experimental/Dataflow/Source/DataflowEditor/Public/Dataflow/DataflowCollectionSpreadSheetWidget.h
--- a/Engine/Plugins/Experimental/Dataflow/Source/DataflowEditor/Public/Dataflow/DataflowCollectionSpreadSheetWidget.h
+++ b/Engine/Plugins/Experimental/Dataflow/Source/DataflowEditor/Public/Dataflow/DataflowCollectionSpreadSheetWidget.h
@@ -33,6 +33,7 @@ struct SCollectionInfo
 struct FCollectionSpreadSheetHeader
 {
 	static const FName IndexColumnName;
+	static const FName NoColumnName;
 
 	TArray<FName> ColumnNames;
 };
