Dataflow : make sure the vertex attribute painting tool checks for that the mode is not null

diff --git a/Engine/Plugins/Experimental/Dataflow/Source/DataflowEditor/Private/DataflowEditorTools/DataflowEditorVertexAttributePaintTool.cpp b/Engine/Plugins/Experimental/Dataflow/Source/DataflowEditor/Private/DataflowEditorTools/DataflowEditorVertexAttributePaintTool.cpp
--- a/Engine/Plugins/Experimental/Dataflow/Source/DataflowEditor/Private/DataflowEditorTools/DataflowEditorVertexAttributePaintTool.cpp
+++ b/Engine/Plugins/Experimental/Dataflow/Source/DataflowEditor/Private/DataflowEditorTools/DataflowEditorVertexAttributePaintTool.cpp
@@ -623,9 +623,12 @@ void UDataflowEditorVertexAttributePaintTool::Setup()
 	SetToolDisplayName(LOCTEXT("ToolName", "Paint Weight Maps"));
 
 	// Hide all meshes in the DataflowConstructionScene, as we will be painting onto our own Preview mesh
-	if (const TSharedPtr<FDataflowConstructionScene> Scene = Mode->GetDataflowConstructionScene().Pin())
+	if (ensure(Mode))
 	{
-		Scene->SetVisibility(false);
+		if (const TSharedPtr<FDataflowConstructionScene> Scene = Mode->GetDataflowConstructionScene().Pin())
+		{
+			Scene->SetVisibility(false);
+		}
 	}
 
 	//OnDynamicMeshComponentChangedHandle = DynamicMeshComponent->OnMeshVerticesChanged.AddUObject(this, &UDataflowEditorVertexAttributePaintTool::OnDynamicMeshComponentChanged);
