Dataflow : make sure to cancel the current opened tool if the matching node is being deleted

diff --git a/Engine/Plugins/Experimental/Dataflow/Source/DataflowEditor/Private/Dataflow/DataflowEditorToolkit.cpp b/Engine/Plugins/Experimental/Dataflow/Source/DataflowEditor/Private/Dataflow/DataflowEditorToolkit.cpp
--- a/Engine/Plugins/Experimental/Dataflow/Source/DataflowEditor/Private/Dataflow/DataflowEditorToolkit.cpp
+++ b/Engine/Plugins/Experimental/Dataflow/Source/DataflowEditor/Private/Dataflow/DataflowEditorToolkit.cpp
@@ -1818,12 +1818,20 @@ void FDataflowEditorToolkit::OnNodeInvalidated(UDataflow& DataflowAsset, FDatafl
 void FDataflowEditorToolkit::OnNodeDeleted(const TSet<UObject*>& NewSelection)
 {
 	const TSet<TObjectPtr<UDataflowEdNode>> DeletedNodes = FilterDataflowEdNodesFromSet(NewSelection);
-	for (const UDataflowEdNode* Node : DeletedNodes)
+
+	FGraphPanelSelectionSet DeletedNodesSet;
+	for (UDataflowEdNode* Node : DeletedNodes)
 	{
 		if (SelectedDataflowNodes.Contains(Node))
 		{
 			SelectedDataflowNodes.Remove(Node);
 		}
+		DeletedNodesSet.Add(Node);
+	}
+
+	if (UDataflowEditorMode* DataflowEdMode = Cast<UDataflowEditorMode>(EditorModeManager->GetActiveScriptableMode(UDataflowEditorMode::EM_DataflowEditorModeId)))
+	{
+		DataflowEdMode->OnDataflowNodeDeleted(DeletedNodesSet);
 	}
 }
 
