Dataflow : fix potential crash with EditSkeletonBone node when double clicking on an asset laredy opened using the node

diff --git a/Engine/Plugins/Experimental/Dataflow/Source/DataflowNodes/Private/Dataflow/DataflowCollectionEditSkeletonBonesNode.cpp b/Engine/Plugins/Experimental/Dataflow/Source/DataflowNodes/Private/Dataflow/DataflowCollectionEditSkeletonBonesNode.cpp
--- a/Engine/Plugins/Experimental/Dataflow/Source/DataflowNodes/Private/Dataflow/DataflowCollectionEditSkeletonBonesNode.cpp
+++ b/Engine/Plugins/Experimental/Dataflow/Source/DataflowNodes/Private/Dataflow/DataflowCollectionEditSkeletonBonesNode.cpp
@@ -35,7 +35,7 @@ TArray<UE::Dataflow::FRenderingParameter> FDataflowCollectionEditSkeletonBonesNo
 void FDataflowCollectionEditSkeletonBonesNode::AddPrimitiveComponents(UE::Dataflow::FContext& Context, const TSharedPtr<const FManagedArrayCollection> RenderCollection, TObjectPtr<UObject> NodeOwner, 
 	TObjectPtr<AActor> RootActor, TArray<TObjectPtr<UPrimitiveComponent>>& PrimitiveComponents)  
 {
-	if(RootActor)
+	if(RootActor && RenderCollection)
 	{
 		if(TObjectPtr<USkeleton> TransientSkeleton = UpdateToolSkeleton(Context))
 		{
