Dataflow : nodes with debug draw enabled and using the new rednering system are now rendering correctly in construction view

diff --git a/Engine/Plugins/Experimental/Dataflow/Source/DataflowEditor/Private/Dataflow/DataflowConstructionScene.cpp b/Engine/Plugins/Experimental/Dataflow/Source/DataflowEditor/Private/Dataflow/DataflowConstructionScene.cpp
--- a/Engine/Plugins/Experimental/Dataflow/Source/DataflowEditor/Private/Dataflow/DataflowConstructionScene.cpp
+++ b/Engine/Plugins/Experimental/Dataflow/Source/DataflowEditor/Private/Dataflow/DataflowConstructionScene.cpp
@@ -594,8 +594,9 @@ void FDataflowConstructionScene::UpdateComponentsFromNode(const UDataflowEdNode*
 		GeometryCollection::Facades::FRenderingFacade RenderingFacade(*RenderCollection);
 		RenderingFacade.DefineSchema();
 
+		// do not check debug draw as this is orthogonal to the old/new rendering system 
 		const bool bCanRenderOutput = (ConstructionViewMode)
-			? UE::Dataflow::CanRenderNodeOutput(*EdNode, *EditorContent, *ConstructionViewMode)
+			? UE::Dataflow::CanRenderNodeOutput(*EdNode, *EditorContent, *ConstructionViewMode, /*bCheckDebugDraw*/false)
 			: false;
 
 		// step 2 : Render nodes that have explictly specifiy how to be rendered ( Old path )
diff --git a/Engine/Plugins/Experimental/Dataflow/Source/DataflowEditor/Private/Dataflow/DataflowEditorUtil.cpp b/Engine/Plugins/Experimental/Dataflow/Source/DataflowEditor/Private/Dataflow/DataflowEditorUtil.cpp
--- a/Engine/Plugins/Experimental/Dataflow/Source/DataflowEditor/Private/Dataflow/DataflowEditorUtil.cpp
+++ b/Engine/Plugins/Experimental/Dataflow/Source/DataflowEditor/Private/Dataflow/DataflowEditorUtil.cpp
@@ -122,7 +122,7 @@ namespace UE::Dataflow
 		return MakeShared<FEngineContext>(nullptr);
 	}
 
-	bool CanRenderNodeOutput(const UDataflowEdNode& EdNode, const UDataflowBaseContent& EditorContent, const IDataflowConstructionViewMode& ViewMode)
+	bool CanRenderNodeOutput(const UDataflowEdNode& EdNode, const UDataflowBaseContent& EditorContent, const IDataflowConstructionViewMode& ViewMode, bool bCheckDebugDraw)
 	{
 		if (const TSharedPtr<FEngineContext> Context = EditorContent.GetDataflowContext())
 		{
@@ -138,7 +138,7 @@ namespace UE::Dataflow
 						}
 					}
 				}
-				if (NodeTarget->CanDebugDrawViewMode(ViewMode.GetName()))
+				if (bCheckDebugDraw && NodeTarget->CanDebugDrawViewMode(ViewMode.GetName()))
 				{
 					return true;
 				}
diff --git a/Engine/Plugins/Experimental/Dataflow/Source/DataflowEditor/Private/Dataflow/DataflowEditorUtil.h b/Engine/Plugins/Experimental/Dataflow/Source/DataflowEditor/Private/Dataflow/DataflowEditorUtil.h
--- a/Engine/Plugins/Experimental/Dataflow/Source/DataflowEditor/Private/Dataflow/DataflowEditorUtil.h
+++ b/Engine/Plugins/Experimental/Dataflow/Source/DataflowEditor/Private/Dataflow/DataflowEditorUtil.h
@@ -54,7 +54,7 @@ namespace UE::Dataflow
 {
 	TSharedPtr<UE::Dataflow::FEngineContext> GetContext(TObjectPtr<UDataflowBaseContent> Content);
 
-	bool CanRenderNodeOutput(const UDataflowEdNode& EdNode, const UDataflowBaseContent& EditorContent, const UE::Dataflow::IDataflowConstructionViewMode& ViewMode);
+	bool CanRenderNodeOutput(const UDataflowEdNode& EdNode, const UDataflowBaseContent& EditorContent, const UE::Dataflow::IDataflowConstructionViewMode& ViewMode, bool bCheckDebugDraw = true);
 
 	/** This function fills the managed array collection linked to the rendering facade from the callbacks.
 	 * Returns true if the resulting collection should be directly rendered of false if the collection will be used to generate a primitive component */
