Add some explicit transitive dependencies that are needed for UE_WITH_CONSTINIT_UOBJECTS.
In that case, object classes need a link dependency on every module in their inheritance chain, and PublicDependencyModuleNames on the intermediate dependency doesn't provide that.

diff --git a/Engine/Plugins/Experimental/Dataflow/Source/DataflowEditor/DataflowEditor.Build.cs b/Engine/Plugins/Experimental/Dataflow/Source/DataflowEditor/DataflowEditor.Build.cs
--- a/Engine/Plugins/Experimental/Dataflow/Source/DataflowEditor/DataflowEditor.Build.cs
+++ b/Engine/Plugins/Experimental/Dataflow/Source/DataflowEditor/DataflowEditor.Build.cs
@@ -55,6 +55,7 @@ namespace UnrealBuildTool.Rules
 					"EditorFramework",
 					"EditorInteractiveToolsFramework",
 					"EditorStyle",
+					"EditorSubsystem",
 					"GeometryCache",
 					"GeometryCore",
 					"GeometryCollectionEngine",
