Dataflow : fix the skeleton debug draw object calling a delegate from a thread different from the game thread

diff --git a/Engine/Plugins/Experimental/Dataflow/Source/DataflowNodes/Private/Dataflow/DataflowSkeletalMeshNodes.cpp b/Engine/Plugins/Experimental/Dataflow/Source/DataflowNodes/Private/Dataflow/DataflowSkeletalMeshNodes.cpp
--- a/Engine/Plugins/Experimental/Dataflow/Source/DataflowNodes/Private/Dataflow/DataflowSkeletalMeshNodes.cpp
+++ b/Engine/Plugins/Experimental/Dataflow/Source/DataflowNodes/Private/Dataflow/DataflowSkeletalMeshNodes.cpp
@@ -44,6 +44,11 @@ void FDataflowDebugDrawSkeletonObject::OverrideBoneTransforms(const TArray<FTran
 	TransformOverrides = InBoneTransforms;
 }
 
+const FReferenceSkeleton& FDataflowDebugDrawSkeletonObject::GetReferenceSkeleton() const
+{
+	return ReferenceSkeleton;
+}
+
 void FDataflowDebugDrawSkeletonObject::PopulateDataflowElements()
 {
 	ElementsOffset = DataflowElements.Num();
@@ -155,7 +160,13 @@ void FDataflowDebugDrawSkeletonObject::DrawDataflowElements(FPrimitiveDrawInterf
 
 	if(bSelectionChanged)
 	{
-		OnBoneSelectionChanged.Broadcast(BoneNames);
+		// this method can be called on a different thread than the game thread 
+		// let's make sure the bon change broadcast happens on the game thread
+		Async(EAsyncExecution::TaskGraphMainTick,
+			[this, BoneNames]
+			{
+				OnBoneSelectionChanged.Broadcast(BoneNames);
+			});
 	}
 	if(bBonesVisible)
 	{
diff --git a/Engine/Plugins/Experimental/Dataflow/Source/DataflowNodes/Public/Dataflow/DataflowSkeletalMeshNodes.h b/Engine/Plugins/Experimental/Dataflow/Source/DataflowNodes/Public/Dataflow/DataflowSkeletalMeshNodes.h
--- a/Engine/Plugins/Experimental/Dataflow/Source/DataflowNodes/Public/Dataflow/DataflowSkeletalMeshNodes.h
+++ b/Engine/Plugins/Experimental/Dataflow/Source/DataflowNodes/Public/Dataflow/DataflowSkeletalMeshNodes.h
@@ -40,6 +40,8 @@ struct FDataflowDebugDrawSkeletonObject : public FDataflowDebugDrawBaseObject
 	/** override the bone transforms */
 	UE_API void OverrideBoneTransforms(const TArray<FTransform>& InBoneTransforms);
 
+	UE_API const FReferenceSkeleton& GetReferenceSkeleton() const;
+
 	/** Delegate to broadcast bones selection changes */
 	FDataflowBoneSelectionChangedNotifyDelegate OnBoneSelectionChanged;
 
