Dataflow outliner : move visibility to gutter  , add support SHIFT + H for hiding to mimic the level outliner

diff --git a/Engine/Plugins/Experimental/Dataflow/Source/DataflowEditor/Private/Dataflow/DataflowOutlinerMode.cpp b/Engine/Plugins/Experimental/Dataflow/Source/DataflowEditor/Private/Dataflow/DataflowOutlinerMode.cpp
--- a/Engine/Plugins/Experimental/Dataflow/Source/DataflowEditor/Private/Dataflow/DataflowOutlinerMode.cpp
+++ b/Engine/Plugins/Experimental/Dataflow/Source/DataflowEditor/Private/Dataflow/DataflowOutlinerMode.cpp
@@ -328,7 +328,7 @@ FReply FDataflowOutlinerMode::OnKeyDown(const FKeyEvent& InKeyEvent)
 		}
 		else if (InKeyEvent.GetKey() == EKeys::H)
 		{
-			const bool bIsVisible = InKeyEvent.IsControlDown() ? true : false;
+			const bool bIsVisible = (InKeyEvent.IsShiftDown() || InKeyEvent.IsControlDown()) ? true : false;
 			for(TWeakPtr<ISceneOutlinerTreeItem>& SelectedItem : Selection.SelectedItems)
 			{
 				if(UPrimitiveComponent* SelectedComponent = UE::Dataflow::Private::GetOutlinerItemObject<UPrimitiveComponent>(SelectedItem))
diff --git a/Engine/Plugins/Experimental/Dataflow/Source/DataflowEditor/Private/Dataflow/DataflowOutlinerView.cpp b/Engine/Plugins/Experimental/Dataflow/Source/DataflowEditor/Private/Dataflow/DataflowOutlinerView.cpp
--- a/Engine/Plugins/Experimental/Dataflow/Source/DataflowEditor/Private/Dataflow/DataflowOutlinerView.cpp
+++ b/Engine/Plugins/Experimental/Dataflow/Source/DataflowEditor/Private/Dataflow/DataflowOutlinerView.cpp
@@ -89,7 +89,8 @@ TSharedPtr<ISceneOutliner> FDataflowOutlinerView::CreateWidget()
 		return new FDataflowOutlinerMode(Params, ConstructionScene, SimulationScene);
 	});
 	
-	InitOptions.ColumnMap.Add(FSceneOutlinerBuiltInColumnTypes::Label(), FSceneOutlinerColumnInfo(ESceneOutlinerColumnVisibility::Visible, 10));
+	InitOptions.ColumnMap.Add(FSceneOutlinerBuiltInColumnTypes::Gutter(), FSceneOutlinerColumnInfo(ESceneOutlinerColumnVisibility::Visible, 40));
+	InitOptions.ColumnMap.Add(FSceneOutlinerBuiltInColumnTypes::Label(), FSceneOutlinerColumnInfo(ESceneOutlinerColumnVisibility::Visible, 50));
 
 	OutlinerWidget = SNew(SSceneOutliner, InitOptions);
 
