Dataflow: Improve the render group default visibility logic and the robustness of the metadata parsing

diff --git a/Engine/Plugins/Experimental/Dataflow/Source/DataflowEditor/Private/DataflowRendering/DataflowRenderableTypeInstance.cpp b/Engine/Plugins/Experimental/Dataflow/Source/DataflowEditor/Private/DataflowRendering/DataflowRenderableTypeInstance.cpp
--- a/Engine/Plugins/Experimental/Dataflow/Source/DataflowEditor/Private/DataflowRendering/DataflowRenderableTypeInstance.cpp
+++ b/Engine/Plugins/Experimental/Dataflow/Source/DataflowEditor/Private/DataflowRendering/DataflowRenderableTypeInstance.cpp
@@ -23,11 +23,16 @@ namespace UE::Dataflow
 					RenderGroupsMetaData.ParseIntoArray(RenderGroups, TEXT(","));
 					if (RenderGroups.Num() > 0)
 					{
+						for (FString& ParsedRenderGroup : RenderGroups)
+						{
+							ParsedRenderGroup.TrimStartAndEndInline();
+						}
 						return RenderGroups.Contains(RenderGroup);
 					}
 				}
 			}
-			return true;
+			// we don't want to render all render groups, so we select the priority ones
+			return (RenderGroup == TEXT("Surface") || RenderGroup == TEXT("UVs"));
 		}
 	}
 
