[Dataflow] Fixed name extraction from paste buffer. Node connection is encoded as "/NodeName:PropertyName|PropertyType" in the paste buffer but the extraction was looking for the final "/" as the beginning of the name. Nodes can have "/" characters in their name. This means the extraction was slicing the real node name. The process that creates nodes from the paste buffer assumes that the names are correct so this leads to a crash in looking up the node from the connection name.

diff --git a/Engine/Plugins/Experimental/Dataflow/Source/DataflowEditor/Private/Dataflow/AssetDefinition_DataflowAsset.cpp b/Engine/Plugins/Experimental/Dataflow/Source/DataflowEditor/Private/Dataflow/AssetDefinition_DataflowAsset.cpp
--- a/Engine/Plugins/Experimental/Dataflow/Source/DataflowEditor/Private/Dataflow/AssetDefinition_DataflowAsset.cpp
+++ b/Engine/Plugins/Experimental/Dataflow/Source/DataflowEditor/Private/Dataflow/AssetDefinition_DataflowAsset.cpp
@@ -276,7 +276,10 @@ void FDataflowConnectionData::GetNodePropertyAndType(const FString InConnection,
 	if (InConnection.Split(TEXT(":"), &Left, &Right, ESearchCase::IgnoreCase, ESearchDir::FromEnd))
 	{
 		// get the node name from "/NodeName"
-		Left.Split(TEXT("/"), nullptr, &OutNode, ESearchCase::IgnoreCase, ESearchDir::FromEnd);
+		bool bTrimmed = false;
+		OutNode = Left.TrimChar('/', &bTrimmed);
+
+		checkf(bTrimmed, TEXT("Invalid node name received in GetNodePropertyAndType, Attempted to get node name from: %s"), *Left)
 
 		// get the property and type from "Property|Type"
 		const bool bHasType = Right.Split(TEXT("|"), &OutProperty, &OutType, ESearchCase::IgnoreCase, ESearchDir::FromEnd);
