Dataflow : fix comment properties not being completely reflected in dataflow
this was causing
-- tool tip to  not refresh properly
-- option like show bubble when zoomed was not working
this was caused by 2 things
-- we were not calling PostEditPropertychange on the comment structure
-- we were not invaliding properly the cached visualization data (fro tooltips for example )

diff --git a/Engine/Plugins/Experimental/Dataflow/Source/DataflowEditor/Private/Dataflow/DataflowEditorCommands.cpp b/Engine/Plugins/Experimental/Dataflow/Source/DataflowEditor/Private/Dataflow/DataflowEditorCommands.cpp
--- a/Engine/Plugins/Experimental/Dataflow/Source/DataflowEditor/Private/Dataflow/DataflowEditorCommands.cpp
+++ b/Engine/Plugins/Experimental/Dataflow/Source/DataflowEditor/Private/Dataflow/DataflowEditorCommands.cpp
@@ -436,6 +436,21 @@ void FDataflowEditorCommands::OnPropertyValueChanged(UDataflow* OutDataflow, TSh
 	}
 }
 
+void FDataflowEditorCommands::OnCommentPropertyValueChanged(const FPropertyChangedEvent& InPropertyChangedEvent, const TSet<TObjectPtr<UEdGraphNode_Comment>>& SelectedComments)
+{
+	if (InPropertyChangedEvent.ChangeType != EPropertyChangeType::Interactive)
+	{
+		for (UEdGraphNode_Comment* const SelectedComment : SelectedComments)
+		{
+			if (SelectedComment)
+			{
+				FPropertyChangedEvent NonConstEvent = InPropertyChangedEvent;
+				SelectedComment->PostEditChangeProperty(NonConstEvent);
+			}
+		}
+	}
+}
+
 void FDataflowEditorCommands::FreezeNodes(UE::Dataflow::FContext& Context, const FGraphPanelSelectionSet& SelectedNodes)
 {
 	for (UObject* const SelectedNode : SelectedNodes)
diff --git a/Engine/Plugins/Experimental/Dataflow/Source/DataflowEditor/Private/Dataflow/DataflowEditorToolkit.cpp b/Engine/Plugins/Experimental/Dataflow/Source/DataflowEditor/Private/Dataflow/DataflowEditorToolkit.cpp
--- a/Engine/Plugins/Experimental/Dataflow/Source/DataflowEditor/Private/Dataflow/DataflowEditorToolkit.cpp
+++ b/Engine/Plugins/Experimental/Dataflow/Source/DataflowEditor/Private/Dataflow/DataflowEditorToolkit.cpp
@@ -1459,6 +1459,11 @@ void FDataflowEditorToolkit::OnPropertyValueChanged(const FPropertyChangedEvent&
 			
 			FDataflowEditorCommands::OnPropertyValueChanged(DataflowAsset, DataflowContext, LastNodeTimestamp, PropertyChangedEvent, SelectedDataflowNodes);
 
+			if (!SelectedComments.IsEmpty())
+			{
+				FDataflowEditorCommands::OnCommentPropertyValueChanged(PropertyChangedEvent, SelectedComments);
+			}
+
 			// For manual evaluation no need to invalidate the construction scene if we change any parameters
 			// since it will trigger the rendering callbacks with the same cached collection
 			const bool bMakesDirty = (EvaluationMode != EDataflowEditorEvaluationMode::Manual);
@@ -1590,6 +1595,11 @@ void FDataflowEditorToolkit::OnNodeSelectionChanged(const TSet<UObject*>& InNewS
 	TSet<TObjectPtr<UDataflowEdNode>> PreviouslySelectedNodes = SelectedDataflowNodes;
 	SelectedDataflowNodes = FilterDataflowEdNodesFromSet(InNewSelection);
 
+	SelectedComments.Reset();
+	auto IsACommentNode = [](UObject* Obj) { return Cast<UEdGraphNode_Comment>(Obj) != nullptr; };
+	auto CastToCommentNode = [](UObject* Obj) { return Cast<UEdGraphNode_Comment>(Obj); };
+	Algo::TransformIf(InNewSelection, SelectedComments, IsACommentNode, CastToCommentNode);
+
 	// compute the various delta sets
 	TSet<TObjectPtr<UDataflowEdNode>> DeselectedNodes = PreviouslySelectedNodes.Difference(SelectedDataflowNodes);
 	TSet<TObjectPtr<UDataflowEdNode>> StillSelectedNodes = PreviouslySelectedNodes.Intersect(SelectedDataflowNodes);
diff --git a/Engine/Plugins/Experimental/Dataflow/Source/DataflowEditor/Private/Dataflow/DataflowSchema.cpp b/Engine/Plugins/Experimental/Dataflow/Source/DataflowEditor/Private/Dataflow/DataflowSchema.cpp
--- a/Engine/Plugins/Experimental/Dataflow/Source/DataflowEditor/Private/Dataflow/DataflowSchema.cpp
+++ b/Engine/Plugins/Experimental/Dataflow/Source/DataflowEditor/Private/Dataflow/DataflowSchema.cpp
@@ -31,6 +31,8 @@
 
 static const float CDefaultWireThickness = 1.5f;
 
+int32 UDataflowSchema::CurrentCacheRefreshID = 0;
+
 UDataflowSchema::UDataflowSchema()
 {
 }
@@ -864,6 +866,24 @@ FConnectionDrawingPolicy* UDataflowSchema::CreateConnectionDrawingPolicy(int32 I
 	return new FDataflowConnectionDrawingPolicy(InBackLayerID, InFrontLayerID, InZoomFactor, InClippingRect, InDrawElements, InGraphObj);
 }
 
+
+bool UDataflowSchema::IsCacheVisualizationOutOfDate(int32 InVisualizationCacheID) const
+{
+	return CurrentCacheRefreshID != InVisualizationCacheID;
+}
+
+int32 UDataflowSchema::GetCurrentVisualizationCacheID() const
+{
+	return CurrentCacheRefreshID;
+}
+
+void UDataflowSchema::ForceVisualizationCacheClear() const
+{
+	++CurrentCacheRefreshID;
+}
+
+//////////////////////////////////////////////////////////////////////////////////////////////////////////
+
 FDataflowConnectionDrawingPolicy::FDataflowConnectionDrawingPolicy(int32 InBackLayerID, int32 InFrontLayerID, float InZoomFactor, const FSlateRect& InClippingRect, FSlateWindowElementList& InDrawElements, UEdGraph* InGraph)
 	: FConnectionDrawingPolicy(InBackLayerID, InFrontLayerID, InZoomFactor, InClippingRect, InDrawElements)
 	, Schema((UDataflowSchema*)(InGraph->GetSchema()))
diff --git a/Engine/Plugins/Experimental/Dataflow/Source/DataflowEditor/Public/Dataflow/DataflowEditorCommands.h b/Engine/Plugins/Experimental/Dataflow/Source/DataflowEditor/Public/Dataflow/DataflowEditorCommands.h
--- a/Engine/Plugins/Experimental/Dataflow/Source/DataflowEditor/Public/Dataflow/DataflowEditorCommands.h
+++ b/Engine/Plugins/Experimental/Dataflow/Source/DataflowEditor/Public/Dataflow/DataflowEditorCommands.h
@@ -18,6 +18,7 @@ class UDataflowEdNode;
 struct FDataflowNode;
 class UEdGraph;
 class UEdGraphNode;
+class UEdGraphNode_Comment;
 class SDataflowGraphEditor;
 class UDataflowBaseContent;
 
@@ -163,6 +164,7 @@ public:
 	*  OnPropertyValueChanged
 	*/
 	static UE_API void OnPropertyValueChanged(UDataflow* Graph, TSharedPtr<UE::Dataflow::FEngineContext>& Context, UE::Dataflow::FTimestamp& OutLastNodeTimestamp, const FPropertyChangedEvent& PropertyChangedEvent, const TSet<TObjectPtr<UDataflowEdNode>>& NewSelection = {});
+	static UE_API void OnCommentPropertyValueChanged(const FPropertyChangedEvent& PropertyChangedEvent, const TSet<TObjectPtr<UEdGraphNode_Comment>>& CommentSelection = {});
 	static UE_API void OnAssetPropertyValueChanged(TObjectPtr<UDataflowBaseContent> Content, const FPropertyChangedEvent& PropertyChangedEvent);
 
 	/*
diff --git a/Engine/Plugins/Experimental/Dataflow/Source/DataflowEditor/Public/Dataflow/DataflowEditorToolkit.h b/Engine/Plugins/Experimental/Dataflow/Source/DataflowEditor/Public/Dataflow/DataflowEditorToolkit.h
--- a/Engine/Plugins/Experimental/Dataflow/Source/DataflowEditor/Public/Dataflow/DataflowEditorToolkit.h
+++ b/Engine/Plugins/Experimental/Dataflow/Source/DataflowEditor/Public/Dataflow/DataflowEditorToolkit.h
@@ -346,6 +346,7 @@ private:
 
 	// The currently selected set of dataflow nodes. 
 	TSet<TObjectPtr<UDataflowEdNode>> SelectedDataflowNodes;
+	TSet<TObjectPtr<UEdGraphNode_Comment>> SelectedComments;
 
 	/** Pointer to the construction viewport scene. Note this is an alias of ObjectScene in FBaseCharacterFXEditorToolkit but with the specific type */
 	TSharedPtr<FDataflowConstructionScene> ConstructionScene;
diff --git a/Engine/Plugins/Experimental/Dataflow/Source/DataflowEditor/Public/Dataflow/DataflowSchema.h b/Engine/Plugins/Experimental/Dataflow/Source/DataflowEditor/Public/Dataflow/DataflowSchema.h
--- a/Engine/Plugins/Experimental/Dataflow/Source/DataflowEditor/Public/Dataflow/DataflowSchema.h
+++ b/Engine/Plugins/Experimental/Dataflow/Source/DataflowEditor/Public/Dataflow/DataflowSchema.h
@@ -67,6 +67,11 @@ private:
 
 	// used by SupportsDropPinOnNode to know which original connection is being dropped
 	mutable UEdGraphPin* PinBeingDropped = nullptr;
+
+	UE_API virtual bool IsCacheVisualizationOutOfDate(int32 InVisualizationCacheID) const override;
+	UE_API virtual int32 GetCurrentVisualizationCacheID() const override;
+	UE_API virtual void ForceVisualizationCacheClear() const override;
+	static int32 CurrentCacheRefreshID;
 };
 
 class FDataflowConnectionDrawingPolicy : public FConnectionDrawingPolicy, public FGCObject
