Dataflow : update the construction when the Dataflow view settings are changed withoiut the need to relauch the editor

diff --git a/Engine/Plugins/Experimental/Dataflow/Source/DataflowEditor/Private/Dataflow/DataflowConstructionViewportClient.cpp b/Engine/Plugins/Experimental/Dataflow/Source/DataflowEditor/Private/Dataflow/DataflowConstructionViewportClient.cpp
--- a/Engine/Plugins/Experimental/Dataflow/Source/DataflowEditor/Private/Dataflow/DataflowConstructionViewportClient.cpp
+++ b/Engine/Plugins/Experimental/Dataflow/Source/DataflowEditor/Private/Dataflow/DataflowConstructionViewportClient.cpp
@@ -43,6 +43,26 @@ FDataflowConstructionViewportClient::FDataflowConstructionViewportClient(FEditor
 
 	bEnableSceneTicking = bCouldTickScene;
 
+	UpdateFromSettings();
+
+	if (UDataflowEditorOptions* const MutableOptions = GetMutableDefault<UDataflowEditorOptions>())
+	{
+		OnSettingsChangeHandle = MutableOptions->OnSettingChanged().AddLambda([this](UObject*, FPropertyChangedEvent& Event) { UpdateFromSettings(); });
+	}
+
+	WidgetMode = UE::Widget::WM_None;
+	if (ModeTools)
+	{
+		ModeTools->SetWidgetMode(WidgetMode);
+		ModeTools->SetCoordSystem(COORD_Local);
+	}
+}
+
+void FDataflowConstructionViewportClient::UpdateFromSettings()
+{
+	BehaviorsFor2DMode.Reset();
+	OrthoScrollBehaviorTarget.Reset();
+
 	bool bUseRightMouseButton = true;
 	bool bUseMiddleMouseButton = true;
 	if (const UDataflowEditorOptions* const Options = GetDefault<UDataflowEditorOptions>())
@@ -62,9 +82,9 @@ FDataflowConstructionViewportClient::FDataflowConstructionViewportClient(FEditor
 		RightMouseClickDragBehavior->SetDefaultPriority(ViewportBehaviorPriority);
 		RightMouseClickDragBehavior->SetUseRightMouseButton();
 		RightMouseClickDragBehavior->CanBeginClickDragFunc = [](const FInputDeviceRay& InputDeviceRay)
-		{
-			return FInputRayHit(TNumericLimits<float>::Max()); // bHit is true. Depth is max to lose the standard tiebreaker.
-		};
+			{
+				return FInputRayHit(TNumericLimits<float>::Max()); // bHit is true. Depth is max to lose the standard tiebreaker.
+			};
 		BehaviorsFor2DMode.Add(RightMouseClickDragBehavior);
 	}
 
@@ -87,13 +107,20 @@ FDataflowConstructionViewportClient::FDataflowConstructionViewportClient(FEditor
 		ViewFOV = FOVAngle;
 		ExposureSettings.bFixed = Options->bConstructionViewFixedExposure;
 	}
-	WidgetMode = UE::Widget::WM_None;
-	ModeTools->SetWidgetMode(WidgetMode);
-	ModeTools->SetCoordSystem(COORD_Local);
+
+	if (ConstructionViewMode != nullptr)
+	{
+		SetConstructionViewMode(ConstructionViewMode);
+	}
 }
 
 FDataflowConstructionViewportClient::~FDataflowConstructionViewportClient()
 {
+	if (UDataflowEditorOptions* const MutableOptions = GetMutableDefault<UDataflowEditorOptions>())
+	{
+		MutableOptions->OnSettingChanged().Remove(OnSettingsChangeHandle);
+	}
+
 	if (UDataflowEditorOptions* const Options = UDataflowEditorOptions::StaticClass()->GetDefaultObject<UDataflowEditorOptions>())
 	{
 		Options->ConstructionViewFOV = FOVAngle;
@@ -117,7 +144,10 @@ bool FDataflowConstructionViewportClient::CanSetWidgetMode(UE::Widget::EWidgetMo
 
 void FDataflowConstructionViewportClient::SetWidgetMode(UE::Widget::EWidgetMode NewMode)
 {
-	ModeTools->SetWidgetMode(NewMode);
+	if (ModeTools)
+	{
+		ModeTools->SetWidgetMode(NewMode);
+	}
 	WidgetMode = NewMode;
 }
 
@@ -149,7 +179,7 @@ USelection* FDataflowConstructionViewportClient::GetSelectedComponents() const
 			return SceneSelection;
 		}
 	}
-	return ModeTools->GetSelectedComponents();
+	return ModeTools? ModeTools->GetSelectedComponents(): nullptr;
 }
 
 bool FDataflowConstructionViewportClient::InputKey(const FInputKeyEventArgs& EventArgs)
@@ -316,8 +346,11 @@ void FDataflowConstructionViewportClient::SetConstructionViewMode(const UE::Data
 		}
 	}
 
-	ModeTools->GetInteractiveToolsContext()->InputRouter->DeregisterSource(this);
-	ModeTools->GetInteractiveToolsContext()->InputRouter->RegisterSource(this);
+	if (ModeTools)
+	{
+		ModeTools->GetInteractiveToolsContext()->InputRouter->DeregisterSource(this);
+		ModeTools->GetInteractiveToolsContext()->InputRouter->RegisterSource(this);
+	}
 
 
 	SetViewportType((ELevelViewportType)ConstructionViewMode->GetViewportType());
diff --git a/Engine/Plugins/Experimental/Dataflow/Source/DataflowEditor/Public/Dataflow/DataflowConstructionViewportClient.h b/Engine/Plugins/Experimental/Dataflow/Source/DataflowEditor/Public/Dataflow/DataflowConstructionViewportClient.h
--- a/Engine/Plugins/Experimental/Dataflow/Source/DataflowEditor/Public/Dataflow/DataflowConstructionViewportClient.h
+++ b/Engine/Plugins/Experimental/Dataflow/Source/DataflowEditor/Public/Dataflow/DataflowConstructionViewportClient.h
@@ -72,6 +72,8 @@ private:
 	UE_API virtual UE::Widget::EWidgetMode GetWidgetMode() const override;
 	UE_API virtual void SetWidgetMode(UE::Widget::EWidgetMode NewMode) override;
 
+	UE_API void UpdateFromSettings();
+
 	TUniquePtr<IClickDragBehaviorTarget> OrthoScrollBehaviorTarget;
 	TArray<TObjectPtr<UInputBehavior>> BehaviorsFor2DMode;
 
@@ -95,6 +97,8 @@ private:
 
 	/** Currently active transform widget type. */
 	UE::Widget::EWidgetMode WidgetMode = UE::Widget::WM_Translate;
+
+	FDelegateHandle OnSettingsChangeHandle;
 };
 
 #undef UE_API
