Dataflow utility: add method to embed / save a shared dataflow asset

diff --git a/Engine/Plugins/Experimental/Dataflow/Source/DataflowEditor/Private/Dataflow/AssetDefinition_DataflowAsset.cpp b/Engine/Plugins/Experimental/Dataflow/Source/DataflowEditor/Private/Dataflow/AssetDefinition_DataflowAsset.cpp
--- a/Engine/Plugins/Experimental/Dataflow/Source/DataflowEditor/Private/Dataflow/AssetDefinition_DataflowAsset.cpp
+++ b/Engine/Plugins/Experimental/Dataflow/Source/DataflowEditor/Private/Dataflow/AssetDefinition_DataflowAsset.cpp
@@ -7,6 +7,7 @@
 #include "Dataflow/DataflowEditor.h"
 #include "Dataflow/DataflowEditorModule.h"
 #include "Dataflow/DataflowObject.h"
+#include "Dataflow/DataflowInstance.h"
 #include "Dialog/SMessageDialog.h"
 #include "IContentBrowserSingleton.h"
 #include "Math/Color.h"
@@ -24,16 +25,15 @@ FAutoConsoleVariableRef CVarDataflowIsEditable(TEXT("p.Dataflow.IsEditable"), bC
 
 namespace UE::DataflowAssetDefinitionHelpers
 {
-	// Return true if we should proceed, false if we should re-open the dialog
-	bool CreateNewDataflowAsset(const UObject* Asset, UObject*& OutDataflowAsset)
+	static UDataflow* CreateNewDataflowAssetInteractive(const UObject* BoundAsset, TFunction<UDataflow* (UPackage* const NewPackage, FName NewAssetName)> CreateAssetFunction)
 	{
 		const UClass* const DataflowClass = UDataflow::StaticClass();
 
 		FSaveAssetDialogConfig NewDataflowAssetDialogConfig;
 		{
-			const FString PackageName = Asset->GetOutermost()->GetName();
+			const FString PackageName = BoundAsset->GetOutermost()->GetName();
 			NewDataflowAssetDialogConfig.DefaultPath = FPackageName::GetLongPackagePath(PackageName);
-			const FString AssetName = Asset->GetName();
+			const FString AssetName = BoundAsset->GetName();
 			NewDataflowAssetDialogConfig.DefaultAssetName = "DF_" + AssetName;
 			NewDataflowAssetDialogConfig.AssetClassNames.Add(DataflowClass->GetClassPathName());
 			NewDataflowAssetDialogConfig.ExistingAssetPolicy = ESaveAssetDialogExistingAssetPolicy::Disallow;
@@ -49,25 +49,108 @@ namespace UE::DataflowAssetDefinitionHelpers
 			const FString AssetSavePath = ContentBrowserModule.Get().CreateModalSaveAssetDialog(NewDataflowAssetDialogConfig);
 			if (AssetSavePath.IsEmpty())
 			{
-				OutDataflowAsset = nullptr;
-				return false;
+				return nullptr;
 			}
 			NewPackageName = FPackageName::ObjectPathToPackageName(AssetSavePath);
 		}
 
 		const FName NewAssetName(FPackageName::GetLongPackageAssetName(NewPackageName));
 		UPackage* const NewPackage = CreatePackage(*NewPackageName);
-		UObject* const NewAsset = NewObject<UObject>(NewPackage, DataflowClass, NewAssetName, RF_Public | RF_Standalone | RF_Transactional);
 
-		NewAsset->MarkPackageDirty();
+		UDataflow* const NewDataflowAsset = CreateAssetFunction(NewPackage, NewAssetName);
+		if (NewDataflowAsset)
+		{
+			NewDataflowAsset->MarkPackageDirty();
 
-		// Notify the asset registry
-		FAssetRegistryModule::AssetCreated(NewAsset);
+			// Notify the asset registry
+			FAssetRegistryModule::AssetCreated(NewDataflowAsset);
+		}
+		return NewDataflowAsset;
+	}
 
-		OutDataflowAsset = NewAsset;
-		return true;
+	// Return true if we should proceed, false if we should re-open the dialog
+	bool CreateNewDataflowAsset(const UObject* BoundAsset, UObject*& OutDataflowAsset)
+	{
+		OutDataflowAsset = CreateNewDataflowAssetInteractive(
+			BoundAsset,
+			[](UPackage* const NewPackage, FName NewAssetName)
+			{
+				return NewObject<UDataflow>(NewPackage, UDataflow::StaticClass(), NewAssetName, RF_Public | RF_Standalone | RF_Transactional);
+			});
+		return (OutDataflowAsset != nullptr);
 	}
 
+	bool MakeShareableDataflowAsset(UObject* BoundAsset, UDataflow*& OutDataflowAsset)
+	{
+		const UDataflow* DataflowObject = UE::Dataflow::InstanceUtils::GetDataflowAssetFromObject(BoundAsset);
+		if (!DataflowObject)
+		{
+			// fallback to simply create a new one
+			UObject* OutObject = nullptr;
+			if (CreateNewDataflowAsset(BoundAsset, OutObject))
+			{
+				OutDataflowAsset = Cast<UDataflow>(OutObject);
+				return true;
+			}
+			return false;
+		}
+
+		// Is the asset is already a standalone shared asset then do nothing 
+		if (DataflowObject->IsAsset())
+		{
+			return false;
+		}
+
+		IDataflowInstanceInterface* Interface = Cast<IDataflowInstanceInterface>(BoundAsset);
+		if (!Interface)
+		{
+			// the bound asset do not support embedding
+			return false;
+		}
+
+		OutDataflowAsset = CreateNewDataflowAssetInteractive(
+			BoundAsset,
+			[&DataflowObject](UPackage* const NewPackage, FName NewAssetName)
+			{
+				return DuplicateObject<UDataflow>(DataflowObject, NewPackage);
+			});
+
+		if (OutDataflowAsset)
+		{
+			Interface->GetDataflowInstance().SetDataflowAsset(OutDataflowAsset);
+		}
+
+		return (OutDataflowAsset != nullptr);
+	}
+
+	bool EmbedDataflowAsset(UObject* BoundAsset, UDataflow*& OutDataflowAsset)
+	{
+		const UDataflow* ExistingDataflowAsset = UE::Dataflow::InstanceUtils::GetDataflowAssetFromObject(BoundAsset);
+		if (ExistingDataflowAsset && !ExistingDataflowAsset->IsAsset())
+		{
+			// already embedded, nothing to do
+			return false;
+		}
+
+		IDataflowInstanceInterface* Interface = Cast<IDataflowInstanceInterface>(BoundAsset);
+		if (!Interface)
+		{
+			// the bound asset do not support embedding
+			return false;
+		}
+
+		OutDataflowAsset =
+			ExistingDataflowAsset
+			? DuplicateObject<UDataflow>(ExistingDataflowAsset, BoundAsset, TEXT("EmbeddedDataflow"))
+			: NewObject<UDataflow>(BoundAsset, UDataflow::StaticClass(), TEXT("EmbeddedDataflow"), RF_Public | RF_Standalone | RF_Transactional)
+			;
+
+		if (OutDataflowAsset)
+		{
+			Interface->GetDataflowInstance().SetDataflowAsset(OutDataflowAsset);
+		}
+		return (OutDataflowAsset != nullptr);
+	}
 
 	// Return true if we should proceed, false if we should re-open the dialog
 	bool OpenDataflowAsset(const UObject* Asset, UObject*& OutDataflowAsset)
diff --git a/Engine/Plugins/Experimental/Dataflow/Source/DataflowEditor/Public/Dataflow/AssetDefinition_DataflowAsset.h b/Engine/Plugins/Experimental/Dataflow/Source/DataflowEditor/Public/Dataflow/AssetDefinition_DataflowAsset.h
--- a/Engine/Plugins/Experimental/Dataflow/Source/DataflowEditor/Public/Dataflow/AssetDefinition_DataflowAsset.h
+++ b/Engine/Plugins/Experimental/Dataflow/Source/DataflowEditor/Public/Dataflow/AssetDefinition_DataflowAsset.h
@@ -8,10 +8,23 @@
 
 #define UE_API DATAFLOWEDITOR_API
 
+class UDataflow;
+
 namespace UE::DataflowAssetDefinitionHelpers
 {
-	// Return true if we should proceed, false if we should re-open the dialog
-	DATAFLOWEDITOR_API bool CreateNewDataflowAsset(const UObject* Asset, UObject*& OutDataflowAsset);
+	// Create a brand new empty dataflow asset
+	// @return true is the asset is successfully created ( in that case OutDataflowAsset will containt the pointer to this asset )
+	DATAFLOWEDITOR_API bool CreateNewDataflowAsset(const UObject* BoundAsset, UObject*& OutDataflowAsset);
+
+	// Take an embedded dataflow asset from BoundAsset and make it shareable by being it's own seperate asset
+	// If the Dataflow asset was already shareable , do nothing 
+	// @return true if the operation was a success ( in that case OutDataflowAsset will containt the pointer to the shareable Dataflow asset )
+	DATAFLOWEDITOR_API bool MakeShareableDataflowAsset(UObject* BoundAsset, UDataflow*& OutDataflowAsset);
+
+	// Take an shared dataflow asset from BoundAsset and embded it in BoundAsset
+	// If the Dataflow asset was already embedded, do nothing 
+	// @return true if the operation was a success ( in that case OutDataflowAsset will containt the pointer to the Embdedded Dataflow asset )
+	DATAFLOWEDITOR_API bool EmbedDataflowAsset(UObject* BoundAsset, UDataflow*& OutDataflowAsset);
 
 	// Return true if we should proceed, false if we should re-open the dialog
 	DATAFLOWEDITOR_API bool OpenDataflowAsset(const UObject* Asset, UObject*& OutDataflowAsset);
