Dataflow : Add support for embedded/ shared dataflow
- Disable graph editor, tools, variable, subgraph and node details when opening the ditor with an external Dataflow asset
- Add ability to allow edit for external dataflow
- Add buttons to convert back and forth between embed  and external dataflow

diff --git a/Engine/Plugins/ChaosClothAssetEditor/Source/ChaosClothAssetEditor/Private/ChaosClothAsset/ClothEditorToolkit.cpp b/Engine/Plugins/ChaosClothAssetEditor/Source/ChaosClothAssetEditor/Private/ChaosClothAsset/ClothEditorToolkit.cpp
--- a/Engine/Plugins/ChaosClothAssetEditor/Source/ChaosClothAssetEditor/Private/ChaosClothAsset/ClothEditorToolkit.cpp
+++ b/Engine/Plugins/ChaosClothAssetEditor/Source/ChaosClothAssetEditor/Private/ChaosClothAsset/ClothEditorToolkit.cpp
@@ -58,7 +58,7 @@ public:
 	SLATE_ARGUMENT_DEFAULT(UEdGraph*, GraphToEdit) = nullptr;
 	SLATE_ARGUMENT(FGraphEditorEvents, GraphEvents)
 	SLATE_ARGUMENT(TSharedPtr<IStructureDetailsView>, DetailsView)
-	SLATE_ARGUMENT(FDataflowEditorCommands::FGraphEvaluationCallback, EvaluateGraph)
+	SLATE_ARGUMENT(SDataflowGraphEditor::FGraphEvaluationCallback, EvaluateGraph)
 	SLATE_ARGUMENT_DEFAULT(UE::Chaos::ClothAsset::FChaosClothAssetEditorToolkit*, ClothAssetEditorToolkit) = nullptr;
 	SLATE_END_ARGS()
 
@@ -72,6 +72,7 @@ public:
 		Arguments._GraphEvents = InArgs._GraphEvents;
 		Arguments._DetailsView = InArgs._DetailsView;
 		Arguments._EvaluateGraph = InArgs._EvaluateGraph;
+		Arguments._IsEditable = []()->bool { return true; };
 		SDataflowGraphEditor::Construct(Arguments, InAssetOwner);
 	}
 
@@ -1011,6 +1012,7 @@ void FChaosClothAssetEditorToolkit::ReinitializeGraphEditorWidget()
 	Args._GraphEvents = InEvents;
 	Args._DetailsView = NodeDetailsEditor;
 	Args._EvaluateGraph = EvalLambda;
+	Args._IsEditable = []()->bool { return true; };
 
 	GraphEditor->Construct(Args, ClothAsset);
 
diff --git a/Engine/Plugins/Experimental/Dataflow/Source/DataflowEditor/Private/Dataflow/DataflowEditorCommands.cpp b/Engine/Plugins/Experimental/Dataflow/Source/DataflowEditor/Private/Dataflow/DataflowEditorCommands.cpp
--- a/Engine/Plugins/Experimental/Dataflow/Source/DataflowEditor/Private/Dataflow/DataflowEditorCommands.cpp
+++ b/Engine/Plugins/Experimental/Dataflow/Source/DataflowEditor/Private/Dataflow/DataflowEditorCommands.cpp
@@ -80,6 +80,10 @@ void FDataflowEditorCommandsImpl::RegisterCommands()
 	UI_COMMAND(ConvertToForEachSubGraph, "Convert to For Each Subgraph", "Convert the subgraph to be able to iterate through an array.", EUserInterfaceActionType::Button, FInputChord());
 	UI_COMMAND(CollapseToSubGraph, "Collapse to Sub-graph", "Collapse the selection into a subgraph", EUserInterfaceActionType::Button, FInputChord());
 
+	UI_COMMAND(ConvertToEmbeddedDataflow, "Convert to Embedded", "Convert this Dataflow to be embedded in the current asset", EUserInterfaceActionType::Button, FInputChord());
+	UI_COMMAND(ConvertToExternalDataflow, "Convert to External", "Convert this Dataflow to be an external Dataflow asset that can be shared by other asset", EUserInterfaceActionType::Button, FInputChord());
+	UI_COMMAND(ToggleDataflowEditingLock, "Lock/Unlock Editing", "Toggle editing for the current Dataflow graph.", EUserInterfaceActionType::ToggleButton, FInputChord());
+
 	for (const TPair<FName, TUniquePtr<UE::Dataflow::IDataflowConstructionViewMode>>& NameAndMode : UE::Dataflow::FRenderingViewModeFactory::GetInstance().GetViewModes())
 	{
 		TSharedPtr< FUICommandInfo > SetViewModeCommand;
diff --git a/Engine/Plugins/Experimental/Dataflow/Source/DataflowEditor/Private/Dataflow/DataflowEditorMode.cpp b/Engine/Plugins/Experimental/Dataflow/Source/DataflowEditor/Private/Dataflow/DataflowEditorMode.cpp
--- a/Engine/Plugins/Experimental/Dataflow/Source/DataflowEditor/Private/Dataflow/DataflowEditorMode.cpp
+++ b/Engine/Plugins/Experimental/Dataflow/Source/DataflowEditor/Private/Dataflow/DataflowEditorMode.cpp
@@ -240,6 +240,20 @@ void UDataflowEditorMode::RegisterDataflowTool(TSharedPtr<FUICommandInfo> UIComm
 	);
 }
 
+bool UDataflowEditorMode::CanUseTools() const
+{
+	return IsDataflowEditable();
+}
+
+bool UDataflowEditorMode::IsDataflowEditable() const
+{
+	if (TSharedPtr<SDataflowGraphEditor> GraphEditor = DataflowGraphEditor.Pin())
+	{
+		return GraphEditor->IsEditable();
+	}
+	return false;
+}
+
 void UDataflowEditorMode::AddNode(FName NewNodeType)
 {
 	const UE::Dataflow::FDataflowToolRegistry& ToolRegistry = UE::Dataflow::FDataflowToolRegistry::Get();
@@ -262,14 +276,17 @@ void UDataflowEditorMode::AddNode(FName NewNodeType)
 
 bool UDataflowEditorMode::CanAddNode(FName NewNodeType) const
 {
-	const UE::Dataflow::FDataflowToolRegistry& ToolRegistry = UE::Dataflow::FDataflowToolRegistry::Get();
-	const FName ConnectionType = ToolRegistry.HasToolInfoForNodeType(NewNodeType) ? ToolRegistry.GetAddNodeConnectionType(NewNodeType) : FManagedArrayCollection::StaticType();
-	
-	const UEdGraphNode* const CurrentlySelectedNode = GetSingleSelectedNodeWithOutputType(ConnectionType);
-	if (CurrentlySelectedNode && GetToolManager())
+	if (IsDataflowEditable())
 	{
-		const bool bToolActive = GetToolManager()->HasActiveTool(EToolSide::Left);
-		return !bToolActive;
+		const UE::Dataflow::FDataflowToolRegistry& ToolRegistry = UE::Dataflow::FDataflowToolRegistry::Get();
+		const FName ConnectionType = ToolRegistry.HasToolInfoForNodeType(NewNodeType) ? ToolRegistry.GetAddNodeConnectionType(NewNodeType) : FManagedArrayCollection::StaticType();
+
+		const UEdGraphNode* const CurrentlySelectedNode = GetSingleSelectedNodeWithOutputType(ConnectionType);
+		if (CurrentlySelectedNode && GetToolManager())
+		{
+			const bool bToolActive = GetToolManager()->HasActiveTool(EToolSide::Left);
+			return !bToolActive;
+		}
 	}
 	return false;
 }
@@ -657,19 +674,21 @@ void UDataflowEditorMode::ModeTick(float DeltaTime)
 		FirstTimeFocusConstructionViewport();
 	}
 
-	if (!NodeTypeForPendingToolStart.IsNone() && !GetToolManager()->HasActiveTool(EToolSide::Left) && !bPendingNodeSelectionChanged)
+	if (CanUseTools())
 	{
-		const TSharedRef<FUICommandList> CommandList = Toolkit->GetToolkitCommands();
-		const FDataflowEditorCommandsImpl& CommandInfos = FDataflowEditorCommandsImpl::Get();
-
-		if (const TSharedPtr<const FUICommandInfo>* const Command = NodeTypeToToolCommandMap.Find(NodeTypeForPendingToolStart))
+		if (!NodeTypeForPendingToolStart.IsNone() && !GetToolManager()->HasActiveTool(EToolSide::Left) && !bPendingNodeSelectionChanged)
 		{
-			CommandList->TryExecuteAction(Command->ToSharedRef());
-		}
+			const TSharedRef<FUICommandList> CommandList = Toolkit->GetToolkitCommands();
+			const FDataflowEditorCommandsImpl& CommandInfos = FDataflowEditorCommandsImpl::Get();
 
-		NodeTypeForPendingToolStart = FName();
-	}
+			if (const TSharedPtr<const FUICommandInfo>* const Command = NodeTypeToToolCommandMap.Find(NodeTypeForPendingToolStart))
+			{
+				CommandList->TryExecuteAction(Command->ToSharedRef());
+			}
 
+			NodeTypeForPendingToolStart = FName();
+		}
+	}
 
 	if (bShouldRestartToolNextTick)
 	{
diff --git a/Engine/Plugins/Experimental/Dataflow/Source/DataflowEditor/Private/Dataflow/DataflowEditorToolkit.cpp b/Engine/Plugins/Experimental/Dataflow/Source/DataflowEditor/Private/Dataflow/DataflowEditorToolkit.cpp
--- a/Engine/Plugins/Experimental/Dataflow/Source/DataflowEditor/Private/Dataflow/DataflowEditorToolkit.cpp
+++ b/Engine/Plugins/Experimental/Dataflow/Source/DataflowEditor/Private/Dataflow/DataflowEditorToolkit.cpp
@@ -792,7 +792,8 @@ void FDataflowEditorToolkit::PostInitAssetEditor()
 	FToolMenuSection& Section = AssetToolbar->FindOrAddSection("DataflowTools");
 
 	AddEvaluationWidget(Section);
-	
+	AddDataflowActionWidget(Section);
+
 	// Force scenes to update loaded asset viewer settings
 	UAssetViewerSettings::Get()->OnAssetViewerSettingsChanged().Broadcast(NAME_None);
 }
@@ -1087,6 +1088,131 @@ bool FDataflowEditorToolkit::IsEvaluateButtonEnabled() const
 	return (EvaluationMode == EDataflowEditorEvaluationMode::Manual || IsGraphDirty());
 }
 
+void FDataflowEditorToolkit::AddDataflowActionWidget(FToolMenuSection& InSection)
+{
+	InSection.AddSeparator(TEXT("Dataflow Actions"));
+
+	// Convert to embedded Dataflow button
+	ToolkitCommands->MapAction(FDataflowEditorCommands::Get().ConvertToEmbeddedDataflow,
+		FExecuteAction::CreateSP(this, &FDataflowEditorToolkit::ConvertToEmbeddedDataflow),
+		FCanExecuteAction::CreateSP(this, &FDataflowEditorToolkit::IsExternalDataflow),
+		FIsActionChecked(),
+		FIsActionButtonVisible::CreateSP(this, &FDataflowEditorToolkit::IsExternalDataflow)
+	);
+
+	FToolMenuEntry ConvertToEmbeddedButton = FToolMenuEntry::InitToolBarButton(
+		FDataflowEditorCommands::Get().ConvertToEmbeddedDataflow,
+		TAttribute<FText>(),
+		TAttribute<FText>(),
+		FSlateIcon());
+
+	InSection.AddEntry(ConvertToEmbeddedButton);
+
+	// Convert to external Dataflow button
+	ToolkitCommands->MapAction(FDataflowEditorCommands::Get().ConvertToExternalDataflow,
+		FExecuteAction::CreateSP(this, &FDataflowEditorToolkit::ConvertToExternalDataflow),
+		FCanExecuteAction::CreateSP(this, &FDataflowEditorToolkit::IsEmbeddedDataflow),
+		FIsActionChecked(),
+		FIsActionButtonVisible::CreateSP(this, &FDataflowEditorToolkit::IsEmbeddedDataflow)
+	);
+
+	FToolMenuEntry ConvertToExternalButton = FToolMenuEntry::InitToolBarButton(
+		FDataflowEditorCommands::Get().ConvertToExternalDataflow,
+		TAttribute<FText>(),
+		TAttribute<FText>(),
+		FSlateIcon());
+
+	InSection.AddEntry(ConvertToExternalButton);
+
+	// Lock/Unlock editing button
+	ToolkitCommands->MapAction(FDataflowEditorCommands::Get().ToggleDataflowEditingLock,
+		FExecuteAction::CreateSP(this, &FDataflowEditorToolkit::ToggleDataflowEditingLock),
+		FCanExecuteAction::CreateSP(this, &FDataflowEditorToolkit::IsExternalDataflow),
+		FIsActionChecked::CreateSP(this, &FDataflowEditorToolkit::IsEditable),
+		FIsActionButtonVisible::CreateSP(this, &FDataflowEditorToolkit::IsExternalDataflow)
+	);
+
+	const FSlateIcon UnlockedIcon = FSlateIcon(FAppStyle::GetAppStyleSetName(), "Icons.Unlock");
+	const FSlateIcon LockedIcon = FSlateIcon(FAppStyle::GetAppStyleSetName(), "Icons.Lock");
+	const FText UnlockedText = LOCTEXT("LockEditing_Text", "Lock Editing");
+	const FText UnlockedTooltip = LOCTEXT("lockEditing_Tooltip", "Lock Editing for this external Dataflow asset so that other assets using will not be affected by any potential changes (Default behavior)");
+	const FText LockedText = LOCTEXT("UnlockEditing_Text", "Unlock Editing");
+	const FText LockedTooltip = LOCTEXT("UnlockEditing_Tooltip", "Unlock Editing for this external Dataflow asset\nWarning : this asset may be used by other asset and changes made will affect them");
+
+	FToolMenuEntry ToggleEditingLockButton = FToolMenuEntry::InitToolBarButton(
+		FDataflowEditorCommands::Get().ToggleDataflowEditingLock,
+		TAttribute<FText>::CreateSPLambda(this, [this, UnlockedText, LockedText]() { return bAllowExternalDataflowEditing ? UnlockedText : LockedText; } ),
+		TAttribute<FText>::CreateSPLambda(this, [this, UnlockedTooltip, LockedTooltip]() { return bAllowExternalDataflowEditing ? UnlockedTooltip : LockedTooltip; }),
+		TAttribute<FSlateIcon>::CreateSPLambda(this, [this, UnlockedIcon, LockedIcon]() { return bAllowExternalDataflowEditing ? UnlockedIcon : LockedIcon; })
+	);
+
+	InSection.AddEntry(ToggleEditingLockButton);
+
+}
+
+bool FDataflowEditorToolkit::IsEmbeddedDataflow() const
+{
+	if (const TObjectPtr<UDataflowBaseContent>& EditorContent = GetEditorContent())
+	{
+		if (const UDataflow* DataflowAsset = EditorContent->GetDataflowAsset())
+		{
+			return !DataflowAsset->IsAsset();
+		}
+	}
+	return false;
+}
+
+bool FDataflowEditorToolkit::IsExternalDataflow() const
+{
+	return !IsEmbeddedDataflow();
+}
+
+bool FDataflowEditorToolkit::IsEditable() const
+{
+	return (IsEmbeddedDataflow() || bAllowExternalDataflowEditing);
+}
+
+void FDataflowEditorToolkit::ConvertToEmbeddedDataflow()
+{
+	ensure(IsExternalDataflow());
+
+	if (const TObjectPtr<UDataflowBaseContent>& EditorContent = GetEditorContent())
+	{
+		if (TObjectPtr<UObject> BoundAsset = EditorContent->GetDataflowOwner())
+		{
+			UDataflow* NewDataflowAsset = nullptr;
+			if (UE::DataflowAssetDefinitionHelpers::EmbedDataflowAsset(BoundAsset, NewDataflowAsset))
+			{
+				ResetEditor();
+			}
+		}
+	}
+}
+
+void FDataflowEditorToolkit::ConvertToExternalDataflow()
+{
+	ensure(IsEmbeddedDataflow());
+	if (const TObjectPtr<UDataflowBaseContent>& EditorContent = GetEditorContent())
+	{
+		if (TObjectPtr<UObject> BoundAsset = EditorContent->GetDataflowOwner())
+		{
+			UDataflow* NewDataflowAsset = nullptr;
+			if (UE::DataflowAssetDefinitionHelpers::MakeShareableDataflowAsset(BoundAsset, NewDataflowAsset))
+			{
+				ResetEditor();
+			}
+		}
+	}
+}
+
+void FDataflowEditorToolkit::ToggleDataflowEditingLock()
+{
+	if (IsExternalDataflow())
+	{
+		bAllowExternalDataflowEditing = !bAllowExternalDataflowEditing;
+	}
+}
+
 void FDataflowEditorToolkit::InitializeEdMode(UBaseCharacterFXEditorMode* EdMode)
 {
 	UDataflowEditorMode* DataflowMode = Cast<UDataflowEditorMode>(EdMode);
@@ -1130,8 +1256,10 @@ void FDataflowEditorToolkit::GetSaveableObjects(TArray<UObject*>& OutObjects) co
 	{
 		if (UDataflow* DataflowAsset = GetEditorContent()->GetDataflowAsset())
 		{
-			check(DataflowAsset->IsAsset());
-			OutObjects.AddUnique(DataflowAsset);
+			if (DataflowAsset->IsAsset())
+			{
+				OutObjects.AddUnique(DataflowAsset);
+			}
 			if (DataflowAsset->PreviewCacheAsset.IsValid())
 			{
 				OutObjects.AddUnique(DataflowAsset->PreviewCacheAsset.Get());
@@ -1477,17 +1605,16 @@ void FDataflowEditorToolkit::OnPropertyValueChanged(const FPropertyChangedEvent&
 
 void FDataflowEditorToolkit::OnDataflowAssetChanged()
 {
-	const FPropertyChangedEvent PropertyChangedEvent(FDataflowInstance::StaticStruct()->FindPropertyByName(FDataflowInstance::GetDataflowAssetPropertyName()));
-	OnAssetPropertyValueChanged(PropertyChangedEvent);
+	ResetEditor();
 }
 
-void FDataflowEditorToolkit::OnAssetPropertyValueChanged(const FPropertyChangedEvent& PropertyChangedEvent)
+void FDataflowEditorToolkit::ResetEditor()
 {
 	if (const TObjectPtr<UDataflowBaseContent>& EditorContent = GetEditorContent())
 	{
-		if (IDataflowInstanceInterface* DataflowInstanceInterface = Cast<IDataflowInstanceInterface>(EditorContent->GetDataflowOwner()))
+		if (EditorContent)
 		{
-			if (PropertyChangedEvent.GetPropertyName()== FDataflowInstance::GetDataflowAssetPropertyName())
+			if (IDataflowInstanceInterface* DataflowInstanceInterface = Cast<IDataflowInstanceInterface>(EditorContent->GetDataflowOwner()))
 			{
 				// close all subgraphs tabs and selection to make sure no other widget refer to this graphs data
 				const TArray<TSharedPtr<SDockTab>> Tabs = DocumentManager->GetAllDocumentTabs();
@@ -1503,7 +1630,7 @@ void FDataflowEditorToolkit::OnAssetPropertyValueChanged(const FPropertyChangedE
 					GraphEditor->ClearSelectionSet();
 				}
 				bViewsNeedRefresh = true;
-				
+
 				// change the asset and open a new graph in the tab 
 				if (UDataflow* NewDataflowAsset = DataflowInstanceInterface->GetDataflowInstance().GetDataflowAsset())
 				{
@@ -1534,8 +1661,24 @@ void FDataflowEditorToolkit::OnAssetPropertyValueChanged(const FPropertyChangedE
 				}
 			}
 		}
+	}
+
+	// Reset the editing state for external dataflows
+	bAllowExternalDataflowEditing = false;
+}
 
-		ensure(EditorContent);	FDataflowEditorCommands::OnAssetPropertyValueChanged(EditorContent, PropertyChangedEvent);
+void FDataflowEditorToolkit::OnAssetPropertyValueChanged(const FPropertyChangedEvent& PropertyChangedEvent)
+{
+	if (const TObjectPtr<UDataflowBaseContent>& EditorContent = GetEditorContent())
+	{
+		if (PropertyChangedEvent.GetPropertyName()== FDataflowInstance::GetDataflowAssetPropertyName())
+		{
+			ResetEditor();
+		}
+		if (ensure(EditorContent))
+		{
+			FDataflowEditorCommands::OnAssetPropertyValueChanged(EditorContent, PropertyChangedEvent);
+		}
 	}
 }
 
@@ -2213,7 +2356,7 @@ TSharedRef<SDataflowGraphEditor> FDataflowEditorToolkit::CreateGraphEditorWidget
 	ensure(GraphToEdit);
 	using namespace UE::Dataflow;
 
-	const FDataflowEditorCommands::FGraphEvaluationCallback Evaluate =
+	const SDataflowGraphEditor::FGraphEvaluationCallback Evaluate =
 		[this](const FDataflowNode* Node, const FDataflowOutput* Output)
 		{
 			// This method is called when a node is explicily called to be evaluated form the UI 
@@ -2234,7 +2377,7 @@ TSharedRef<SDataflowGraphEditor> FDataflowEditorToolkit::CreateGraphEditorWidget
 				}
 			}
 		};
-	
+
 	DataflowEditor->UpdateTerminalContents(FTimestamp::Invalid);
 	
 	SGraphEditor::FGraphEditorEvents InEvents;
@@ -2251,6 +2394,7 @@ TSharedRef<SDataflowGraphEditor> FDataflowEditorToolkit::CreateGraphEditorWidget
 		.GraphEvents(InEvents)
 		.DetailsView(InNodeDetailsEditor)
 		.EvaluateGraph(Evaluate)
+		.IsEditable([this]() { return IsEditable(); })
 		.DataflowEditor(DataflowEditor);
 
 	OnSelectionChangedMulticastDelegateHandle = NewGraphEditor->OnSelectionChangedMulticast.AddSP(this, &FDataflowEditorToolkit::OnNodeSelectionChanged);
@@ -2291,6 +2435,8 @@ TSharedPtr<IStructureDetailsView> FDataflowEditorToolkit::CreateNodeDetailsEdito
 	LocalDetailsView->GetDetailsView()->SetObject(ObjectToEdit);
 	OnFinishedChangingPropertiesDelegateHandle = LocalDetailsView->GetOnFinishedChangingPropertiesDelegate().AddSP(this, &FDataflowEditorToolkit::OnPropertyValueChanged);
 
+	LocalDetailsView->GetDetailsView()->SetIsPropertyEditingEnabledDelegate(FIsPropertyEditingEnabled::CreateRaw(this, &FDataflowEditorToolkit::IsEditable));
+
 	NodeDetailsExtensionHandler = MakeShared<UE::Dataflow::FDataflowNodeDetailExtensionHandler>();
 	LocalDetailsView->GetDetailsView()->SetExtensionHandler(NodeDetailsExtensionHandler);
 
diff --git a/Engine/Plugins/Experimental/Dataflow/Source/DataflowEditor/Private/Dataflow/DataflowGraphEditor.cpp b/Engine/Plugins/Experimental/Dataflow/Source/DataflowEditor/Private/Dataflow/DataflowGraphEditor.cpp
--- a/Engine/Plugins/Experimental/Dataflow/Source/DataflowEditor/Private/Dataflow/DataflowGraphEditor.cpp
+++ b/Engine/Plugins/Experimental/Dataflow/Source/DataflowEditor/Private/Dataflow/DataflowGraphEditor.cpp
@@ -134,6 +134,7 @@ void SDataflowGraphEditor::Construct(const FArguments& InArgs, UObject* InAssetO
 	DetailsView = InArgs._DetailsView;
 	EvaluateGraphCallback = InArgs._EvaluateGraph;
 	OnDragDropEventCallback = InArgs._OnDragDropEvent;
+	IsEditableCallback = InArgs._IsEditable;
 	DataflowEditor = InArgs._DataflowEditor;
 
 	FGraphAppearanceInfo AppearanceInfo;
@@ -150,7 +151,8 @@ void SDataflowGraphEditor::Construct(const FArguments& InArgs, UObject* InAssetO
 
 			GraphEditorCommands->MapAction(
 				FGenericCommands::Get().Delete,
-				FExecuteAction::CreateSP(this, &SDataflowGraphEditor::DeleteNode)
+				FExecuteAction::CreateSP(this, &SDataflowGraphEditor::DeleteNode),
+				FCanExecuteAction::CreateSP(this, &SDataflowGraphEditor::IsEditable)
 			);
 			GraphEditorCommands->MapAction(
 				FDataflowEditorCommands::Get().EvaluateNode,
@@ -158,55 +160,68 @@ void SDataflowGraphEditor::Construct(const FArguments& InArgs, UObject* InAssetO
 			);
 			GraphEditorCommands->MapAction(
 				FDataflowEditorCommands::Get().FreezeNodes,
-				FExecuteAction::CreateSP(this, &SDataflowGraphEditor::FreezeNodes)
+				FExecuteAction::CreateSP(this, &SDataflowGraphEditor::FreezeNodes),
+				FCanExecuteAction::CreateSP(this, &SDataflowGraphEditor::IsEditable)
 			);
 			GraphEditorCommands->MapAction(
 				FDataflowEditorCommands::Get().UnfreezeNodes,
-				FExecuteAction::CreateSP(this, &SDataflowGraphEditor::UnfreezeNodes)
+				FExecuteAction::CreateSP(this, &SDataflowGraphEditor::UnfreezeNodes),
+				FCanExecuteAction::CreateSP(this, &SDataflowGraphEditor::IsEditable)
 			);
 			GraphEditorCommands->MapAction(
 				FGraphEditorCommands::Get().CreateComment,
-				FExecuteAction::CreateRaw(this, &SDataflowGraphEditor::CreateComment)
+				FExecuteAction::CreateRaw(this, &SDataflowGraphEditor::CreateComment),
+				FCanExecuteAction::CreateSP(this, &SDataflowGraphEditor::IsEditable)
 			);
 			GraphEditorCommands->MapAction(
 				FGraphEditorCommands::Get().AlignNodesTop,
-				FExecuteAction::CreateSP(this, &SDataflowGraphEditor::AlignTop)
+				FExecuteAction::CreateSP(this, &SDataflowGraphEditor::AlignTop),
+				FCanExecuteAction::CreateSP(this, &SDataflowGraphEditor::IsEditable)
 			);
 			GraphEditorCommands->MapAction(
 				FGraphEditorCommands::Get().AlignNodesMiddle,
-				FExecuteAction::CreateSP(this, &SDataflowGraphEditor::AlignMiddle)
+				FExecuteAction::CreateSP(this, &SDataflowGraphEditor::AlignMiddle),
+				FCanExecuteAction::CreateSP(this, &SDataflowGraphEditor::IsEditable)
 			);
 			GraphEditorCommands->MapAction(
 				FGraphEditorCommands::Get().AlignNodesBottom,
-				FExecuteAction::CreateSP(this, &SDataflowGraphEditor::AlignBottom)
+				FExecuteAction::CreateSP(this, &SDataflowGraphEditor::AlignBottom),
+				FCanExecuteAction::CreateSP(this, &SDataflowGraphEditor::IsEditable)
 			);
 			GraphEditorCommands->MapAction(
 				FGraphEditorCommands::Get().AlignNodesLeft,
-				FExecuteAction::CreateSP(this, &SDataflowGraphEditor::AlignLeft)
+				FExecuteAction::CreateSP(this, &SDataflowGraphEditor::AlignLeft),
+				FCanExecuteAction::CreateSP(this, &SDataflowGraphEditor::IsEditable)
 			);
 			GraphEditorCommands->MapAction(
 				FGraphEditorCommands::Get().AlignNodesCenter,
-				FExecuteAction::CreateSP(this, &SDataflowGraphEditor::AlignCenter)
+				FExecuteAction::CreateSP(this, &SDataflowGraphEditor::AlignCenter),
+				FCanExecuteAction::CreateSP(this, &SDataflowGraphEditor::IsEditable)
 			);
 			GraphEditorCommands->MapAction(
 				FGraphEditorCommands::Get().AlignNodesRight,
-				FExecuteAction::CreateSP(this, &SDataflowGraphEditor::AlignRight)
+				FExecuteAction::CreateSP(this, &SDataflowGraphEditor::AlignRight),
+				FCanExecuteAction::CreateSP(this, &SDataflowGraphEditor::IsEditable)
 			);
 			GraphEditorCommands->MapAction(
 				FGraphEditorCommands::Get().StraightenConnections,
-				FExecuteAction::CreateSP(this, &SDataflowGraphEditor::StraightenConnections)
+				FExecuteAction::CreateSP(this, &SDataflowGraphEditor::StraightenConnections),
+				FCanExecuteAction::CreateSP(this, &SDataflowGraphEditor::IsEditable)
 			);
 			GraphEditorCommands->MapAction(
 				FGraphEditorCommands::Get().DistributeNodesHorizontally,
-				FExecuteAction::CreateSP(this, &SDataflowGraphEditor::DistributeHorizontally)
+				FExecuteAction::CreateSP(this, &SDataflowGraphEditor::DistributeHorizontally),
+				FCanExecuteAction::CreateSP(this, &SDataflowGraphEditor::IsEditable)
 			);
 			GraphEditorCommands->MapAction(
 				FGraphEditorCommands::Get().DistributeNodesVertically,
-				FExecuteAction::CreateSP(this, &SDataflowGraphEditor::DistributeVertically)
+				FExecuteAction::CreateSP(this, &SDataflowGraphEditor::DistributeVertically),
+				FCanExecuteAction::CreateSP(this, &SDataflowGraphEditor::IsEditable)
 			);
 			GraphEditorCommands->MapAction(
 				FDataflowEditorCommands::Get().ToggleEnabledState,
-				FExecuteAction::CreateSP(this, &SDataflowGraphEditor::ToggleEnabledState)
+				FExecuteAction::CreateSP(this, &SDataflowGraphEditor::ToggleEnabledState),
+				FCanExecuteAction::CreateSP(this, &SDataflowGraphEditor::IsEditable)
 			);
 			GraphEditorCommands->MapAction(
 				FDataflowEditorCommands::Get().AddOptionPin,
@@ -220,7 +235,8 @@ void SDataflowGraphEditor::Construct(const FArguments& InArgs, UObject* InAssetO
 			);
 			GraphEditorCommands->MapAction(
 				FGenericCommands::Get().Duplicate,
-				FExecuteAction::CreateSP(this, &SDataflowGraphEditor::DuplicateSelectedNodes)
+				FExecuteAction::CreateSP(this, &SDataflowGraphEditor::DuplicateSelectedNodes),
+				FCanExecuteAction::CreateSP(this, &SDataflowGraphEditor::IsEditable)
 			);
 			GraphEditorCommands->MapAction(
 				FDataflowEditorCommands::Get().ZoomToFitGraph,
@@ -229,19 +245,19 @@ void SDataflowGraphEditor::Construct(const FArguments& InArgs, UObject* InAssetO
 			GraphEditorCommands->MapAction(
 				FGraphEditorCommands::Get().ShowAllPins,
 				FExecuteAction::CreateSP(this, &SDataflowGraphEditor::SetPinVisibility, SGraphEditor::Pin_Show),
-				FCanExecuteAction(),
+				FCanExecuteAction::CreateSP(this, &SDataflowGraphEditor::IsEditable),
 				FIsActionChecked::CreateSP(this, &SDataflowGraphEditor::GetPinVisibility, SGraphEditor::Pin_Show)
 			);
 			GraphEditorCommands->MapAction(
 				FGraphEditorCommands::Get().HideNoConnectionPins,
 				FExecuteAction::CreateSP(this, &SDataflowGraphEditor::SetPinVisibility, SGraphEditor::Pin_HideNoConnection),
-				FCanExecuteAction(),
+				FCanExecuteAction::CreateSP(this, &SDataflowGraphEditor::IsEditable),
 				FIsActionChecked::CreateSP(this, &SDataflowGraphEditor::GetPinVisibility, SGraphEditor::Pin_HideNoConnection)
 			);
 			GraphEditorCommands->MapAction(
 				FGraphEditorCommands::Get().HideNoConnectionNoDefaultPins,
 				FExecuteAction::CreateSP(this, &SDataflowGraphEditor::SetPinVisibility, SGraphEditor::Pin_HideNoConnectionNoDefault),
-				FCanExecuteAction(),
+				FCanExecuteAction::CreateSP(this, &SDataflowGraphEditor::IsEditable),
 				FIsActionChecked::CreateSP(this, &SDataflowGraphEditor::GetPinVisibility, SGraphEditor::Pin_HideNoConnectionNoDefault)
 			);
 			GraphEditorCommands->MapAction(
@@ -250,11 +266,13 @@ void SDataflowGraphEditor::Construct(const FArguments& InArgs, UObject* InAssetO
 			);
 			GraphEditorCommands->MapAction(
 				FGenericCommands::Get().Cut,
-				FExecuteAction::CreateSP(this, &SDataflowGraphEditor::CutSelectedNodes)
+				FExecuteAction::CreateSP(this, &SDataflowGraphEditor::CutSelectedNodes),
+				FCanExecuteAction::CreateSP(this, &SDataflowGraphEditor::IsEditable)
 			);
 			GraphEditorCommands->MapAction(
 				FGenericCommands::Get().Paste,
-				FExecuteAction::CreateSP(this, &SDataflowGraphEditor::PasteSelectedNodes)
+				FExecuteAction::CreateSP(this, &SDataflowGraphEditor::PasteSelectedNodes),
+				FCanExecuteAction::CreateSP(this, &SDataflowGraphEditor::IsEditable)
 			);
 			GraphEditorCommands->MapAction(
 				FGenericCommands::Get().Rename,
@@ -263,11 +281,13 @@ void SDataflowGraphEditor::Construct(const FArguments& InArgs, UObject* InAssetO
 			);
 			GraphEditorCommands->MapAction(
 				FDataflowEditorCommands::Get().AddNewVariable,
-				FExecuteAction::CreateSP(this, &SDataflowGraphEditor::AddNewVariable)
+				FExecuteAction::CreateSP(this, &SDataflowGraphEditor::AddNewVariable),
+				FCanExecuteAction::CreateSP(this, &SDataflowGraphEditor::IsEditable)
 			);
 			GraphEditorCommands->MapAction(
 					FDataflowEditorCommands::Get().AddNewSubGraph,
-					FExecuteAction::CreateSP(this, &SDataflowGraphEditor::AddNewSubGraph)
+					FExecuteAction::CreateSP(this, &SDataflowGraphEditor::AddNewSubGraph),
+				FCanExecuteAction::CreateSP(this, &SDataflowGraphEditor::IsEditable)
 			);
 			GraphEditorCommands->MapAction(FGraphEditorCommands::Get().StartWatchingPin,
 				FExecuteAction::CreateSP(this, &SDataflowGraphEditor::OnStartWatchingPin),
@@ -287,12 +307,13 @@ void SDataflowGraphEditor::Construct(const FArguments& InArgs, UObject* InAssetO
 		}
 	}
 
-
 	SGraphEditor::FArguments Arguments;
 	Arguments._AdditionalCommands = GraphEditorCommands;
 	Arguments._Appearance = AppearanceInfo;
 	Arguments._GraphToEdit = InArgs._GraphToEdit;
 	Arguments._GraphEvents = InArgs._GraphEvents;
+	Arguments.IsEditable_Lambda([this]()->bool { return IsEditable(); });
+	Arguments.DisplayAsReadOnly_Lambda([this]()->bool { return !IsEditable(); });
 
 	ensureMsgf(!Arguments._GraphEvents.OnSelectionChanged.IsBound(), TEXT("DataflowGraphEditor::OnSelectionChanged rebound during construction."));
 	Arguments._GraphEvents.OnSelectionChanged = SGraphEditor::FOnSelectionChanged::CreateSP(this, &SDataflowGraphEditor::OnSelectedNodesChanged);
@@ -340,6 +361,11 @@ SDataflowGraphEditor::~SDataflowGraphEditor()
 	}
 }
 
+bool SDataflowGraphEditor::IsEditable() const
+{
+	return IsEditableCallback ? IsEditableCallback() : false;
+}
+
 FActionMenuContent SDataflowGraphEditor::OnCreateActionMenu(UEdGraph* Graph, const FVector2f& Position, const TArray<UEdGraphPin*>& DraggedPins, bool bAutoExpandActionMenu, SGraphEditor::FActionMenuClosed OnClosed)
 {
 	LastActionMenuGraphEditor = StaticCastSharedRef<SDataflowGraphEditor>(AsShared()).ToWeakPtr();
@@ -552,6 +578,10 @@ void SDataflowGraphEditor::RenameNode()
 
 bool SDataflowGraphEditor::CanRenameNode() const
 {
+	if (!IsEditable())
+	{
+		return false;
+	}
 	const FGraphPanelSelectionSet SelectedNodes = GetSelectedNodes();
 	if (SelectedNodes.Num() == 1)
 	{
@@ -641,7 +671,7 @@ FReply SDataflowGraphEditor::OnKeyDown(const FGeometry& MyGeometry, const FKeyEv
 	{
 		VKeyDown = true;
 	}
-	return SGraphEditor::OnKeyUp(MyGeometry, InKeyEvent);
+	return SGraphEditor::OnKeyDown(MyGeometry, InKeyEvent);
 }
 
 bool SDataflowGraphEditor::IsControlDown() const
@@ -777,7 +807,12 @@ void SDataflowGraphEditor::AlignRight()
 
 void SDataflowGraphEditor::StraightenConnections()
 {
-	GetGraphEditor()->OnStraightenConnections();
+	// somehow the StraightenConnections menu entry does not seems to accound for the editable state of the graph and remain accessible despite the readonly state 
+	// to protect from any unvoluntary change to the graph let's make sure we ignore the command if iot still goes through 
+	if (IsEditable())
+	{
+		GetGraphEditor()->OnStraightenConnections();
+	}
 }
 
 void SDataflowGraphEditor::DistributeHorizontally()
@@ -871,7 +906,7 @@ bool SDataflowGraphEditor::CanAddOptionPin() const
 		}
 	}
 
-	return bCanAddOptionPin;
+	return bCanAddOptionPin && IsEditable();
 }
 
 void SDataflowGraphEditor::OnRemoveOptionPin()
@@ -947,7 +982,7 @@ bool SDataflowGraphEditor::CanRemoveOptionPin() const
 		}
 	}
 
-	return bCanRemoveOptionPin;
+	return bCanRemoveOptionPin && IsEditable();
 }
 
 void SDataflowGraphEditor::DuplicateSelectedNodes()
@@ -1122,7 +1157,7 @@ bool SDataflowGraphEditor::CanPromoteToVariable() const
 {
 	if (UEdGraphPin* Pin = const_cast<SDataflowGraphEditor*>(this)->GetGraphPinForMenu())
 	{
-		return !(Pin->HasAnyConnections());
+		return !(Pin->HasAnyConnections()) && IsEditable();
 	}
 	return false;
 }
@@ -1210,7 +1245,7 @@ void SDataflowGraphEditor::InitEvaluationProgressBar()
 				Context->GetAsyncEvaluationStats(NumPendingTasks, NumRunningTasks, NumCompletedTasks);
 
 				ProgressText = FText::Format(
-					LOCTEXT("DataflowGraphEditor_ProgressText", "Pending: {1} | Running: {0} | Completed: {2}"),
+					LOCTEXT("DataflowGraphEditor_ProgressText", "Pending: {0} | Running: {1} | Completed: {2}"),
 					NumPendingTasks,
 					NumRunningTasks,
 					NumCompletedTasks
@@ -1230,7 +1265,10 @@ void SDataflowGraphEditor::InitEvaluationProgressBar()
 				Context->GetAsyncEvaluationStats(NumPendingTasks, NumRunningTasks, NumCompletedTasks);
 
 				const int32 TotalTasks = (NumPendingTasks + NumRunningTasks + NumCompletedTasks);
-				Percent = static_cast<float>(NumCompletedTasks) / static_cast<float>(TotalTasks);
+				if (TotalTasks > 0)
+				{
+					Percent = static_cast<float>(NumCompletedTasks) / static_cast<float>(TotalTasks);
+				}
 			}
 			return Percent;
 		};
diff --git a/Engine/Plugins/Experimental/Dataflow/Source/DataflowEditor/Private/Dataflow/DataflowMembersWidget.cpp b/Engine/Plugins/Experimental/Dataflow/Source/DataflowEditor/Private/Dataflow/DataflowMembersWidget.cpp
--- a/Engine/Plugins/Experimental/Dataflow/Source/DataflowEditor/Private/Dataflow/DataflowMembersWidget.cpp
+++ b/Engine/Plugins/Experimental/Dataflow/Source/DataflowEditor/Private/Dataflow/DataflowMembersWidget.cpp
@@ -395,6 +395,15 @@ const TSharedPtr<SDataflowGraphEditor> SDataflowMembersWidget::GetGraphEditor()
 	return {};
 }
 
+bool SDataflowMembersWidget::IsDataflowAssetEditable() const
+{
+	if (TSharedPtr<SDataflowGraphEditor> GraphEditor = GetGraphEditor())
+	{
+		return GraphEditor->IsEditable();
+	}
+	return false;
+}
+
 void SDataflowMembersWidget::InitializeCommands()
 {
 	const TSharedPtr<SDataflowGraphEditor> GraphEditor = GetGraphEditor();
@@ -434,11 +443,11 @@ void SDataflowMembersWidget::InitializeCommands()
 
 	CommandList->MapAction(FDataflowEditorCommands::Get().ConvertToBasicSubGraph,
 		FExecuteAction::CreateSP(this, &SDataflowMembersWidget::SetForEachSubGraphOnSelection, false),
-		FCanExecuteAction(), FIsActionChecked());
+		FCanExecuteAction::CreateSP(this, &SDataflowMembersWidget::IsDataflowAssetEditable));
 
 	CommandList->MapAction(FDataflowEditorCommands::Get().ConvertToForEachSubGraph,
 		FExecuteAction::CreateSP(this, &SDataflowMembersWidget::SetForEachSubGraphOnSelection, true),
-		FCanExecuteAction(), FIsActionChecked());
+		FCanExecuteAction::CreateSP(this, &SDataflowMembersWidget::IsDataflowAssetEditable));
 }
 
 void SDataflowMembersWidget::InitializeSections()
@@ -678,7 +687,10 @@ TSharedRef<SWidget> SDataflowMembersWidget::OnCreateWidgetForAction(FCreateWidge
 	{
 		if (const TSharedPtr<ISection> Section = GetSectionById(InCreateData->Action->SectionID))
 		{
-			return Section->CreateWidgetForAction(InCreateData, GetGraphEditor());
+			TSharedRef<SWidget> Widget = Section->CreateWidgetForAction(InCreateData, GetGraphEditor());
+			TAttribute<bool> SetEnabledAttribute(this, &SDataflowMembersWidget::IsDataflowAssetEditable);
+			Widget->SetEnabled(SetEnabledAttribute);
+			return Widget;
 		}
 	}
 	return SNullWidget::NullWidget;
@@ -848,8 +860,7 @@ FReply SDataflowMembersWidget::OnAddButtonClickedOnSection(int32 InSectionID)
 
 bool SDataflowMembersWidget::CanAddNewElementToSection(int32 InSectionID) const
 {
-	// for now always allowed
-	return true;
+	return IsDataflowAssetEditable();
 }
 
 bool SDataflowMembersWidget::HandleActionMatchesName(FEdGraphSchemaAction* InAction, const FName& InName) const
@@ -926,11 +937,17 @@ bool SDataflowMembersWidget::CanRequestRename() const
 
 bool SDataflowMembersWidget::CanRequestRenameOnActionNode(TWeakPtr<FGraphActionNode> InSelectedNode) const
 {
-	if (TSharedPtr<FEdGraphSchemaAction> Action = InSelectedNode.Pin()->Action)
+	if (IsDataflowAssetEditable())
 	{
-		if (const TSharedPtr<ISection> Section = GetSectionById(Action->GetSectionID()))
+		if (TSharedPtr<FGraphActionNode> ActionNode = InSelectedNode.Pin())
 		{
-			return Section->CanRequestRename();
+			if (TSharedPtr<FEdGraphSchemaAction> Action = ActionNode->Action)
+			{
+				if (const TSharedPtr<ISection> Section = GetSectionById(Action->GetSectionID()))
+				{
+					return Section->CanRequestRename();
+				}
+			}
 		}
 	}
 	return false;
@@ -973,11 +990,14 @@ void SDataflowMembersWidget::OnCut()
 
 bool SDataflowMembersWidget::CanCut() const
 {
-	if (const TSharedPtr<FEdGraphSchemaAction> Action = GetFirstSelectedAction())
+	if (IsDataflowAssetEditable())
 	{
-		if (const TSharedPtr<ISection> Section = GetSectionById(Action->GetSectionID()))
+		if (const TSharedPtr<FEdGraphSchemaAction> Action = GetFirstSelectedAction())
 		{
-			return Section->CanCopy() && Section->CanDelete();
+			if (const TSharedPtr<ISection> Section = GetSectionById(Action->GetSectionID()))
+			{
+				return Section->CanCopy() && Section->CanDelete();
+			}
 		}
 	}
 	return false;
@@ -996,6 +1016,11 @@ void SDataflowMembersWidget::OnPaste()
 
 bool SDataflowMembersWidget::CanPaste() const
 {
+	if (!IsDataflowAssetEditable())
+	{
+		return false;
+	}
+	
 	if (const TSharedPtr<FEdGraphSchemaAction> Action = GetFirstSelectedAction())
 	{
 		if (const TSharedPtr<ISection> Section = GetSectionById(Action->GetSectionID()))
@@ -1003,6 +1028,7 @@ bool SDataflowMembersWidget::CanPaste() const
 			return Section->CanPaste();
 		}
 	}
+
 	// paste is allowed anywhere on the empty space of the widget 
 	return true;
 }
@@ -1020,11 +1046,14 @@ void SDataflowMembersWidget::OnDuplicate()
 
 bool SDataflowMembersWidget::CanDuplicate() const
 {
-	if (const TSharedPtr<FEdGraphSchemaAction> Action = GetFirstSelectedAction())
+	if (IsDataflowAssetEditable())
 	{
-		if (const TSharedPtr<ISection> Section = GetSectionById(Action->GetSectionID()))
+		if (const TSharedPtr<FEdGraphSchemaAction> Action = GetFirstSelectedAction())
 		{
-			return Section->CanDuplicate();
+			if (const TSharedPtr<ISection> Section = GetSectionById(Action->GetSectionID()))
+			{
+				return Section->CanDuplicate();
+			}
 		}
 	}
 	return false;
@@ -1043,11 +1072,14 @@ void SDataflowMembersWidget::OnDelete()
 
 bool SDataflowMembersWidget::CanDelete() const
 {
-	if (const TSharedPtr<FEdGraphSchemaAction> Action = GetFirstSelectedAction())
+	if (IsDataflowAssetEditable())
 	{
-		if (const TSharedPtr<ISection> Section = GetSectionById(Action->GetSectionID()))
+		if (const TSharedPtr<FEdGraphSchemaAction> Action = GetFirstSelectedAction())
 		{
-			return Section->CanDelete();
+			if (const TSharedPtr<ISection> Section = GetSectionById(Action->GetSectionID()))
+			{
+				return Section->CanDelete();
+			}
 		}
 	}
 	return false;
diff --git a/Engine/Plugins/Experimental/Dataflow/Source/DataflowEditor/Public/Dataflow/DataflowEditorCommands.h b/Engine/Plugins/Experimental/Dataflow/Source/DataflowEditor/Public/Dataflow/DataflowEditorCommands.h
--- a/Engine/Plugins/Experimental/Dataflow/Source/DataflowEditor/Public/Dataflow/DataflowEditorCommands.h
+++ b/Engine/Plugins/Experimental/Dataflow/Source/DataflowEditor/Public/Dataflow/DataflowEditorCommands.h
@@ -86,6 +86,9 @@ public:
 	TSharedPtr<FUICommandInfo> ConvertToForEachSubGraph;
 	TSharedPtr<FUICommandInfo> ToggleAsyncEvaluation;
 	TSharedPtr<FUICommandInfo> CollapseToSubGraph;
+	TSharedPtr<FUICommandInfo> ConvertToEmbeddedDataflow;
+	TSharedPtr<FUICommandInfo> ConvertToExternalDataflow;
+	TSharedPtr<FUICommandInfo> ToggleDataflowEditingLock;
 
 	TMap<FName, TSharedPtr<FUICommandInfo>> SetConstructionViewModeCommands;
 	TMap<FName, TSharedPtr<FUICommandInfo>> ToolPaletteCommands;
@@ -103,7 +106,10 @@ public:
 class FDataflowEditorCommands
 {
 public:
+	UE_DEPRECATED(5.8, "Use SDataflowGraphEditor::FGraphEvaluationCallback instead ")
 	typedef TFunction<void(const FDataflowNode*, const FDataflowOutput*)> FGraphEvaluationCallback;
+
+	UE_DEPRECATED(5.8, "Use SDataflowGraphEditor::FOnDragDropEventCallback instead ")
 	typedef TFunction<void(const FGeometry& MyGeometry, const FDragDropEvent& DragDropEvent)> FOnDragDropEventCallback;
 
 	static UE_API void Register();
diff --git a/Engine/Plugins/Experimental/Dataflow/Source/DataflowEditor/Public/Dataflow/DataflowEditorMode.h b/Engine/Plugins/Experimental/Dataflow/Source/DataflowEditor/Public/Dataflow/DataflowEditorMode.h
--- a/Engine/Plugins/Experimental/Dataflow/Source/DataflowEditor/Public/Dataflow/DataflowEditorMode.h
+++ b/Engine/Plugins/Experimental/Dataflow/Source/DataflowEditor/Public/Dataflow/DataflowEditorMode.h
@@ -187,6 +187,9 @@ private:
 	/** Get the tool categories from the construction scene */
 	UE_API const TArray<FName>& GetToolCategories();
 
+	bool CanUseTools() const;
+	bool IsDataflowEditable() const;
+
 	/** Dataflow construction scene from the toolkit */
 	TWeakPtr<FDataflowConstructionScene> ConstructionScene;
 
diff --git a/Engine/Plugins/Experimental/Dataflow/Source/DataflowEditor/Public/Dataflow/DataflowEditorToolkit.h b/Engine/Plugins/Experimental/Dataflow/Source/DataflowEditor/Public/Dataflow/DataflowEditorToolkit.h
--- a/Engine/Plugins/Experimental/Dataflow/Source/DataflowEditor/Public/Dataflow/DataflowEditorToolkit.h
+++ b/Engine/Plugins/Experimental/Dataflow/Source/DataflowEditor/Public/Dataflow/DataflowEditorToolkit.h
@@ -315,6 +315,32 @@ private:
 	/** Update the debug draw based on a change of currently selected or pinned nodes */
 	UE_API void UpdateDebugDraw();
 
+	/** reset the editor. Used when the dataflow graph as changed */
+	void ResetEditor();
+
+	/** This adds the widget to switch from embdedded or external dataflow as well as read - only lock when applicable */
+	void AddDataflowActionWidget(FToolMenuSection& Section);
+
+	/** return true if the Dataflow is not an asset and is an object embedded into the owner asset */
+	bool IsEmbeddedDataflow() const;
+
+	/** return true if the Dataflow is a seperate asset or no dataflow is assigned to the owner asset */
+	bool IsExternalDataflow() const;
+
+	/** whether the dataflow is in an editable state */
+	bool IsEditable() const;
+
+	/** Duplicate the current external dataflow and embed it into the owner asset */
+	void ConvertToEmbeddedDataflow();
+
+	/** 
+	* Duplicate the embedded dataflow object and store it into a seperate asset  
+	* this methid is interactive and will prompt a dialog for the user to choose a path to save the duplicated Dataflow to 
+	*/
+	void ConvertToExternalDataflow();
+
+	void ToggleDataflowEditingLock();
+
 	void UpdateViewsFromNode(UDataflowEdNode* Node = nullptr);
 	static TSet<TObjectPtr<UDataflowEdNode>> FilterDataflowEdNodesFromSet(const TSet<UObject*>& Set);
 	static const TObjectPtr<UDataflowEdNode> GetOnlyFromSet(const TSet<TObjectPtr<UDataflowEdNode>>& Set);
@@ -377,6 +403,9 @@ private:
 	/** Cached value of the p.Dataflow.EnableGraphEval cvar, to avoid calling FindConsoleVariable too often */
 	bool bDataflowEnableGraphEval;
 
+	/** Allow editing of external dataflow */
+	bool bAllowExternalDataflowEditing = false;
+
 	/** Cached value from the editor evaluation settings */
 	bool bAllowEvaluationInPIE = false;
 
diff --git a/Engine/Plugins/Experimental/Dataflow/Source/DataflowEditor/Public/Dataflow/DataflowGraphEditor.h b/Engine/Plugins/Experimental/Dataflow/Source/DataflowEditor/Public/Dataflow/DataflowGraphEditor.h
--- a/Engine/Plugins/Experimental/Dataflow/Source/DataflowEditor/Public/Dataflow/DataflowGraphEditor.h
+++ b/Engine/Plugins/Experimental/Dataflow/Source/DataflowEditor/Public/Dataflow/DataflowGraphEditor.h
@@ -42,6 +42,10 @@ class SDataflowGraphEditor : public SGraphEditor, public FGCObject, public FData
 {
 public:
 
+	using FGraphEvaluationCallback = TFunction<void(const FDataflowNode*, const FDataflowOutput*)>;
+	using FOnDragDropEventCallback = TFunction<void(const FGeometry& MyGeometry, const FDragDropEvent& DragDropEvent)>;
+	using FIsEditableCallback = TFunction<bool()>;
+
 	SLATE_BEGIN_ARGS(SDataflowGraphEditor)
 		: _AdditionalCommands(static_cast<FUICommandList*>(nullptr))
 		, _GraphToEdit(nullptr)
@@ -52,8 +56,9 @@ public:
 	SLATE_ARGUMENT_DEFAULT(UEdGraph*, GraphToEdit) = nullptr;
 	SLATE_ARGUMENT(FGraphEditorEvents, GraphEvents)
 	SLATE_ARGUMENT(TSharedPtr<IStructureDetailsView>, DetailsView)
-	SLATE_ARGUMENT(FDataflowEditorCommands::FGraphEvaluationCallback, EvaluateGraph)
-	SLATE_ARGUMENT(FDataflowEditorCommands::FOnDragDropEventCallback, OnDragDropEvent)
+	SLATE_ARGUMENT(FGraphEvaluationCallback, EvaluateGraph)
+	SLATE_ARGUMENT(FOnDragDropEventCallback, OnDragDropEvent)
+	SLATE_ARGUMENT(FIsEditableCallback, IsEditable)
 	SLATE_ARGUMENT_DEFAULT(UDataflowEditor*, DataflowEditor) = nullptr;
 	SLATE_END_ARGS()
 
@@ -81,6 +86,9 @@ public:
 	/** */
 	UE_API void Construct(const FArguments& InArgs, UObject* AssetOwner);
 
+	/** return whether or not the graph is editable */
+	UE_API bool IsEditable() const;
+
 	/** */
 	UE_API void EvaluateNode();
 
@@ -205,8 +213,9 @@ private:
 	/** get the dataflow asset from the edGraph being edited */
 	UE_API UDataflow* GetDataflowAsset() const;
 
-	FDataflowEditorCommands::FOnDragDropEventCallback OnDragDropEventCallback;
-	FDataflowEditorCommands::FGraphEvaluationCallback EvaluateGraphCallback;
+	FOnDragDropEventCallback OnDragDropEventCallback;
+	FGraphEvaluationCallback EvaluateGraphCallback;
+	FIsEditableCallback IsEditableCallback;
 
 	UE_API FActionMenuContent OnCreateActionMenu(UEdGraph* Graph, const FVector2f& Position, const TArray<UEdGraphPin*>& DraggedPins, bool bAutoExpandActionMenu, SGraphEditor::FActionMenuClosed OnClosed);
 	UE_API void OnActionMenuFilterByAssetTypeChanged(ECheckBoxState NewState, const TWeakPtr<SGraphEditorActionMenu> WeakActionMenu);
diff --git a/Engine/Plugins/Experimental/Dataflow/Source/DataflowEditor/Public/Dataflow/DataflowMembersWidget.h b/Engine/Plugins/Experimental/Dataflow/Source/DataflowEditor/Public/Dataflow/DataflowMembersWidget.h
--- a/Engine/Plugins/Experimental/Dataflow/Source/DataflowEditor/Public/Dataflow/DataflowMembersWidget.h
+++ b/Engine/Plugins/Experimental/Dataflow/Source/DataflowEditor/Public/Dataflow/DataflowMembersWidget.h
@@ -127,6 +127,8 @@ private:
 	bool CanDelete() const;
 	bool IsSelectionForEachSubGraph() const;
 	void SetForEachSubGraphOnSelection(bool bValue);
+	
+	bool IsDataflowAssetEditable() const;
 
 	// UI refresh 
 	void Refresh();
