Dataflow add method to register various asset menus : Embed Dataflow & Re-Evaluate dataflow
- Update geometry collection to add those entries to their asset right click meny

diff --git a/Engine/Plugins/Experimental/Dataflow/Source/DataflowEditor/DataflowEditor.Build.cs b/Engine/Plugins/Experimental/Dataflow/Source/DataflowEditor/DataflowEditor.Build.cs
--- a/Engine/Plugins/Experimental/Dataflow/Source/DataflowEditor/DataflowEditor.Build.cs
+++ b/Engine/Plugins/Experimental/Dataflow/Source/DataflowEditor/DataflowEditor.Build.cs
@@ -18,7 +18,8 @@ namespace UnrealBuildTool.Rules
 				new string[]
 				{
 					"TedsOutliner",
-					"SkeletalMeshModelingTools"
+					"SkeletalMeshModelingTools",
+					"ContentBrowser",
 				}
 			);
 
diff --git a/Engine/Plugins/Experimental/Dataflow/Source/DataflowEditor/Private/Dataflow/AssetDefinition_DataflowAsset.cpp b/Engine/Plugins/Experimental/Dataflow/Source/DataflowEditor/Private/Dataflow/AssetDefinition_DataflowAsset.cpp
--- a/Engine/Plugins/Experimental/Dataflow/Source/DataflowEditor/Private/Dataflow/AssetDefinition_DataflowAsset.cpp
+++ b/Engine/Plugins/Experimental/Dataflow/Source/DataflowEditor/Private/Dataflow/AssetDefinition_DataflowAsset.cpp
@@ -14,6 +14,11 @@
 #include "Misc/FileHelper.h"
 #include "ThumbnailRendering/SceneThumbnailInfo.h"
 
+#include "ContentBrowserMenuContexts.h"
+#include "ToolMenu.h"
+#include "ToolMenus.h"
+
+
 #include UE_INLINE_GENERATED_CPP_BY_NAME(AssetDefinition_DataflowAsset)
 
 
@@ -112,7 +117,7 @@ namespace UE::DataflowAssetDefinitionHelpers
 			BoundAsset,
 			[&DataflowObject](UPackage* const NewPackage, FName NewAssetName)
 			{
-				return DuplicateObject<UDataflow>(DataflowObject, NewPackage);
+				return DuplicateObject<UDataflow>(DataflowObject, NewPackage, NewAssetName);
 			});
 
 		if (OutDataflowAsset)
@@ -222,6 +227,143 @@ namespace UE::DataflowAssetDefinitionHelpers
 		const UDataflow* const DataflowAsset = Cast<UDataflow>(Asset);
 		return DataflowAsset && (bCanEditDataflow || DataflowAsset->Type == EDataflowType::Simulation);
 	}
+
+	//////////////////////////////////////////////////////////////////////////////////////////////////////////////////////////////////////
+
+	static void ExecuteEmbedDataflow(const FToolMenuContext& InContext)
+	{
+		if (const UContentBrowserAssetContextMenuContext* Context = UContentBrowserAssetContextMenuContext::FindContextWithAssets(InContext))
+		{
+			TArray<UObject*> SelectedAssets = Context->LoadSelectedObjects<UObject>();
+			if (!SelectedAssets.IsEmpty())
+			{
+				for (UObject* SelectedAsset : SelectedAssets)
+				{
+					if (SelectedAsset)
+					{
+						UDataflow* NewDataflow = nullptr;
+						if (UE::DataflowAssetDefinitionHelpers::EmbedDataflowAsset(SelectedAsset, NewDataflow))
+						{
+							SelectedAsset->Modify();
+						}
+					}
+				}
+			}
+		}
+	}
+
+	static bool CanExecuteEmbedDataflow(const FToolMenuContext& InContext)
+	{
+		if (const UContentBrowserAssetContextMenuContext* Context = UContentBrowserAssetContextMenuContext::FindContextWithAssets(InContext))
+		{
+			TArray<UObject*> SelectedAssets = Context->LoadSelectedObjects<UObject>();
+			if (!SelectedAssets.IsEmpty())
+			{
+				for (UObject* SelectedAsset : SelectedAssets)
+				{
+					if (IDataflowInstanceInterface* Interface = Cast<IDataflowInstanceInterface>(SelectedAsset))
+					{
+						if (UDataflow* DataflowObject = Interface->GetDataflowInstance().GetDataflowAsset())
+						{
+							if (DataflowObject->IsAsset())
+							{
+								return true;
+							}
+						}
+					}
+				}
+			}
+		}
+		return false;
+	}
+
+	void RegisterEmbedDataflowAssetMenu(UClass* AssetClass)
+	{
+		UToolMenus::RegisterStartupCallback(
+			FSimpleMulticastDelegate::FDelegate::CreateLambda([AssetClass]()
+				{
+					FToolMenuOwnerScoped OwnerScoped(UE_MODULE_NAME);
+					UToolMenu* Menu = UE::ContentBrowser::ExtendToolMenu_AssetContextMenu(AssetClass);
+					FToolMenuSection& Section = Menu->FindOrAddSection("GetAssetActions");
+					{
+						const TAttribute<FText> Label = LOCTEXT("EmbedDataflow_Label", "Embed Dataflow");
+						const TAttribute<FText> ToolTip = LOCTEXT("EmbedDataflow_Tooltip", "Embed the external Dataflow graph in the current asset, the external asset will remain but will be disconnected");
+						const TAttribute<FSlateIcon> Icon = FSlateIcon(FAppStyle::GetAppStyleSetName(), "GraphEditor.EventGraph_16x");
+
+						FToolUIAction UIAction;
+						UIAction.ExecuteAction = FToolMenuExecuteAction::CreateStatic(&ExecuteEmbedDataflow);
+						UIAction.CanExecuteAction = FToolMenuCanExecuteAction::CreateStatic(&CanExecuteEmbedDataflow);
+						Section.AddMenuEntry("Dataflow_EmbedDataflow", Label, ToolTip, Icon, UIAction);
+					}
+				}));
+	}
+
+	//////////////////////////////////////////////////////////////////////////////////////////////////////////////////////////////////////
+
+	static void ExecuteReEvaluateDataflow(const FToolMenuContext& InContext)
+	{
+		if (const UContentBrowserAssetContextMenuContext* Context = UContentBrowserAssetContextMenuContext::FindContextWithAssets(InContext))
+		{
+			TArray<UObject*> SelectedAssets = Context->LoadSelectedObjects<UObject>();
+			if (!SelectedAssets.IsEmpty())
+			{
+				for (UObject* SelectedAsset : SelectedAssets)
+				{
+					if (IDataflowInstanceInterface* Interface = Cast<IDataflowInstanceInterface>(SelectedAsset))
+					{
+						if (Interface->GetDataflowInstance().UpdateOwnerAsset(/*bUpdateDependentAsset*/true))
+						{
+							SelectedAsset->Modify();
+						}
+					}
+				}
+			}
+		}
+	}
+
+	static bool CanReEvaluateDataflow(const FToolMenuContext& InContext)
+	{
+		if (const UContentBrowserAssetContextMenuContext* Context = UContentBrowserAssetContextMenuContext::FindContextWithAssets(InContext))
+		{
+			TArray<UObject*> SelectedAssets = Context->LoadSelectedObjects<UObject>();
+			if (!SelectedAssets.IsEmpty())
+			{
+				for (UObject* SelectedAsset : SelectedAssets)
+				{
+					if (IDataflowInstanceInterface* Interface = Cast<IDataflowInstanceInterface>(SelectedAsset))
+					{
+						if (UDataflow* DataflowObject = Interface->GetDataflowInstance().GetDataflowAsset())
+						{
+							return true;
+						}
+					}
+				}
+			}
+		}
+		return false;
+	}
+
+	void RegisterReEvaluateDataflowAssetMenu(UClass* AssetClass)
+	{
+		UToolMenus::RegisterStartupCallback(
+			FSimpleMulticastDelegate::FDelegate::CreateLambda([AssetClass]()
+				{
+					FToolMenuOwnerScoped OwnerScoped(UE_MODULE_NAME);
+					UToolMenu* Menu = UE::ContentBrowser::ExtendToolMenu_AssetContextMenu(AssetClass);
+					FToolMenuSection& Section = Menu->FindOrAddSection("GetAssetActions");
+					{
+						const TAttribute<FText> Label = LOCTEXT("ReEvaluateDataflow_Label", "Re-Evaluate Dataflow");
+						const TAttribute<FText> ToolTip = LOCTEXT("ReEvaluateDataflow_Tooltip", "Re-evaluate the associated Dataflow to regenerate the selected asset(s)");
+						const TAttribute<FSlateIcon> Icon = FSlateIcon(FAppStyle::GetAppStyleSetName(), "GraphEditor.EventGraph_16x");
+
+						FToolUIAction UIAction;
+						UIAction.ExecuteAction = FToolMenuExecuteAction::CreateStatic(&ExecuteReEvaluateDataflow);
+						UIAction.CanExecuteAction = FToolMenuCanExecuteAction::CreateStatic(&CanReEvaluateDataflow);
+						Section.AddMenuEntry("Dataflow_ReEvaluateDataflow", Label, ToolTip, Icon, UIAction);
+					}
+				}));
+	}
+
 }
 
 
@@ -374,4 +516,5 @@ void FDataflowConnectionData::GetNodePropertyAndType(const FString InConnection,
 	}
 }
 
+
 #undef LOCTEXT_NAMESPACE
diff --git a/Engine/Plugins/Experimental/Dataflow/Source/DataflowEditor/Public/Dataflow/AssetDefinition_DataflowAsset.h b/Engine/Plugins/Experimental/Dataflow/Source/DataflowEditor/Public/Dataflow/AssetDefinition_DataflowAsset.h
--- a/Engine/Plugins/Experimental/Dataflow/Source/DataflowEditor/Public/Dataflow/AssetDefinition_DataflowAsset.h
+++ b/Engine/Plugins/Experimental/Dataflow/Source/DataflowEditor/Public/Dataflow/AssetDefinition_DataflowAsset.h
@@ -37,6 +37,12 @@ namespace UE::DataflowAssetDefinitionHelpers
 
 	// Determine whether or not a UDataflow can be opened directly in the Dataflow Editor
 	DATAFLOWEDITOR_API bool CanOpenDataflowAssetInEditor(const UObject* DataflowAsset);
+
+	// Register a entry on the asset menu that will embed an external dataflow in the selected asset
+	DATAFLOWEDITOR_API void RegisterEmbedDataflowAssetMenu(UClass* AssetClass);
+
+	// Register a entry on the asset menu that will Re-Evaluate the assoiated dataflow for the selected asset
+	DATAFLOWEDITOR_API void RegisterReEvaluateDataflowAssetMenu(UClass* AssetClass);
 }
 
 struct FDataflowInput;
diff --git a/Engine/Plugins/Experimental/GeometryCollectionPlugin/Source/GeometryCollectionEditor/Private/GeometryCollection/AssetDefinition_GeometryCollection.cpp b/Engine/Plugins/Experimental/GeometryCollectionPlugin/Source/GeometryCollectionEditor/Private/GeometryCollection/AssetDefinition_GeometryCollection.cpp
--- a/Engine/Plugins/Experimental/GeometryCollectionPlugin/Source/GeometryCollectionEditor/Private/GeometryCollection/AssetDefinition_GeometryCollection.cpp
+++ b/Engine/Plugins/Experimental/GeometryCollectionPlugin/Source/GeometryCollectionEditor/Private/GeometryCollection/AssetDefinition_GeometryCollection.cpp
@@ -215,9 +215,10 @@ namespace MenuExtension_GeometryCollection
 
 	}
 
-	static FDelayedAutoRegisterHelper DelayedAutoRegister(EDelayedRegisterRunPhase::EndOfEngineInit, [] {
-		UToolMenus::RegisterStartupCallback(
-			FSimpleMulticastDelegate::FDelegate::CreateLambda([]()
+	static FDelayedAutoRegisterHelper DelayedAutoRegister(EDelayedRegisterRunPhase::EndOfEngineInit, []()
+		{
+			UToolMenus::RegisterStartupCallback(
+				FSimpleMulticastDelegate::FDelegate::CreateLambda([]()
 				{
 					FToolMenuOwnerScoped OwnerScoped(UE_MODULE_NAME);
 					UToolMenu* Menu = UE::ContentBrowser::ExtendToolMenu_AssetContextMenu(UGeometryCollection::StaticClass());
@@ -229,9 +230,13 @@ namespace MenuExtension_GeometryCollection
 
 						FToolUIAction UIAction;
 						UIAction.ExecuteAction = FToolMenuExecuteAction::CreateStatic(&ExecuteExportToSkeletalMesh);
-						Section.AddMenuEntry("ControlRigPose_SelectControls", Label, ToolTip, Icon, UIAction);
+						Section.AddMenuEntry("GeometryCollection_ConvertToSkeletalMesh", Label, ToolTip, Icon, UIAction);
 					}
 				}));
+
+			UE::DataflowAssetDefinitionHelpers::RegisterEmbedDataflowAssetMenu(UGeometryCollection::StaticClass());
+
+			UE::DataflowAssetDefinitionHelpers::RegisterReEvaluateDataflowAssetMenu(UGeometryCollection::StaticClass());
 		});
 }
 
