Dataflow:
- Added Array DynamicMesh rendering support

diff --git a/Engine/Plugins/Experimental/Dataflow/Source/DataflowEditor/Private/DataflowRendering/DataflowDynamicMeshRenderableType.cpp b/Engine/Plugins/Experimental/Dataflow/Source/DataflowEditor/Private/DataflowRendering/DataflowDynamicMeshRenderableType.cpp
--- a/Engine/Plugins/Experimental/Dataflow/Source/DataflowEditor/Private/DataflowRendering/DataflowDynamicMeshRenderableType.cpp
+++ b/Engine/Plugins/Experimental/Dataflow/Source/DataflowEditor/Private/DataflowRendering/DataflowDynamicMeshRenderableType.cpp
@@ -36,6 +36,31 @@ namespace UE::Dataflow::Private
 		}
 	};
 
+	////////////////////////////////////////////////////////////////////////////////////////////////////////////////////////////
+	
+	class FDynamicMeshArraySurfaceRenderableType : public IRenderableType
+	{
+		UE_DATAFLOW_IRENDERABLE_PRIMARY_TYPE(TArray<TObjectPtr<UDynamicMesh>>, DynamicMeshArray);
+		UE_DATAFLOW_IRENDERABLE_RENDER_GROUP(Surface);
+		UE_DATAFLOW_IRENDERABLE_VIEW_MODE(FDataflowConstruction3DViewMode);
+
+		virtual void GetPrimitiveComponents(const FRenderableTypeInstance& Instance, FRenderableComponents& OutComponents) const override
+		{
+			const TArray<TObjectPtr<UDynamicMesh>>& DynamicMeshArray = GetDynamicMeshArray(Instance, {});
+
+			for (const TObjectPtr<UDynamicMesh>& DynamicMesh : DynamicMeshArray)
+			{
+				if (DynamicMesh)
+				{
+					if (UDynamicMeshComponent* Component = OutComponents.AddNewComponent<UDynamicMeshComponent>())
+					{
+						Component->SetDynamicMesh(DynamicMesh);
+					}
+				}
+			}
+		}
+	};
+
 	////////////////////////////////////////////////////////////////////////////////////////////////////////////////////////////
 
 	class FDynamicMeshUvsRenderableType : public IRenderableType
@@ -58,6 +83,10 @@ namespace UE::Dataflow::Private
 	void RegisterDynamicMeshRenderableTypes()
 	{
 		UE_DATAFLOW_REGISTER_RENDERABLE_TYPE(FDynamicMeshSurfaceRenderableType);
+		UE_DATAFLOW_REGISTER_RENDERABLE_TYPE(FDynamicMeshArraySurfaceRenderableType);
 		UE_DATAFLOW_REGISTER_RENDERABLE_TYPE(FDynamicMeshUvsRenderableType);
 	}
+
+	///////////////////////////////////////////////////////////////////////////////////////////////////////////////
+
 }
\ No newline at end of file
diff --git a/Engine/Plugins/Experimental/GeometryCollectionPlugin/Source/GeometryCollectionNodes/Public/Dataflow/GeometryCollectionDebugNodes.h b/Engine/Plugins/Experimental/GeometryCollectionPlugin/Source/GeometryCollectionNodes/Public/Dataflow/GeometryCollectionDebugNodes.h
--- a/Engine/Plugins/Experimental/GeometryCollectionPlugin/Source/GeometryCollectionNodes/Public/Dataflow/GeometryCollectionDebugNodes.h
+++ b/Engine/Plugins/Experimental/GeometryCollectionPlugin/Source/GeometryCollectionNodes/Public/Dataflow/GeometryCollectionDebugNodes.h
@@ -189,7 +189,7 @@ private:
 	int32 VerticesAlongEachSide = 4;
 
 	/** Single mesh aggregating all the spheres together */
-	UPROPERTY(meta = (DataflowOutput))
+	UPROPERTY(meta = (DataflowOutput, DataflowRenderGroups = ""))
 	TObjectPtr<UDynamicMesh> Mesh;
 
 	/** Array of meshes for each spheres found */
