Attempt to fix Content Worker Test build failing because of an ensure in registering the tools
- wrapping the ensure in a WITH_EDITOR in order to avoid throwing an build error

diff --git a/Engine/Plugins/Experimental/Dataflow/Source/DataflowEditor/Private/Dataflow/DataflowToolRegistry.cpp b/Engine/Plugins/Experimental/Dataflow/Source/DataflowEditor/Private/Dataflow/DataflowToolRegistry.cpp
--- a/Engine/Plugins/Experimental/Dataflow/Source/DataflowEditor/Private/Dataflow/DataflowToolRegistry.cpp
+++ b/Engine/Plugins/Experimental/Dataflow/Source/DataflowEditor/Private/Dataflow/DataflowToolRegistry.cpp
@@ -21,10 +21,11 @@ namespace UE::Dataflow
 	void FDataflowToolRegistry::AddNodeToToolMapping(const FName& NodeName, TObjectPtr<UInteractiveToolBuilder> ToolBuilder, const TSharedRef<const IDataflowToolActionCommands>& ToolActionCommands,
 		const FSlateIcon& AddNodeButtonIcon, const FText& AddNodeButtonText, const FName& ToolCategory, const FName& AddNodeConnectionType, const FName& AddNodeConnectionName)
 	{
+#if WITH_EDITOR
 		ensureMsgf(!FDataflowEditorCommands::IsRegistered(), TEXT("FDataflowToolRegistry: DataflowEditorCommands have already been registered. \
 				Newly registered Tools may not be available in the Editor. \
 				Ensure that AddNodeToToolMapping is called before the DataflowEditor module is loaded."));
-
+#endif
 		// FUICommandInfo is uninitialized, it will be created later in FDataflowEditorCommandsImpl::RegisterCommands
 		NodeTypeToToolMap.Add(NodeName, FToolInfo{ ToolBuilder, ToolActionCommands, AddNodeButtonIcon, AddNodeButtonText, AddNodeConnectionType, AddNodeConnectionName, nullptr, nullptr, ToolCategory });
 	}
