Dataflow Editor : Add a thumbnail extension to add an indication that assets are bound to a dataflow asset
- this also show the widget in the tooltip with the indication of the path or embed state of the dataflow
- this only works if the host asset has the right asset regsitry tags : HasDataflow / IsDataflowEmbedded

diff --git a/Engine/Plugins/Experimental/Dataflow/Source/DataflowEditor/Private/Dataflow/DataflowEditorModule.cpp b/Engine/Plugins/Experimental/Dataflow/Source/DataflowEditor/Private/Dataflow/DataflowEditorModule.cpp
--- a/Engine/Plugins/Experimental/Dataflow/Source/DataflowEditor/Private/Dataflow/DataflowEditorModule.cpp
+++ b/Engine/Plugins/Experimental/Dataflow/Source/DataflowEditor/Private/Dataflow/DataflowEditorModule.cpp
@@ -11,6 +11,7 @@
 #include "Dataflow/DataflowFreezeActionsCustomization.h"
 #include "Dataflow/DataflowFunctionProperty.h"
 #include "Dataflow/DataflowFunctionPropertyCustomization.h"
+#include "Dataflow/DataflowInstance.h"
 #include "Dataflow/DataflowSNodeFactories.h"
 #include "Dataflow/DebugDrawText3dVisualization.h"
 #include "Dataflow/MeshStatsConstructionVisualization.h"
@@ -41,6 +42,11 @@
 #include "DataflowRendering/DataflowGeometryCollectionRenderableType.h"
 #include "DataflowRendering/DataflowGeometryCollectionProximityRenderableType.h"
 
+#include "ContentBrowserModule.h"
+#include "Modules/ModuleManager.h"
+#include "Styling/SlateIconFinder.h"
+#include "Widgets/Images/SImage.h"
+
 #define LOCTEXT_NAMESPACE "DataflowEditor"
 
 namespace UE::Dataflow::Private
@@ -143,6 +149,73 @@ namespace UE::Dataflow::Private
 			}
 		}
 	};
+
+	TSharedRef<SWidget> OnGenerateDataflowAssetBindingIcons(const FAssetData& AssetData)
+	{
+		static const FName HasDataflowTag = TEXT("HasDataflow");
+		bool bHasDataflowAsset = false;
+		AssetData.GetTagValue(HasDataflowTag, bHasDataflowAsset);
+		if (bHasDataflowAsset)
+		{
+			const FSlateBrush* DataflowAssetIconBrush = FSlateIcon(FAppStyle::GetAppStyleSetName(), "GraphEditor.EventGraph_16x").GetIcon();
+
+			// Wrap the icon in an outer background layer to improve visibility on light thumbnails
+			const float Padding = 4.f;
+			const float OuterBoxSize = 16.f;
+			return SNew(SBox)
+				.Padding(Padding, Padding, Padding, Padding)
+				[
+					SNew(SBorder)
+					.BorderImage(FAppStyle::GetBrush("AssetThumbnail", ".AssetThumbnailStatusBar"))
+					[
+						SNew(SBox)
+						.WidthOverride(OuterBoxSize)
+						.HeightOverride(OuterBoxSize)
+						.HAlign(HAlign_Center)
+						.VAlign(VAlign_Center)
+						[
+							SNew(SImage)
+							.Image(DataflowAssetIconBrush)
+						]
+					]
+				];
+		}
+		return SNullWidget::NullWidget;
+	}
+
+	TSharedRef<SWidget> OnGenerateDataflowAssetBindingTooltip(const FAssetData& AssetData)
+	{
+		static const FName HasDataflowTag = TEXT("HasDataflow");
+		bool bHasDataflowAsset = false;
+		AssetData.GetTagValue(HasDataflowTag, bHasDataflowAsset);
+
+		static const FName IsDataflowEmbeddedTag = TEXT("IsDataflowEmbedded");
+		bool bIsDataflowEmbedded = false;
+		AssetData.GetTagValue(IsDataflowEmbeddedTag, bIsDataflowEmbedded);
+
+		if (bHasDataflowAsset)
+		{
+			FText TooltipText;
+			if (bIsDataflowEmbedded)
+			{
+				TooltipText = LOCTEXT("DataflowIndicatorTooltip_EmbeddedText", "Embedded Dataflow Graph");
+			}
+			else
+			{
+				
+				const FText DataflowAssetPathName = FText::FromString(AssetData.GetObjectPathString());
+				TooltipText = FText::Format(LOCTEXT("DataflowIndicatorTooltip_Format", "Shared Dataflow Asset : {0}"), DataflowAssetPathName);
+			}
+
+			const float Padding = 8.f;
+			const float OuterBoxSize = 16.f;
+			return SNew(STextBlock)
+				.Text(TooltipText)
+				.Font(FAppStyle::GetFontStyle("ContentBrowser.Tooltip.EntryFont"))
+				;
+		}
+		return SNullWidget::NullWidget;
+	}
 }
 
 const FColor FDataflowEditorModule::SurfaceColor = FLinearColor(0.6, 0.6, 0.6).ToRGBE();
@@ -201,12 +274,30 @@ void FDataflowEditorModule::StartupModule()
 	UE::Dataflow::Private::RegisterDynamicMeshRenderableTypes();
 	UE::Dataflow::Private::RegisterGeometryCollectionRenderableTypes();
 	UE::Dataflow::Private::RegisterGeometryCollectionProximityRenderableTypes();
+
+	// register an browser extension to show when an asset is bound to a dataflow asset
+	if (FContentBrowserModule* ContentBrowserModule = FModuleManager::Get().GetModulePtr<FContentBrowserModule>(TEXT("ContentBrowser")))
+	{
+		AssetViewExtraStateGeneratorHandle = ContentBrowserModule->AddAssetViewExtraStateGenerator(FAssetViewExtraStateGenerator(
+			FOnGenerateAssetViewExtraStateIndicators::CreateStatic(&UE::Dataflow::Private::OnGenerateDataflowAssetBindingIcons),
+			FOnGenerateAssetViewExtraStateIndicators::CreateStatic(&UE::Dataflow::Private::OnGenerateDataflowAssetBindingTooltip)
+		));
+	}
+
 }
 
 void FDataflowEditorModule::ShutdownModule()
 {
 	using namespace UE::Dataflow::Private;
 
+	if (FContentBrowserModule* ContentBrowserModule = FModuleManager::Get().GetModulePtr<FContentBrowserModule>(TEXT("ContentBrowser")))
+	{
+		if (AssetViewExtraStateGeneratorHandle.IsValid())
+		{
+			ContentBrowserModule->RemoveAssetViewExtraStateGenerator(AssetViewExtraStateGeneratorHandle);
+		}
+	}
+
 	FEditorModeRegistry::Get().UnregisterMode(UDataflowEditorMode::EM_DataflowEditorModeId);
 
 	// Deregister type customizations
diff --git a/Engine/Plugins/Experimental/Dataflow/Source/DataflowEditor/Public/Dataflow/DataflowEditorModule.h b/Engine/Plugins/Experimental/Dataflow/Source/DataflowEditor/Public/Dataflow/DataflowEditorModule.h
--- a/Engine/Plugins/Experimental/Dataflow/Source/DataflowEditor/Public/Dataflow/DataflowEditorModule.h
+++ b/Engine/Plugins/Experimental/Dataflow/Source/DataflowEditor/Public/Dataflow/DataflowEditorModule.h
@@ -26,6 +26,7 @@ public:
 private:
 
 	TSharedPtr<FDataflowSNodeFactory> DataflowSNodeFactory;
+	FDelegateHandle AssetViewExtraStateGeneratorHandle;
 
 };
 
