Dataflow - Added BP variables to the Simulation Scene.

#rnx
#Synced-CL 48035595

diff --git a/Engine/Plugins/Experimental/Dataflow/Source/DataflowEditor/Private/Dataflow/DataflowSimulationControls.cpp b/Engine/Plugins/Experimental/Dataflow/Source/DataflowEditor/Private/Dataflow/DataflowSimulationControls.cpp
--- a/Engine/Plugins/Experimental/Dataflow/Source/DataflowEditor/Private/Dataflow/DataflowSimulationControls.cpp
+++ b/Engine/Plugins/Experimental/Dataflow/Source/DataflowEditor/Private/Dataflow/DataflowSimulationControls.cpp
@@ -11,6 +11,8 @@
 #include "Dataflow/DataflowObject.h"
 #include "Features/IModularFeatures.h"
 #include "Components/SkeletalMeshComponent.h"
+#include "Engine/Blueprint.h"
+#include "Engine/BlueprintGeneratedClass.h"
 #include "Engine/SkeletalMesh.h"
 #include "Engine/World.h"
 #include "BonePose.h"
@@ -40,7 +42,8 @@ namespace UE::Dataflow
 	
 	TObjectPtr<AActor> SpawnSimulatedActor(const TSubclassOf<AActor>& ActorClass,
 		const TObjectPtr<AChaosCacheManager>& CacheManager, const TObjectPtr<UChaosCacheCollection>& CacheCollection,
-		const bool bIsRecording, const TObjectPtr<UDataflowBaseContent>& DataflowContent, const FTransform& ActorTransform)
+		const bool bIsRecording, const TObjectPtr<UDataflowBaseContent>& DataflowContent, const FTransform& ActorTransform,
+		const FInstancedPropertyBag& BlueprintVariables)
 	{
 		if(CacheManager)
 		{
@@ -59,6 +62,31 @@ namespace UE::Dataflow
 			TObjectPtr<AActor> PreviewActor = CacheManager->GetWorld()->SpawnActor<AActor>(ActorClass, SpawnParameters);
 			if(PreviewActor)
 			{
+				// Init BP variables
+				if (UBlueprintGeneratedClass* const BlueprintGeneratedClass = Cast<UBlueprintGeneratedClass>(PreviewActor->GetClass()))
+				{
+					if (UBlueprint* const Blueprint = Cast<UBlueprint>(BlueprintGeneratedClass->ClassGeneratedBy))
+					{
+						for (FBPVariableDescription& BPVariableDescription : Blueprint->NewVariables)
+						{
+							if (const FProperty* const BlueprintProperty = BlueprintGeneratedClass->FindPropertyByName(BPVariableDescription.VarName))
+							{
+								const FName VarGuidAsName(BPVariableDescription.VarGuid.ToString());  // VarGuid instead of the VarName, because the variable name can change and FInstancedPropertyBag::SanitizePropertyName
+								if (const FPropertyBagPropertyDesc* const PropertyBagPropertyDesc = BlueprintVariables.FindPropertyDescByName(VarGuidAsName))
+								{
+									const FProperty* PropertyBagProperty = PropertyBagPropertyDesc->CachedProperty;
+									if (BlueprintProperty->SameType(PropertyBagProperty))
+									{
+										const uint8* const SourceMemory = BlueprintVariables.GetValue().GetMemory() + PropertyBagProperty->GetOffset_ForInternal();
+										uint8* const TargetMemory = BlueprintProperty->ContainerPtrToValuePtr<uint8>(PreviewActor);
+										BlueprintProperty->CopyCompleteValue(TargetMemory, SourceMemory);
+									}
+								}
+							}
+						}
+					}
+				}
+
 				// Link the editor content properties to the BP actor one 
 				DataflowContent->SetActorProperties(PreviewActor);
 
@@ -92,7 +120,18 @@ namespace UE::Dataflow
 		}
 		return nullptr;
 	}
-	
+
+	// UE_DEPRECATED(5.8, "Use BlueprintVariables version instead.")
+	TObjectPtr<AActor> SpawnSimulatedActor(const TSubclassOf<AActor>& ActorClass,
+		const TObjectPtr<AChaosCacheManager>& CacheManager, const TObjectPtr<UChaosCacheCollection>& CacheCollection,
+		const bool bIsRecording, const TObjectPtr<UDataflowBaseContent>& DataflowContent, const FTransform& ActorTransform)
+	{
+		return SpawnSimulatedActor(ActorClass,
+			CacheManager, CacheCollection,
+			bIsRecording, DataflowContent, ActorTransform,
+			FInstancedPropertyBag());
+	}
+
 	void SetupSkeletonAnimation(const TObjectPtr<AActor>& PreviewActor, const bool bSkeletalMeshVisibility)
 	{
 		if(PreviewActor)
diff --git a/Engine/Plugins/Experimental/Dataflow/Source/DataflowEditor/Private/Dataflow/DataflowSimulationGenerator.cpp b/Engine/Plugins/Experimental/Dataflow/Source/DataflowEditor/Private/Dataflow/DataflowSimulationGenerator.cpp
--- a/Engine/Plugins/Experimental/Dataflow/Source/DataflowEditor/Private/Dataflow/DataflowSimulationGenerator.cpp
+++ b/Engine/Plugins/Experimental/Dataflow/Source/DataflowEditor/Private/Dataflow/DataflowSimulationGenerator.cpp
@@ -65,7 +65,8 @@ void FDataflowSimulationTask::DoWork()
 
 bool FDataflowTaskManager::AllocateSimulationResource(const FVector2f& TimeRange, const int32 FrameRate,
 		const TObjectPtr<UChaosCacheCollection>& CacheAsset, const TSubclassOf<AActor>& ActorClass,
-		const TObjectPtr<UDataflowBaseContent>& DataflowContent, const FTransform& BlueprintTransform, const bool bSkeletalMeshVisibility, float DeltaTime)
+		const TObjectPtr<UDataflowBaseContent>& DataflowContent, const FTransform& BlueprintTransform, 
+		const FInstancedPropertyBag& BlueprintVariables, const bool bSkeletalMeshVisibility, float DeltaTime)
 {
 	SimulationWorld = UWorld::CreateWorld(EWorldType::Editor, false);
 	SimulationWorld->bPostTickComponentUpdate = false;
@@ -76,7 +77,7 @@ bool FDataflowTaskManager::AllocateSimulationResource(const FVector2f& TimeRange
 	CacheManager = SimulationWorld->SpawnActor<AChaosCacheManager>();
 	UDataflowSimulationManager* DataflowManager = SimulationWorld->GetSubsystem<UDataflowSimulationManager>();
 
-	PreviewActor = UE::Dataflow::SpawnSimulatedActor(ActorClass, CacheManager, CacheAsset, true, DataflowContent, BlueprintTransform);
+	PreviewActor = UE::Dataflow::SpawnSimulatedActor(ActorClass, CacheManager, CacheAsset, true, DataflowContent, BlueprintTransform, BlueprintVariables);
 	UE::Dataflow::SetupSkeletonAnimation(PreviewActor, bSkeletalMeshVisibility);
 
 	// Set simulation restart time
@@ -124,6 +125,18 @@ bool FDataflowTaskManager::AllocateSimulationResource(const FVector2f& TimeRange
 	return true;
 }
 
+// UE_DEPRECATED(5.8, "Use BlueprintVariables version instead.")
+bool FDataflowTaskManager::AllocateSimulationResource(const FVector2f& TimeRange, const int32 FrameRate,
+	const TObjectPtr<UChaosCacheCollection>& CacheAsset, const TSubclassOf<AActor>& ActorClass,
+	const TObjectPtr<UDataflowBaseContent>& DataflowContent, const FTransform& BlueprintTransform,
+	const bool bSkeletalMeshVisibility, float DeltaTime)
+{
+	return AllocateSimulationResource(TimeRange, FrameRate,
+		CacheAsset, ActorClass,
+		DataflowContent, BlueprintTransform,
+		FInstancedPropertyBag(), bSkeletalMeshVisibility, DeltaTime);
+}
+
 void FDataflowTaskManager::FreeSimulationResource()
 {
 	if (SimulationTask.IsValid())
@@ -205,7 +218,10 @@ void FDataflowSimulationGenerator::StartGenerateSimulation()
 	TaskManager->SimulationTask->GetTask().bRestartSimulation = CacheParams.bRestartSimulation;
 	FVector2f TimeRange = CacheParams.bRestartSimulation ? CacheParams.RestartTimeRange : CacheParams.TimeRange;
 	TaskManager->StartTime = FDateTime::UtcNow();
-	TaskManager->AllocateSimulationResource(TimeRange, CacheParams.FrameRate, CacheAsset, BlueprintClass, DataflowContent, BlueprintTransform, bSkeletalMeshVisibility, GeneratorDeltaTime);
+	TaskManager->AllocateSimulationResource(TimeRange, CacheParams.FrameRate,
+		CacheAsset, BlueprintClass,
+		DataflowContent, BlueprintTransform,
+		BlueprintVariables, bSkeletalMeshVisibility, GeneratorDeltaTime);
 
 	if(CacheParams.bAsyncCaching)
 	{
@@ -290,6 +306,11 @@ void FDataflowSimulationGenerator::SetBlueprintTransform(const FTransform& InBlu
 	BlueprintTransform = InBlueprintTransform;
 }
 
+void FDataflowSimulationGenerator::SetBlueprintVariables(const FInstancedPropertyBag& InBlueprintVariables)
+{
+	BlueprintVariables = InBlueprintVariables;
+}
+
 void FDataflowSimulationGenerator::SetDataflowContent(const TObjectPtr<UDataflowBaseContent>& InDataflowContent)
 {
 	DataflowContent = InDataflowContent;
diff --git a/Engine/Plugins/Experimental/Dataflow/Source/DataflowEditor/Private/Dataflow/DataflowSimulationScene.cpp b/Engine/Plugins/Experimental/Dataflow/Source/DataflowEditor/Private/Dataflow/DataflowSimulationScene.cpp
--- a/Engine/Plugins/Experimental/Dataflow/Source/DataflowEditor/Private/Dataflow/DataflowSimulationScene.cpp
+++ b/Engine/Plugins/Experimental/Dataflow/Source/DataflowEditor/Private/Dataflow/DataflowSimulationScene.cpp
@@ -5,6 +5,7 @@
 #include "Dataflow/DataflowSimulationControls.h"
 #include "Dataflow/DataflowSimulationVisualization.h"
 #include "Dataflow/DataflowEditor.h"
+#include "Dataflow/DataflowInstance.h"
 #include "Components/PrimitiveComponent.h"
 #include "Chaos/CacheManagerActor.h"
 #include "Misc/TransactionObjectEvent.h"
@@ -22,6 +23,7 @@
 #include "Components/BillboardComponent.h"
 #include "Components/SkeletalMeshComponent.h"
 #include "Engine/SkeletalMesh.h"
+#include "Engine/BlueprintGeneratedClass.h"
 #include "GeometryCache.h"
 #include "Dataflow/DataflowElement.h"
 
@@ -80,15 +82,15 @@ FDataflowSimulationScene::FDataflowSimulationScene(FPreviewScene::ConstructionVa
 			SceneDescription->CacheParams = DataflowAsset->PreviewCacheParams;
 			SceneDescription->CacheAsset = DataflowAsset->PreviewCacheAsset ? Cast<UChaosCacheCollection>(DataflowAsset->PreviewCacheAsset.LoadSynchronous()) :
 					LoadObject<UChaosCacheCollection>(nullptr, TEXT("/Dataflow/CC_DefaultChaosCache.CC_DefaultChaosCache"));
-			SceneDescription->BlueprintClass = DataflowAsset->PreviewBlueprintClass;
+			SceneDescription->SetBlueprintClass(DataflowAsset->PreviewBlueprintClass);
 			SceneDescription->BlueprintTransform = DataflowAsset->PreviewBlueprintTransform; 
 			SceneDescription->GeometryCacheAsset = Cast<UGeometryCache>(DataflowAsset->PreviewGeometryCacheAsset.LoadSynchronous());
 			SceneDescription->EmbeddedSkeletalMesh = Cast<USkeletalMesh>(DataflowAsset->PreviewEmbeddedSkeletalMesh.LoadSynchronous());
 			SceneDescription->EmbeddedStaticMesh = Cast<UStaticMesh>(DataflowAsset->PreviewEmbeddedStaticMesh.LoadSynchronous());
 		}
-		if(SceneDescription->BlueprintClass == nullptr)
+		if(SceneDescription->GetBlueprintClass() == nullptr)
 		{
-			SceneDescription->BlueprintClass = GetEditorContent()->GetPreviewClass();
+			SceneDescription->SetBlueprintClass(GetEditorContent()->GetPreviewClass());
 		}
 #endif
 	}
@@ -288,16 +290,17 @@ void FDataflowSimulationScene::CreateSimulationScene()
 	{
 		SimulationGenerator->SetCacheParams(SceneDescription->CacheParams);
 		SimulationGenerator->SetCacheAsset(SceneDescription->CacheAsset);
-		SimulationGenerator->SetBlueprintClass(SceneDescription->BlueprintClass);
+		SimulationGenerator->SetBlueprintClass(SceneDescription->GetBlueprintClass());
 		SimulationGenerator->SetBlueprintTransform(SceneDescription->BlueprintTransform);
+		SimulationGenerator->SetBlueprintVariables(SceneDescription->BlueprintVariables);
 		SimulationGenerator->SetDataflowContent(GetEditorContent());
 
 		TimeRange = SceneDescription->CacheParams.TimeRange;
 		NumFrames = (TimeRange[1] > TimeRange[0]) ? FMath::Floor((TimeRange[1] - TimeRange[0] + UE_SMALL_NUMBER) * SceneDescription->CacheParams.FrameRate) + 1 : 1;
 		
 		DeltaTime = (SceneDescription->CacheParams.FrameRate > 0) ? 1.f / SceneDescription->CacheParams.FrameRate: 0.f;
-		PreviewActor = UE::Dataflow::SpawnSimulatedActor(SceneDescription->BlueprintClass, Cast<AChaosCacheManager>(RootSceneActor),
-			SceneDescription->CacheAsset, false, GetEditorContent(), SceneDescription->BlueprintTransform);
+		PreviewActor = UE::Dataflow::SpawnSimulatedActor(SceneDescription->GetBlueprintClass(), Cast<AChaosCacheManager>(RootSceneActor),
+			SceneDescription->CacheAsset, false, GetEditorContent(), SceneDescription->BlueprintTransform, SceneDescription->BlueprintVariables);
 		SimulationGenerator->SetDeltaTime(DeltaTime);
 		// Setup all the skelmesh animations
 		UE::Dataflow::SetupSkeletonAnimation(PreviewActor, SceneDescription->bSkeletalMeshVisibility);
@@ -520,11 +523,14 @@ void FDataflowSimulationScene::SceneDescriptionPropertyChanged(const FName& Prop
 			SimulationGenerator->SetCacheAsset(SceneDescription->CacheAsset);
 		}
 	}
+PRAGMA_DISABLE_DEPRECATION_WARNINGS
 	else if(PropertyName == GET_MEMBER_NAME_CHECKED(UDataflowSimulationSceneDescription, BlueprintClass))
+PRAGMA_ENABLE_DEPRECATION_WARNINGS
 	{
 		if(SimulationGenerator)
 		{
-			SimulationGenerator->SetBlueprintClass(SceneDescription->BlueprintClass);
+			SimulationGenerator->SetBlueprintClass(SceneDescription->GetBlueprintClass());
+			SimulationGenerator->SetBlueprintVariables(SceneDescription->BlueprintVariables);
 		}
 	}
 	else if(PropertyName == GET_MEMBER_NAME_CHECKED(UDataflowSimulationSceneDescription, BlueprintTransform))
@@ -534,6 +540,14 @@ void FDataflowSimulationScene::SceneDescriptionPropertyChanged(const FName& Prop
 			SimulationGenerator->SetBlueprintTransform(SceneDescription->BlueprintTransform);
 		}
 	}
+	else if(PropertyName == GET_MEMBER_NAME_CHECKED(UDataflowSimulationSceneDescription, BlueprintVariables))
+	{
+		WritePreviewPropertiesToDataflowOwner();
+		if(SimulationGenerator)
+		{
+			SimulationGenerator->SetBlueprintVariables(SceneDescription->BlueprintVariables);
+		}
+	}
 	if(GetEditorContent())
 	{
 		if(UDataflow* DataflowAsset = GetEditorContent()->GetDataflowAsset())
@@ -541,7 +555,7 @@ void FDataflowSimulationScene::SceneDescriptionPropertyChanged(const FName& Prop
 #if WITH_EDITORONLY_DATA
 			DataflowAsset->PreviewCacheParams = SceneDescription->CacheParams;
 			DataflowAsset->PreviewCacheAsset = SceneDescription->CacheAsset;
-			DataflowAsset->PreviewBlueprintClass = SceneDescription->BlueprintClass;
+			DataflowAsset->PreviewBlueprintClass = SceneDescription->GetBlueprintClass();
 			DataflowAsset->PreviewBlueprintTransform = SceneDescription->BlueprintTransform;
 			DataflowAsset->PreviewGeometryCacheAsset = SceneDescription->GeometryCacheAsset;
 			DataflowAsset->PreviewEmbeddedSkeletalMesh = SceneDescription->EmbeddedSkeletalMesh;
@@ -558,6 +572,38 @@ void FDataflowSimulationScene::SceneDescriptionPropertyChanged(const FName& Prop
 	CreateSimulationScene();
 }
 
+void FDataflowSimulationScene::ReadPreviewPropertiesFromDataflowOwner()
+{
+	if (SceneDescription)
+	{
+		if (const UDataflowBaseContent* const EditorContent = GetEditorContent())
+		{
+			if (const IDataflowInstanceInterface* const DataflowInstanceInterface = Cast<IDataflowInstanceInterface>(EditorContent->GetDataflowOwner()))
+			{
+				const FDataflowInstance& DataflowInstance = DataflowInstanceInterface->GetDataflowInstance();
+				SceneDescription->BlueprintVariables.CopyMatchingValuesByName(DataflowInstance.GetPreviewBlueprintVariableOverrides());  // When reading, match the asset variables
+			}
+		}
+	}
+}
+
+void FDataflowSimulationScene::WritePreviewPropertiesToDataflowOwner() const
+{
+	if (SceneDescription)
+	{
+		if (const UDataflowBaseContent* const EditorContent = GetEditorContent())
+		{
+			UObject* const DataflowOwner = EditorContent->GetDataflowOwner();
+			if (IDataflowInstanceInterface* const DataflowInstanceInterface = Cast<IDataflowInstanceInterface>(DataflowOwner))
+			{
+				FDataflowInstance& DataflowInstance = DataflowInstanceInterface->GetDataflowInstance();
+				DataflowInstance.GetPreviewBlueprintVariableOverrides() = SceneDescription->BlueprintVariables;  // When writing, clobber the asset variables
+				DataflowOwner->MarkPackageDirty();
+			}
+		}
+	}
+}
+
 void UDataflowSimulationSceneDescription::GenerateGeometryCache()
 {
 	if (!bIsGeometryCacheOutputSupported)
@@ -710,10 +756,48 @@ void UDataflowSimulationSceneDescription::SetSimulationScene(FDataflowSimulation
 	SimulationScene = InSimulationScene;
 }
 
+void UDataflowSimulationSceneDescription::SetBlueprintClass(TSubclassOf<AActor> InBlueprintClass)
+{
+PRAGMA_DISABLE_DEPRECATION_WARNINGS
+	if (BlueprintClass != InBlueprintClass)
+	{
+		OnBlueprintClassPreChanged();
+		BlueprintClass = InBlueprintClass;
+		OnBlueprintClassPostChanged();
+	}
+PRAGMA_ENABLE_DEPRECATION_WARNINGS
+}
+
+const TSubclassOf<AActor>& UDataflowSimulationSceneDescription::GetBlueprintClass() const
+{
+PRAGMA_DISABLE_DEPRECATION_WARNINGS
+	return BlueprintClass;
+PRAGMA_ENABLE_DEPRECATION_WARNINGS
+}
+
+void UDataflowSimulationSceneDescription::PreEditChange(FProperty* PropertyAboutToChange)
+{
+	Super::PreEditChange(PropertyAboutToChange);
+
+PRAGMA_DISABLE_DEPRECATION_WARNINGS
+	if (PropertyAboutToChange && PropertyAboutToChange->GetFName() == GET_MEMBER_NAME_CHECKED(UDataflowSimulationSceneDescription, BlueprintClass))
+PRAGMA_ENABLE_DEPRECATION_WARNINGS
+	{
+		OnBlueprintClassPreChanged();
+	}
+}
+
 void UDataflowSimulationSceneDescription::PostEditChangeProperty(FPropertyChangedEvent& PropertyChangedEvent)
 {
 	Super::PostEditChangeProperty(PropertyChangedEvent);
 
+PRAGMA_DISABLE_DEPRECATION_WARNINGS
+	if (PropertyChangedEvent.GetMemberPropertyName() == GET_MEMBER_NAME_CHECKED(UDataflowSimulationSceneDescription, BlueprintClass))
+PRAGMA_ENABLE_DEPRECATION_WARNINGS
+	{
+		OnBlueprintClassPostChanged();
+	}
+
 	if (SimulationScene)
 	{
 		SimulationScene->SceneDescriptionPropertyChanged(PropertyChangedEvent.GetMemberPropertyName());
@@ -737,5 +821,75 @@ void UDataflowSimulationSceneDescription::PostTransacted(const FTransactionObjec
 	}
 }
 
+void UDataflowSimulationSceneDescription::BeginDestroy()
+{
+	Super::BeginDestroy();
+	UnregisterOnBlueprintCompiledDelegate();
+}
+
+void UDataflowSimulationSceneDescription::OnBlueprintClassPreChanged()
+{
+	UnregisterOnBlueprintCompiledDelegate();
+}
+
+void UDataflowSimulationSceneDescription::OnBlueprintClassPostChanged()
+{
+	const UBlueprintGeneratedClass* const BlueprintGeneratedClass = Cast<const UBlueprintGeneratedClass>(GetBlueprintClass());
+	UBlueprint* const Blueprint = BlueprintGeneratedClass ? Cast<UBlueprint>(BlueprintGeneratedClass->ClassGeneratedBy) : nullptr;
+	UpdateBlueprintVariables(Blueprint);  // Always call UpdateBlueprintVariables, even with a null Blueprint, so it correctly resets the variables
+	RegisterOnBlueprintCompiledDelegate();
+}
+
+void UDataflowSimulationSceneDescription::UpdateBlueprintVariables(UBlueprint* Blueprint)
+{
+	BlueprintVariables.Reset();
+	if (const UBlueprintGeneratedClass* const BlueprintGeneratedClass = Blueprint ? Cast<UBlueprintGeneratedClass>(Blueprint->GeneratedClass) : nullptr)
+	{
+		for (const FBPVariableDescription& BPVariableDescription : Blueprint->NewVariables)
+		{
+			if (BPVariableDescription.HasMetaData("BlueprintPrivate") &&
+				BPVariableDescription.GetMetaData("BlueprintPrivate") == TEXT("true"))
+			{
+				continue;
+			}
+
+			if (const FProperty* const Property = BlueprintGeneratedClass->FindPropertyByName(BPVariableDescription.VarName))
+			{
+				const FName VarGuidAsName(BPVariableDescription.VarGuid.ToString());  // VarGuid instead of the VarName, because the variable name can change and FInstancedPropertyBag::SanitizePropertyName
+				BlueprintVariables.AddProperty(VarGuidAsName, Property);
+			}
+		}
+	}
+	// Must re-sync properties when the blueprint has changed
+	if (SimulationScene)
+	{
+		SimulationScene->ReadPreviewPropertiesFromDataflowOwner();
+	}
+}
+
+void UDataflowSimulationSceneDescription::RegisterOnBlueprintCompiledDelegate()
+{
+	if (const UBlueprintGeneratedClass* const BlueprintGeneratedClass = Cast<const UBlueprintGeneratedClass>(GetBlueprintClass()))
+	{
+		if (UBlueprint* const Blueprint = Cast<UBlueprint>(BlueprintGeneratedClass->ClassGeneratedBy))
+		{
+			Blueprint->OnCompiled().RemoveAll(this);
+			Blueprint->OnCompiled().AddUObject(this, &UDataflowSimulationSceneDescription::UpdateBlueprintVariables);
+		}
+	}
+}
+
+void UDataflowSimulationSceneDescription::UnregisterOnBlueprintCompiledDelegate()
+{
+	if (const UBlueprintGeneratedClass* const BlueprintGeneratedClass = Cast<const UBlueprintGeneratedClass>(GetBlueprintClass()))
+	{
+		if (UBlueprint* const Blueprint = Cast<UBlueprint>(BlueprintGeneratedClass->ClassGeneratedBy))
+		{
+			Blueprint->OnCompiled().RemoveAll(this);
+		}
+	}
+}
+
+
 #undef LOCTEXT_NAMESPACE
 
diff --git a/Engine/Plugins/Experimental/Dataflow/Source/DataflowEditor/Public/Dataflow/DataflowSimulationControls.h b/Engine/Plugins/Experimental/Dataflow/Source/DataflowEditor/Public/Dataflow/DataflowSimulationControls.h
--- a/Engine/Plugins/Experimental/Dataflow/Source/DataflowEditor/Public/Dataflow/DataflowSimulationControls.h
+++ b/Engine/Plugins/Experimental/Dataflow/Source/DataflowEditor/Public/Dataflow/DataflowSimulationControls.h
@@ -9,6 +9,7 @@ class UDataflow;
 class UChaosCacheCollection;
 class AChaosCacheManager;
 class UDataflowBaseContent;
+struct FInstancedPropertyBag;
 
 namespace UE::Dataflow
 {
@@ -17,6 +18,12 @@ namespace UE::Dataflow
 		const TObjectPtr<UWorld>& SimulationWorld, UE::Dataflow::FTimestamp& LastTimeStamp);
 
 	/** Spawn an actor given a class type and attach it to the cache manager */
+	TObjectPtr<AActor> SpawnSimulatedActor(const TSubclassOf<AActor>& ActorClass,
+		const TObjectPtr<AChaosCacheManager>& CacheManager, const TObjectPtr<UChaosCacheCollection>& CacheCollection,
+		const bool bIsRecording, const TObjectPtr<UDataflowBaseContent>& DataflowContent, const FTransform& ActorTransform,
+		const FInstancedPropertyBag& BlueprintVariables);
+
+	UE_DEPRECATED(5.8, "Use BlueprintVariables version instead.")
 	TObjectPtr<AActor> SpawnSimulatedActor(const TSubclassOf<AActor>& ActorClass,
 		const TObjectPtr<AChaosCacheManager>& CacheManager, const TObjectPtr<UChaosCacheCollection>& CacheCollection,
 		const bool bIsRecording, const TObjectPtr<UDataflowBaseContent>& DataflowContent, const FTransform& ActorTransform);
diff --git a/Engine/Plugins/Experimental/Dataflow/Source/DataflowEditor/Public/Dataflow/DataflowSimulationGenerator.h b/Engine/Plugins/Experimental/Dataflow/Source/DataflowEditor/Public/Dataflow/DataflowSimulationGenerator.h
--- a/Engine/Plugins/Experimental/Dataflow/Source/DataflowEditor/Public/Dataflow/DataflowSimulationGenerator.h
+++ b/Engine/Plugins/Experimental/Dataflow/Source/DataflowEditor/Public/Dataflow/DataflowSimulationGenerator.h
@@ -9,6 +9,7 @@
 #include "Misc/AsyncTaskNotification.h"
 #include "Dataflow/DataflowContent.h"
 #include "Dataflow/DataflowPreview.h"
+#include "StructUtils/PropertyBag.h"
 
 DECLARE_LOG_CATEGORY_EXTERN(LogDataflowSimulationGenerator, Log, All);
 
@@ -114,15 +115,22 @@ namespace UE::Dataflow
 		FDateTime LastUpdateTime;
 
 		/** Allocate the simulation resource from the properties */
-        bool AllocateSimulationResource(const FVector2f& TimeRange, const int32 FrameRate,
-        	const TObjectPtr<UChaosCacheCollection>& CacheAsset, const TSubclassOf<AActor>& ActorClass,
-        	const TObjectPtr<UDataflowBaseContent>& DataflowContent, const FTransform& BlueprintTransform, const bool bSkeletalMeshVisibility, float DeltaTime);
+		bool AllocateSimulationResource(const FVector2f& TimeRange, const int32 FrameRate,
+			const TObjectPtr<UChaosCacheCollection>& CacheAsset, const TSubclassOf<AActor>& ActorClass,
+			const TObjectPtr<UDataflowBaseContent>& DataflowContent, const FTransform& BlueprintTransform,
+			const FInstancedPropertyBag& BlueprintVariables, const bool bSkeletalMeshVisibility, float DeltaTime);
 
-        /** Free the simulation resource */
-        void FreeSimulationResource();
+		UE_DEPRECATED(5.8, "Use BlueprintVariables version instead.")
+		bool AllocateSimulationResource(const FVector2f& TimeRange, const int32 FrameRate,
+			const TObjectPtr<UChaosCacheCollection>& CacheAsset, const TSubclassOf<AActor>& ActorClass,
+			const TObjectPtr<UDataflowBaseContent>& DataflowContent, const FTransform& BlueprintTransform,
+			const bool bSkeletalMeshVisibility, float DeltaTime);
 
-        /** Cancel the simulation generation */
-        void CancelSimulationGeneration();
+		/** Free the simulation resource */
+		void FreeSimulationResource();
+
+		/** Cancel the simulation generation */
+		void CancelSimulationGeneration();
 
 		/** Number of simulated frames */
 		std::atomic<int32> NumSimulatedFrames = 0;
@@ -170,6 +178,9 @@ namespace UE::Dataflow
 		/** Set the blueprint transform for cache recording */
 		void SetBlueprintTransform(const FTransform& InBlueprintTransform);
 		
+		/** Set the blueprint variables' values used to customize the blueprint */
+		void SetBlueprintVariables(const FInstancedPropertyBag& InBlueprintVariables);
+
 		/** Set the cache params for cache recording */
 		void SetCacheParams(const FDataflowPreviewCacheParams& InCacheParams);
 
@@ -214,6 +225,9 @@ namespace UE::Dataflow
 		/** Blueprint transform used to spawn the actor */
 		FTransform BlueprintTransform = FTransform::Identity;
 
+		/** Blueprint variables' values used to customize the blueprint */
+		FInstancedPropertyBag BlueprintVariables;
+
 		/** Dataflow content */
 		TObjectPtr<UDataflowBaseContent> DataflowContent;
 
@@ -221,7 +235,7 @@ namespace UE::Dataflow
 		bool bSkeletalMeshVisibility = true;
 
 		/** Delta time for a frame */
-		float GeneratorDeltaTime;
+		float GeneratorDeltaTime = 0.f;
 
 		/** Pending action to be send to the async thread */
 		EDataflowGeneratorActions PendingAction = EDataflowGeneratorActions::NoAction;
diff --git a/Engine/Plugins/Experimental/Dataflow/Source/DataflowEditor/Public/Dataflow/DataflowSimulationScene.h b/Engine/Plugins/Experimental/Dataflow/Source/DataflowEditor/Public/Dataflow/DataflowSimulationScene.h
--- a/Engine/Plugins/Experimental/Dataflow/Source/DataflowEditor/Public/Dataflow/DataflowSimulationScene.h
+++ b/Engine/Plugins/Experimental/Dataflow/Source/DataflowEditor/Public/Dataflow/DataflowSimulationScene.h
@@ -2,21 +2,25 @@
 
 #pragma once
 
+#include "Chaos/CacheCollection.h"
+#include "Dataflow/DataflowEditor.h"
 #include "Dataflow/DataflowEditorPreviewSceneBase.h"
 #include "Dataflow/DataflowSimulationGenerator.h"
-#include "Dataflow/DataflowEditor.h"
-#include "Chaos/CacheCollection.h"
+#include "StructUtils/PropertyBag.h"
 
 #include "DataflowSimulationScene.generated.h"
 
 #define UE_API DATAFLOWEDITOR_API
 
+class AActor;
+class FFleshCollection;
 class UDataflowEditor;
 class UDeformableTetrahedralComponent;
-class USkeletalMeshComponent;
 class UGeometryCache;
 class UFleshDynamicAsset;
-class FFleshCollection;
+class USkeletalMesh;
+class USkeletalMeshComponent;
+class UStaticMesh;
 
 DECLARE_EVENT(UDataflowSimulationSceneDescription, FDataflowSimulationSceneDescriptionChanged)
 
@@ -55,10 +59,20 @@ public:
 	/** Set the simulation scene */
 	UE_API void SetSimulationScene(class FDataflowSimulationScene* SimulationScene);
 
+	/** Set the Blueprint. */
+	UE_API void SetBlueprintClass(TSubclassOf<AActor> InBlueprintClass);
+
+	/** Return the Blueprint. */
+	UE_API const TSubclassOf<AActor>& GetBlueprintClass() const;
+
 	/** Caching blueprint actor class to spawn */
-	UPROPERTY(EditAnywhere, Category = "Scene")
+	UE_DEPRECATED(5.8, "Use the SetBlueprintClass setter/getter instead as to get correct operations.")
+	UPROPERTY(EditAnywhere, Setter = SetBlueprintClass, Getter = GetBlueprintClass, Category = "Scene")
 	TSubclassOf<AActor> BlueprintClass;
 
+	UPROPERTY(EditAnywhere, Category = "Scene|BlueprintVariables", Meta = (FixedLayout, ShowOnlyInnerProperties))
+	FInstancedPropertyBag BlueprintVariables;
+
 	/** Blueprint actor transform */
 	UPROPERTY(EditAnywhere, Category = "Scene")
 	FTransform BlueprintTransform = FTransform::Identity;
@@ -106,9 +120,17 @@ public:
 private:
 
 	//~ UObject Interface
-	UE_API virtual void PostEditChangeProperty(struct FPropertyChangedEvent& PropertyChangedEvent) override;
-	UE_API virtual void PostTransacted(const FTransactionObjectEvent& TransactionEvent) override;
-	
+	virtual void PreEditChange(class FProperty* PropertyAboutToChange) override;
+	virtual void PostEditChangeProperty(struct FPropertyChangedEvent& PropertyChangedEvent) override;
+	virtual void PostTransacted(const FTransactionObjectEvent& TransactionEvent) override;
+	virtual void BeginDestroy() override;
+
+	void OnBlueprintClassPreChanged();
+	void OnBlueprintClassPostChanged();
+	void UpdateBlueprintVariables(class UBlueprint* Blueprint);
+	void RegisterOnBlueprintCompiledDelegate();
+	void UnregisterOnBlueprintCompiledDelegate();
+
 	/** Simulation scene linked to that descriptor */
 	class FDataflowSimulationScene* SimulationScene;
 
@@ -223,13 +245,20 @@ public:
 	void SetSimulationLocked(const bool bSimulationLocked, const bool bIsPlaying);
 
 private:
+	friend UDataflowSimulationSceneDescription;
 
 	/** Bind the scene selection to the components */
 	UE_API void BindSceneSelection();
 
 	/** Unbind the scene selection from the components */
 	UE_API void UnbindSceneSelection();
-	
+
+	/** Sync the blueprint preview properties with the edited asset. */
+	void ReadPreviewPropertiesFromDataflowOwner();
+
+	/** Sync the blueprint preview properties with the edited asset. */
+	void WritePreviewPropertiesToDataflowOwner() const;
+
 	/** Simulation scene description */
 	TObjectPtr<UDataflowSimulationSceneDescription> SceneDescription;
 
diff --git a/Engine/Source/Runtime/Experimental/Dataflow/Engine/Public/Dataflow/DataflowInstance.h b/Engine/Source/Runtime/Experimental/Dataflow/Engine/Public/Dataflow/DataflowInstance.h
--- a/Engine/Source/Runtime/Experimental/Dataflow/Engine/Public/Dataflow/DataflowInstance.h
+++ b/Engine/Source/Runtime/Experimental/Dataflow/Engine/Public/Dataflow/DataflowInstance.h
@@ -240,6 +240,19 @@ public:
 
 	DATAFLOWENGINE_API void SyncVariables();
 
+#if WITH_EDITORONLY_DATA
+	/** Get the preview blueprint variable values for this asset. */
+	FInstancedPropertyBag& GetPreviewBlueprintVariableOverrides()
+	{
+		return PreviewBlueprintVariableOverrides;
+	}
+	/** Get the preview blueprint variable values for this asset. Const version. */
+	const FInstancedPropertyBag& GetPreviewBlueprintVariableOverrides() const
+	{
+		return PreviewBlueprintVariableOverrides;
+	}
+#endif
+
 	/** 
 	* Update the asset by re-evaluating the dataflow
 	* return true if the dataflow has been properly evaluated and the asset updated
@@ -259,10 +272,16 @@ private:
 	UPROPERTY(EditAnywhere, Category = "Dataflow")
 	FName DataflowTerminal;
 
-	/** Variables to override */
+	/** Dataflow variables to override */
 	UPROPERTY(EditAnywhere, Category = "Dataflow", meta=(ShowOnlyInnerProperties))
 	FDataflowVariableOverrides VariableOverrides;
 
+#if WITH_EDITORONLY_DATA
+	/** Preview blueprint variables to override */
+	UPROPERTY()
+	FInstancedPropertyBag PreviewBlueprintVariableOverrides;
+#endif
+
 private:
 	friend class UE::Dataflow::Private::FVariablesOverridesDetails;
 
