Dataflow:
- Added rendering support for Box(FBox), BoxArray(TArray<FBox>), Sphere(FSphere), SphereArray(TArray<FSphere>), PointArray(TArray<FVector>) to the new rendering system

#Synced-CL 48239755

diff --git a/Engine/Plugins/Experimental/Dataflow/Source/DataflowEditor/Private/Dataflow/DataflowEditorModule.cpp b/Engine/Plugins/Experimental/Dataflow/Source/DataflowEditor/Private/Dataflow/DataflowEditorModule.cpp
--- a/Engine/Plugins/Experimental/Dataflow/Source/DataflowEditor/Private/Dataflow/DataflowEditorModule.cpp
+++ b/Engine/Plugins/Experimental/Dataflow/Source/DataflowEditor/Private/Dataflow/DataflowEditorModule.cpp
@@ -41,6 +41,9 @@
 #include "DataflowRendering/DataflowDynamicMeshRenderableType.h"
 #include "DataflowRendering/DataflowGeometryCollectionRenderableType.h"
 #include "DataflowRendering/DataflowGeometryCollectionProximityRenderableType.h"
+#include "DataflowRendering/DataflowPointRenderableType.h"
+#include "DataflowRendering/DataflowBoxRenderableType.h"
+#include "DataflowRendering/DataflowSphereRenderableType.h"
 
 #include "ContentBrowserModule.h"
 #include "Modules/ModuleManager.h"
@@ -274,6 +277,9 @@ void FDataflowEditorModule::StartupModule()
 	UE::Dataflow::Private::RegisterDynamicMeshRenderableTypes();
 	UE::Dataflow::Private::RegisterGeometryCollectionRenderableTypes();
 	UE::Dataflow::Private::RegisterGeometryCollectionProximityRenderableTypes();
+	UE::Dataflow::Private::RegisterPointRenderableTypes();
+	UE::Dataflow::Private::RegisterBoxRenderableTypes();
+	UE::Dataflow::Private::RegisterSphereRenderableTypes();
 
 	// register an browser extension to show when an asset is bound to a dataflow asset
 	if (FContentBrowserModule* ContentBrowserModule = FModuleManager::Get().GetModulePtr<FContentBrowserModule>(TEXT("ContentBrowser")))
@@ -283,7 +289,6 @@ void FDataflowEditorModule::StartupModule()
 			FOnGenerateAssetViewExtraStateIndicators::CreateStatic(&UE::Dataflow::Private::OnGenerateDataflowAssetBindingTooltip)
 		));
 	}
-
 }
 
 void FDataflowEditorModule::ShutdownModule()
diff --git a/Engine/Plugins/Experimental/Dataflow/Source/DataflowEditor/Private/DataflowRendering/DataflowBoxRenderableType.cpp b/Engine/Plugins/Experimental/Dataflow/Source/DataflowEditor/Private/DataflowRendering/DataflowBoxRenderableType.cpp
new file mode 100644
--- /dev/null
+++ b/Engine/Plugins/Experimental/Dataflow/Source/DataflowEditor/Private/DataflowRendering/DataflowBoxRenderableType.cpp
@@ -0,0 +1,97 @@
+// Copyright Epic Games, Inc. All Rights Reserved.
+
+#include "DataflowRendering/DataflowBoxRenderableType.h"
+
+#include "Components/BoxComponent.h"
+#include "Components/InstancedStaticMeshComponent.h"
+
+#include "Dataflow/DataflowRenderingViewMode.h"
+
+#include "DataflowRendering/DataflowRenderableComponents.h"
+#include "DataflowRendering/DataflowRenderableType.h"
+#include "DataflowRendering/DataflowRenderableTypeInstance.h"
+#include "DataflowRendering/DataflowRenderableTypeRegistry.h"
+#include "DataflowRendering/DataflowRenderableTypeUtils.h"
+
+#include "UObject/ObjectPtr.h"
+
+namespace UE::Dataflow::Private
+{
+	///////////////////////////////////////////////////////////////////////////////////////////////////////////////
+
+	class FBoxSurfaceRenderableType : public IRenderableType
+	{
+		UE_DATAFLOW_IRENDERABLE_PRIMARY_TYPE(FBox, Box);
+		UE_DATAFLOW_IRENDERABLE_RENDER_GROUP(Surface);
+		UE_DATAFLOW_IRENDERABLE_VIEW_MODE(FDataflowConstruction3DViewMode);
+
+		virtual void GetPrimitiveComponents(const FRenderableTypeInstance& Instance, FRenderableComponents& OutComponents) const override
+		{
+			const FBox Box = GetBox(Instance, {});
+
+			static const FName BoxComponentName = TEXT("Box");
+
+			if (UBoxComponent* Component = OutComponents.AddNewComponent<UBoxComponent>(BoxComponentName))
+			{
+				Component->SetWorldLocation(Box.GetCenter());
+				Component->SetBoxExtent(Box.GetExtent());
+			}
+		}
+	};
+
+	///////////////////////////////////////////////////////////////////////////////////////////////////////////////
+
+	class FBoxArraySurfaceRenderableType : public IRenderableType
+	{
+		UE_DATAFLOW_IRENDERABLE_PRIMARY_TYPE(TArray<FBox>, BoxArray);
+		UE_DATAFLOW_IRENDERABLE_RENDER_GROUP(Surface);
+		UE_DATAFLOW_IRENDERABLE_VIEW_MODE(FDataflowConstruction3DViewMode);
+
+		virtual void GetPrimitiveComponents(const FRenderableTypeInstance& Instance, FRenderableComponents& OutComponents) const override
+		{
+			UStaticMesh* BoxMesh = LoadObject<UStaticMesh>(nullptr, TEXT("/Engine/BasicShapes/Cube.Cube"));
+			if (BoxMesh)
+			{
+				const TArray<FBox>& BoxArray = GetBoxArray(Instance, {});
+
+				static const FName BoxArrayComponentName = TEXT("BoxArray");
+
+				const int32 NumBoxes = BoxArray.Num();
+
+				UInstancedStaticMeshComponent* ISMComponent = OutComponents.AddNewComponent<UInstancedStaticMeshComponent>(BoxArrayComponentName);
+				ISMComponent->PreAllocateInstancesMemory(NumBoxes);
+				ISMComponent->SetStaticMesh(BoxMesh);
+
+				//for (int32 Idx = 0; Idx < NumBoxes; ++Idx)
+				//{
+				//	const FBox& Box = BoxArray[Idx];
+
+				//	FTransform BoxTransform;
+				//	BoxTransform.SetTranslation(Box.GetCenter());
+				//	BoxTransform.SetScale3D(Box.GetExtent() / 50.f);
+
+				//	ISMComponent->instance
+				//}
+				for (const FBox& Box : BoxArray)
+				{
+					FTransform BoxTransform;
+					BoxTransform.SetTranslation(Box.GetCenter());
+					BoxTransform.SetScale3D(Box.GetExtent() / 50.f);
+
+					ISMComponent->AddInstance(BoxTransform);
+				}
+			}
+		}
+	};
+
+	///////////////////////////////////////////////////////////////////////////////////////////////////////////////
+
+	void RegisterBoxRenderableTypes()
+	{
+		UE_DATAFLOW_REGISTER_RENDERABLE_TYPE(FBoxSurfaceRenderableType);
+		UE_DATAFLOW_REGISTER_RENDERABLE_TYPE(FBoxArraySurfaceRenderableType);
+	}
+
+	///////////////////////////////////////////////////////////////////////////////////////////////////////////////
+
+}
\ No newline at end of file
diff --git a/Engine/Plugins/Experimental/Dataflow/Source/DataflowEditor/Private/DataflowRendering/DataflowBoxRenderableType.h b/Engine/Plugins/Experimental/Dataflow/Source/DataflowEditor/Private/DataflowRendering/DataflowBoxRenderableType.h
new file mode 100644
--- /dev/null
+++ b/Engine/Plugins/Experimental/Dataflow/Source/DataflowEditor/Private/DataflowRendering/DataflowBoxRenderableType.h
@@ -0,0 +1,8 @@
+// Copyright Epic Games, Inc. All Rights Reserved.
+
+#pragma once
+
+namespace UE::Dataflow::Private
+{
+	void RegisterBoxRenderableTypes();
+}
diff --git a/Engine/Plugins/Experimental/Dataflow/Source/DataflowEditor/Private/DataflowRendering/DataflowPointRenderableType.cpp b/Engine/Plugins/Experimental/Dataflow/Source/DataflowEditor/Private/DataflowRendering/DataflowPointRenderableType.cpp
new file mode 100644
--- /dev/null
+++ b/Engine/Plugins/Experimental/Dataflow/Source/DataflowEditor/Private/DataflowRendering/DataflowPointRenderableType.cpp
@@ -0,0 +1,68 @@
+// Copyright Epic Games, Inc. All Rights Reserved.
+
+#include "DataflowRendering/DataflowPointRenderableType.h"
+
+#include "Drawing/PointSetComponent.h"
+#include "Components/InstancedStaticMeshComponent.h"
+
+#include "Dataflow/DataflowRenderingViewMode.h"
+
+#include "DataflowRendering/DataflowRenderableComponents.h"
+#include "DataflowRendering/DataflowRenderableType.h"
+#include "DataflowRendering/DataflowRenderableTypeInstance.h"
+#include "DataflowRendering/DataflowRenderableTypeRegistry.h"
+#include "DataflowRendering/DataflowRenderableTypeUtils.h"
+
+#include "UObject/ObjectPtr.h"
+
+namespace UE::Dataflow::Private
+{
+	///////////////////////////////////////////////////////////////////////////////////////////////////////////////
+
+	class FPointArraySurfaceRenderableType : public IRenderableType
+	{
+		UE_DATAFLOW_IRENDERABLE_PRIMARY_TYPE(TArray<FVector>, PointArray);
+		UE_DATAFLOW_IRENDERABLE_RENDER_GROUP(Surface);
+		UE_DATAFLOW_IRENDERABLE_VIEW_MODE(FDataflowConstruction3DViewMode);
+
+	public:
+		FPointArraySurfaceRenderableType()
+		{
+			constexpr bool bDepthTested = false;
+			Material = (bDepthTested)
+				? LoadObject<UMaterial>(nullptr, TEXT("/MeshModelingToolsetExp/Materials/PointSetComponentMaterial"))
+				: LoadObject<UMaterial>(nullptr, TEXT("/MeshModelingToolsetExp/Materials/PointSetOverlaidComponentMaterial"));
+		}
+
+		virtual void GetPrimitiveComponents(const FRenderableTypeInstance& Instance, FRenderableComponents& OutComponents) const override
+		{
+			const TArray<FVector>& PointArray = GetPointArray(Instance, {});
+
+			static const FName PointComponentName = TEXT("PointArray");
+			static const FColor PointColor = FColor::Blue;
+			static const float PointSize = 15.f;
+
+			const int32 NumPoints = PointArray.Num();
+
+			if (UPointSetComponent* PointComponent = OutComponents.AddNewComponent<UPointSetComponent>(PointComponentName))
+			{
+				PointComponent->ReservePoints(NumPoints);
+				PointComponent->AddPoints(PointArray, PointColor, PointSize);
+
+				PointComponent->SetPointMaterial(Material);
+			}
+		}
+
+		UMaterialInterface* Material = nullptr;
+	};
+
+	///////////////////////////////////////////////////////////////////////////////////////////////////////////////
+
+	void RegisterPointRenderableTypes()
+	{
+		UE_DATAFLOW_REGISTER_RENDERABLE_TYPE(FPointArraySurfaceRenderableType);
+	}
+
+	///////////////////////////////////////////////////////////////////////////////////////////////////////////////
+
+}
\ No newline at end of file
diff --git a/Engine/Plugins/Experimental/Dataflow/Source/DataflowEditor/Private/DataflowRendering/DataflowPointRenderableType.h b/Engine/Plugins/Experimental/Dataflow/Source/DataflowEditor/Private/DataflowRendering/DataflowPointRenderableType.h
new file mode 100644
--- /dev/null
+++ b/Engine/Plugins/Experimental/Dataflow/Source/DataflowEditor/Private/DataflowRendering/DataflowPointRenderableType.h
@@ -0,0 +1,8 @@
+// Copyright Epic Games, Inc. All Rights Reserved.
+
+#pragma once
+
+namespace UE::Dataflow::Private
+{
+	void RegisterPointRenderableTypes();
+}
diff --git a/Engine/Plugins/Experimental/Dataflow/Source/DataflowEditor/Private/DataflowRendering/DataflowSphereRenderableType.cpp b/Engine/Plugins/Experimental/Dataflow/Source/DataflowEditor/Private/DataflowRendering/DataflowSphereRenderableType.cpp
new file mode 100644
--- /dev/null
+++ b/Engine/Plugins/Experimental/Dataflow/Source/DataflowEditor/Private/DataflowRendering/DataflowSphereRenderableType.cpp
@@ -0,0 +1,87 @@
+// Copyright Epic Games, Inc. All Rights Reserved.
+
+#include "DataflowRendering/DataflowSphereRenderableType.h"
+
+#include "Components/SphereComponent.h"
+#include "Components/InstancedStaticMeshComponent.h"
+
+#include "Dataflow/DataflowRenderingViewMode.h"
+
+#include "DataflowRendering/DataflowRenderableComponents.h"
+#include "DataflowRendering/DataflowRenderableType.h"
+#include "DataflowRendering/DataflowRenderableTypeInstance.h"
+#include "DataflowRendering/DataflowRenderableTypeRegistry.h"
+#include "DataflowRendering/DataflowRenderableTypeUtils.h"
+
+#include "UObject/ObjectPtr.h"
+
+namespace UE::Dataflow::Private
+{
+	///////////////////////////////////////////////////////////////////////////////////////////////////////////////
+
+	class FSphereSurfaceRenderableType : public IRenderableType
+	{
+		UE_DATAFLOW_IRENDERABLE_PRIMARY_TYPE(FSphere, Sphere);
+		UE_DATAFLOW_IRENDERABLE_RENDER_GROUP(Surface);
+		UE_DATAFLOW_IRENDERABLE_VIEW_MODE(FDataflowConstruction3DViewMode);
+
+		virtual void GetPrimitiveComponents(const FRenderableTypeInstance& Instance, FRenderableComponents& OutComponents) const override
+		{
+			const FSphere Sphere = GetSphere(Instance, {});
+
+			static const FName SphereComponentName = TEXT("Sphere");
+
+			if (USphereComponent* Component = OutComponents.AddNewComponent<USphereComponent>(SphereComponentName))
+			{
+				Component->SetWorldLocation(Sphere.Center);
+				Component->SetSphereRadius(Sphere.W);
+			}
+		}
+	};
+
+	///////////////////////////////////////////////////////////////////////////////////////////////////////////////
+
+	class FSphereArraySurfaceRenderableType : public IRenderableType
+	{
+		UE_DATAFLOW_IRENDERABLE_PRIMARY_TYPE(TArray<FSphere>, SphereArray);
+		UE_DATAFLOW_IRENDERABLE_RENDER_GROUP(Surface);
+		UE_DATAFLOW_IRENDERABLE_VIEW_MODE(FDataflowConstruction3DViewMode);
+
+		virtual void GetPrimitiveComponents(const FRenderableTypeInstance& Instance, FRenderableComponents& OutComponents) const override
+		{
+			UStaticMesh* SphereMesh = LoadObject<UStaticMesh>(nullptr, TEXT("/Engine/BasicShapes/Sphere.Sphere"));
+			if (SphereMesh)
+			{
+				const TArray<FSphere>& SphereArray = GetSphereArray(Instance, {});
+
+				static const FName SphereArrayComponentName = TEXT("SphereArray");
+
+				const int32 NumSpheres = SphereArray.Num();
+
+				UInstancedStaticMeshComponent* ISMComponent = OutComponents.AddNewComponent<UInstancedStaticMeshComponent>(SphereArrayComponentName);
+				ISMComponent->PreAllocateInstancesMemory(NumSpheres);
+				ISMComponent->SetStaticMesh(SphereMesh);
+
+				for (const FSphere& Sphere : SphereArray)
+				{
+					FTransform SphereTransform;
+					SphereTransform.SetTranslation(Sphere.Center);
+					SphereTransform.SetScale3D(FVector(Sphere.W / 50.f));
+
+					ISMComponent->AddInstance(SphereTransform);
+				}
+			}
+		}
+	};
+
+	///////////////////////////////////////////////////////////////////////////////////////////////////////////////
+
+	void RegisterSphereRenderableTypes()
+	{
+		UE_DATAFLOW_REGISTER_RENDERABLE_TYPE(FSphereSurfaceRenderableType);
+		UE_DATAFLOW_REGISTER_RENDERABLE_TYPE(FSphereArraySurfaceRenderableType);
+	}
+
+	///////////////////////////////////////////////////////////////////////////////////////////////////////////////
+
+}
\ No newline at end of file
diff --git a/Engine/Plugins/Experimental/Dataflow/Source/DataflowEditor/Private/DataflowRendering/DataflowSphereRenderableType.h b/Engine/Plugins/Experimental/Dataflow/Source/DataflowEditor/Private/DataflowRendering/DataflowSphereRenderableType.h
new file mode 100644
--- /dev/null
+++ b/Engine/Plugins/Experimental/Dataflow/Source/DataflowEditor/Private/DataflowRendering/DataflowSphereRenderableType.h
@@ -0,0 +1,8 @@
+// Copyright Epic Games, Inc. All Rights Reserved.
+
+#pragma once
+
+namespace UE::Dataflow::Private
+{
+	void RegisterSphereRenderableTypes();
+}
