[Backout] - CL48246314
#rnx
Original CL Desc
-----------------------------------------------------------------
Dataflow:
- Added rendering support for Box(FBox), BoxArray(TArray<FBox>), Sphere(FSphere), SphereArray(TArray<FSphere>), PointArray(TArray<FVector>) to the new rendering system

#Synced-CL 48239755

diff --git a/Engine/Plugins/Experimental/Dataflow/Source/DataflowEditor/Private/Dataflow/DataflowEditorModule.cpp b/Engine/Plugins/Experimental/Dataflow/Source/DataflowEditor/Private/Dataflow/DataflowEditorModule.cpp
--- a/Engine/Plugins/Experimental/Dataflow/Source/DataflowEditor/Private/Dataflow/DataflowEditorModule.cpp
+++ b/Engine/Plugins/Experimental/Dataflow/Source/DataflowEditor/Private/Dataflow/DataflowEditorModule.cpp
@@ -41,9 +41,6 @@
 #include "DataflowRendering/DataflowDynamicMeshRenderableType.h"
 #include "DataflowRendering/DataflowGeometryCollectionRenderableType.h"
 #include "DataflowRendering/DataflowGeometryCollectionProximityRenderableType.h"
-#include "DataflowRendering/DataflowPointRenderableType.h"
-#include "DataflowRendering/DataflowBoxRenderableType.h"
-#include "DataflowRendering/DataflowSphereRenderableType.h"
 
 #include "ContentBrowserModule.h"
 #include "Modules/ModuleManager.h"
@@ -277,9 +274,6 @@ void FDataflowEditorModule::StartupModule()
 	UE::Dataflow::Private::RegisterDynamicMeshRenderableTypes();
 	UE::Dataflow::Private::RegisterGeometryCollectionRenderableTypes();
 	UE::Dataflow::Private::RegisterGeometryCollectionProximityRenderableTypes();
-	UE::Dataflow::Private::RegisterPointRenderableTypes();
-	UE::Dataflow::Private::RegisterBoxRenderableTypes();
-	UE::Dataflow::Private::RegisterSphereRenderableTypes();
 
 	// register an browser extension to show when an asset is bound to a dataflow asset
 	if (FContentBrowserModule* ContentBrowserModule = FModuleManager::Get().GetModulePtr<FContentBrowserModule>(TEXT("ContentBrowser")))
@@ -289,6 +283,7 @@ void FDataflowEditorModule::StartupModule()
 			FOnGenerateAssetViewExtraStateIndicators::CreateStatic(&UE::Dataflow::Private::OnGenerateDataflowAssetBindingTooltip)
 		));
 	}
+
 }
 
 void FDataflowEditorModule::ShutdownModule()
diff --git a/Engine/Plugins/Experimental/Dataflow/Source/DataflowEditor/Private/DataflowRendering/DataflowBoxRenderableType.cpp b/Engine/Plugins/Experimental/Dataflow/Source/DataflowEditor/Private/DataflowRendering/DataflowBoxRenderableType.cpp
deleted file mode 100644
diff --git a/Engine/Plugins/Experimental/Dataflow/Source/DataflowEditor/Private/DataflowRendering/DataflowBoxRenderableType.h b/Engine/Plugins/Experimental/Dataflow/Source/DataflowEditor/Private/DataflowRendering/DataflowBoxRenderableType.h
deleted file mode 100644
diff --git a/Engine/Plugins/Experimental/Dataflow/Source/DataflowEditor/Private/DataflowRendering/DataflowPointRenderableType.cpp b/Engine/Plugins/Experimental/Dataflow/Source/DataflowEditor/Private/DataflowRendering/DataflowPointRenderableType.cpp
deleted file mode 100644
diff --git a/Engine/Plugins/Experimental/Dataflow/Source/DataflowEditor/Private/DataflowRendering/DataflowPointRenderableType.h b/Engine/Plugins/Experimental/Dataflow/Source/DataflowEditor/Private/DataflowRendering/DataflowPointRenderableType.h
deleted file mode 100644
diff --git a/Engine/Plugins/Experimental/Dataflow/Source/DataflowEditor/Private/DataflowRendering/DataflowSphereRenderableType.cpp b/Engine/Plugins/Experimental/Dataflow/Source/DataflowEditor/Private/DataflowRendering/DataflowSphereRenderableType.cpp
deleted file mode 100644
diff --git a/Engine/Plugins/Experimental/Dataflow/Source/DataflowEditor/Private/DataflowRendering/DataflowSphereRenderableType.h b/Engine/Plugins/Experimental/Dataflow/Source/DataflowEditor/Private/DataflowRendering/DataflowSphereRenderableType.h
deleted file mode 100644
