Dataflow : weight map paint tool
- Fixed radius issue when switching between brush modes
- Fixed tooltips for brush mode ( removed references to Strength )
- properly display the world radius when scaling using B+Mouse drag

#Synced-CL 48301667

diff --git a/Engine/Plugins/Experimental/Dataflow/Source/DataflowEditor/Private/DataflowEditorTools/DataflowEditorVertexAttributePaintTool.cpp b/Engine/Plugins/Experimental/Dataflow/Source/DataflowEditor/Private/DataflowEditorTools/DataflowEditorVertexAttributePaintTool.cpp
--- a/Engine/Plugins/Experimental/Dataflow/Source/DataflowEditor/Private/DataflowEditorTools/DataflowEditorVertexAttributePaintTool.cpp
+++ b/Engine/Plugins/Experimental/Dataflow/Source/DataflowEditor/Private/DataflowEditorTools/DataflowEditorVertexAttributePaintTool.cpp
@@ -589,10 +589,23 @@ UBaseDynamicMeshComponent* UDataflowEditorVertexAttributePaintTool::GetSculptMes
 
 void UDataflowEditorVertexAttributePaintTool::SetupBrushEditBehaviorSetup(ULocalTwoAxisPropertyEditInputBehavior& OutBehavior)
 {
-	// map the horizontal behavior to change the brush size
-	UMeshSculptToolBase::MapHorizontalBrushEditBehaviorToBrushSize(OutBehavior);
+	// Map the horizontal behavior to change the brush size 
+	// (we do not want to use the default as we are internal set to use the adapative brush size but external display world radius)
+	//UMeshSculptToolBase::MapHorizontalBrushEditBehaviorToBrushSize(OutBehavior);
 
-	// map the vertical ehavior to change the attribiute value
+	OutBehavior.HorizontalProperty.GetValueFunc = [this]() { return GetCurrentBrushRadius(); };
+	OutBehavior.HorizontalProperty.SetValueFunc = [this](float NewValue) { SetBrushRadius(NewValue); };
+	OutBehavior.HorizontalProperty.MutateDeltaFunc = [this](float Delta)
+		{
+			// Scale delta if brush size is in world units.
+			//return Delta * (BrushProperties->BrushSize.SizeType == EBrushToolSizeType::World ? (CameraState.Position - LastBrushFrameWorld.Origin).Length() : 1.f);
+			return Delta * (CameraState.Position - LastBrushFrameWorld.Origin).Length();
+		};
+	OutBehavior.HorizontalProperty.Name = LOCTEXT("BrushRadius", "Radius");
+	OutBehavior.HorizontalProperty.bEnabled = true;
+
+
+	// map the vertical behavior to change the attribiute value
 	OutBehavior.VerticalProperty.GetValueFunc = [this]()
 		{
 			return ToolProperties->BrushProperties.AttributeValue;
@@ -697,9 +710,6 @@ void UDataflowEditorVertexAttributePaintTool::Setup()
 	// initialize properties
 	ToolProperties = NewObject<UDataflowEditorVertexAttributePaintToolProperties>(this);
 	{
-		ToolProperties->WatchProperty(ToolProperties->BrushProperties.BrushSize,
-			[this](float NewSize) { UMeshSculptToolBase::BrushProperties->BrushSize.AdaptiveSize = NewSize; });
-		//ToolProperties->BrushProperties.BrushSize = UMeshSculptToolBase::BrushProperties->BrushSize.AdaptiveSize;
 		ToolProperties->DisplayProperties.ColorRamp.ColorRamp.OnColorCurveChangedDelegate.AddUObject(this, &UDataflowEditorVertexAttributePaintTool::OnColorRampChanged);
 		//ToolProperties->RestoreProperties(this);
 	}
@@ -717,18 +727,6 @@ void UDataflowEditorVertexAttributePaintTool::Setup()
 	UMeshSculptToolBase::BrushProperties->bShowFalloff = true;
 	UMeshSculptToolBase::BrushProperties->bShowLazyness = false;
 	UMeshSculptToolBase::BrushProperties->FlowRate = 0.0f;
-	UMeshSculptToolBase::BrushProperties->GetOnModified().AddLambda(
-		[this](UObject* InBrushProperties, FProperty* Property)
-		{
-			// make sure we update the settings back from updating the brush using B+Mouse drag
-			if (InBrushProperties == UMeshSculptToolBase::BrushProperties)
-			{
-				if (Property->GetFName() == GET_MEMBER_NAME_CHECKED(FDataflowEditorVertexAttributePaintToolBrushProperties, BrushSize))
-				{
-					ToolProperties->BrushProperties.BrushSize = UMeshSculptToolBase::BrushProperties->BrushSize.AdaptiveSize;
-				}
-			}
-		});
 	CalculateBrushRadius();
 
 	PaintBrushOpProperties = NewObject<UDataflowVertexAttributePaintBrushOpProps>(this);
@@ -787,6 +785,9 @@ void UDataflowEditorVertexAttributePaintTool::Setup()
 
 	ToolProperties->LoadConfig();
 	SetBrushMode(ToolProperties->BrushProperties.BrushMode);
+
+	//SetPrimaryFalloffType(EMeshSculptFalloffType::Smooth);
+
 }
 
 void UDataflowEditorVertexAttributePaintTool::SetDataflowEditorContextObject(TObjectPtr<UDataflowContextObject> InDataflowEditorContextObject)
@@ -870,72 +871,95 @@ void UDataflowEditorVertexAttributePaintTool::OnPropertyModified(UObject* Proper
 	CalculateBrushRadius();
 }
 
+float UDataflowEditorVertexAttributePaintTool::GetBrushMinRadius() const
+{
+	return BrushProperties? BrushProperties->BrushSize.WorldSizeRange.Min * 0.5f: 0.001f;
+}
 
-void UDataflowEditorVertexAttributePaintTool::IncreaseBrushRadiusAction()
+float UDataflowEditorVertexAttributePaintTool::GetBrushMaxRadius() const
 {
-	Super::IncreaseBrushRadiusAction();
-	ToolProperties->BrushProperties.BrushSize = BrushProperties->BrushSize.AdaptiveSize;
-	NotifyOfPropertyChangeByTool(ToolProperties);
+	return BrushProperties ? BrushProperties->BrushSize.WorldSizeRange.Max * 0.5f: 20.0f;
 }
 
-void UDataflowEditorVertexAttributePaintTool::DecreaseBrushRadiusAction()
+float UDataflowEditorVertexAttributePaintTool::GetBrushAdaptiveSize() const
 {
-	Super::DecreaseBrushRadiusAction();
-	ToolProperties->BrushProperties.BrushSize = BrushProperties->BrushSize.AdaptiveSize;
-	NotifyOfPropertyChangeByTool(ToolProperties);
+	return BrushProperties ? BrushProperties->BrushSize.AdaptiveSize : 1.0f;
 }
 
-void UDataflowEditorVertexAttributePaintTool::IncreaseBrushRadiusSmallStepAction()
+void UDataflowEditorVertexAttributePaintTool::SetBrushAdaptiveSize(float AdaptiveSize)
 {
-	Super::IncreaseBrushRadiusSmallStepAction();
-	ToolProperties->BrushProperties.BrushSize = BrushProperties->BrushSize.AdaptiveSize;
-	NotifyOfPropertyChangeByTool(ToolProperties);
+	if (BrushProperties)
+	{
+		if (UMeshSculptToolBase::BrushProperties->BrushSize.SizeType == EBrushToolSizeType::Adaptive)
+		{
+			BrushProperties->BrushSize.AdaptiveSize = FMath::Clamp(AdaptiveSize, 0.0f, 1.0f);
+		}
+		else
+		{
+			BrushProperties->BrushSize.WorldRadius = FMath::Clamp(AdaptiveSize * GetBrushMaxRadius(), BrushProperties->BrushSize.WorldSizeRange.Min, BrushProperties->BrushSize.WorldSizeRange.Max);
+		}
+		NotifyOfPropertyChangeByTool(BrushProperties);
+		CalculateBrushRadius();
+	}
 }
 
-void UDataflowEditorVertexAttributePaintTool::DecreaseBrushRadiusSmallStepAction()
+void UDataflowEditorVertexAttributePaintTool::SetBrushRadius(float NewRadius)
 {
-	Super::DecreaseBrushRadiusSmallStepAction();
-	ToolProperties->BrushProperties.BrushSize = BrushProperties->BrushSize.AdaptiveSize;
-	NotifyOfPropertyChangeByTool(ToolProperties);
+	if (BrushProperties)
+	{
+		if (UMeshSculptToolBase::BrushProperties->BrushSize.SizeType == EBrushToolSizeType::Adaptive)
+		{
+			BrushProperties->BrushSize.AdaptiveSize = FMath::Clamp(NewRadius / FMath::Max(0.0001f, GetBrushMaxRadius()), 0.0f, 1.0f);
+		}
+		else
+		{
+			BrushProperties->BrushSize.WorldRadius = FMath::Clamp(NewRadius, BrushProperties->BrushSize.WorldSizeRange.Min, BrushProperties->BrushSize.WorldSizeRange.Max);
+		}
+		NotifyOfPropertyChangeByTool(BrushProperties);
+		CalculateBrushRadius();
+	}
 }
 
-float UDataflowEditorVertexAttributePaintTool::GetBrushMinRadius() const
+float UDataflowEditorVertexAttributePaintTool::GetBrushFalloff() const
 {
-	return UMeshSculptToolBase::BrushProperties->BrushSize.WorldSizeRange.Min * 0.5f;
+	return BrushProperties ? BrushProperties->BrushFalloffAmount : 1.0f;
 }
 
-float UDataflowEditorVertexAttributePaintTool::GetBrushMaxRadius() const
+void UDataflowEditorVertexAttributePaintTool::SetBrushFalloff(float NewFalloff)
 {
-	return UMeshSculptToolBase::BrushProperties->BrushSize.WorldSizeRange.Max * 0.5f;
+	if (BrushProperties)
+	{
+		BrushProperties->BrushFalloffAmount = FMath::Clamp(NewFalloff, 0.0f, 1.0f);
+		NotifyOfPropertyChangeByTool(BrushProperties);
+	}
 }
 
 void UDataflowEditorVertexAttributePaintTool::SetBrushMode(EDataflowEditorToolEditOperation NewBrushMode)
 {
-	if (ToolProperties)
+	if (ToolProperties && BrushProperties)
 	{
 		const EDataflowEditorToolEditOperation CurrentBrushMode = ToolProperties->BrushProperties.BrushMode;
 		ToolProperties->BrushProperties.BrushMode = NewBrushMode;
 
 		// Save the current brush config 
 		FDataflowEditorVertexAttributePaintToolBrushConfig BrushConfigToSave;
-		BrushConfigToSave.BrushSize = ToolProperties->BrushProperties.BrushSize;
+		BrushConfigToSave.BrushSize = GetBrushAdaptiveSize();
 		BrushConfigToSave.Value = ToolProperties->BrushProperties.AttributeValue;
 		ToolProperties->SetBrushConfig(CurrentBrushMode, BrushConfigToSave);
 
 		// Load the new brush config 
 		const FDataflowEditorVertexAttributePaintToolBrushConfig& NewBrushConfig = ToolProperties->GetBrushConfig(NewBrushMode);
-		ToolProperties->BrushProperties.BrushSize = NewBrushConfig.BrushSize;
+		SetBrushAdaptiveSize(NewBrushConfig.BrushSize);
 		ToolProperties->BrushProperties.AttributeValue = NewBrushConfig.Value;
 
 		// Send notification for the changed properties
 		FPropertyChangedEvent BrushModePropertyChangedEvent(UBrushBaseProperties::StaticClass()->FindPropertyByName(GET_MEMBER_NAME_CHECKED(FDataflowEditorVertexAttributePaintToolBrushProperties, BrushMode)));
 		ToolProperties->PostEditChangeProperty(BrushModePropertyChangedEvent);
-		FPropertyChangedEvent BrushSizePropertyChangedEvent(UBrushBaseProperties::StaticClass()->FindPropertyByName(GET_MEMBER_NAME_CHECKED(FDataflowEditorVertexAttributePaintToolBrushProperties, BrushSize)));
-		ToolProperties->PostEditChangeProperty(BrushSizePropertyChangedEvent);
 		FPropertyChangedEvent ValuePropertyChangedEvent(UBrushBaseProperties::StaticClass()->FindPropertyByName(GET_MEMBER_NAME_CHECKED(FDataflowEditorVertexAttributePaintToolBrushProperties, AttributeValue)));
 		ToolProperties->PostEditChangeProperty(ValuePropertyChangedEvent);
 
 		ToolProperties->SaveConfig();
+		BrushProperties->SaveConfig();
 	}
 }
 
diff --git a/Engine/Plugins/Experimental/Dataflow/Source/DataflowEditor/Private/DataflowEditorTools/DataflowEditorVertexAttributePaintToolDetailCustomization.cpp b/Engine/Plugins/Experimental/Dataflow/Source/DataflowEditor/Private/DataflowEditorTools/DataflowEditorVertexAttributePaintToolDetailCustomization.cpp
--- a/Engine/Plugins/Experimental/Dataflow/Source/DataflowEditor/Private/DataflowEditorTools/DataflowEditorVertexAttributePaintToolDetailCustomization.cpp
+++ b/Engine/Plugins/Experimental/Dataflow/Source/DataflowEditor/Private/DataflowEditorTools/DataflowEditorVertexAttributePaintToolDetailCustomization.cpp
@@ -245,7 +245,6 @@ namespace UE::Dataflow::Editor
 		// Hide all customized properties 
 		HideProperty(DetailBuilder, GET_MEMBER_NAME_CHECKED(UDataflowEditorVertexAttributePaintToolProperties, EditingMode));
 		HideProperty(DetailBuilder, GET_MEMBER_NAME_CHECKED(FDataflowEditorVertexAttributePaintToolBrushProperties, BrushMode));
-		HideProperty(DetailBuilder, GET_MEMBER_NAME_CHECKED(FDataflowEditorVertexAttributePaintToolBrushProperties, BrushSize));
 		HideProperty(DetailBuilder, GET_MEMBER_NAME_CHECKED(FDataflowEditorVertexAttributePaintToolBrushProperties, BrushAreaMode));
 		HideProperty(DetailBuilder, GET_MEMBER_NAME_CHECKED(FDataflowEditorVertexAttributePaintToolBrushProperties, ValueAtBrush));
 		HideProperty(DetailBuilder, GET_MEMBER_NAME_CHECKED(FDataflowEditorVertexAttributePaintToolBrushProperties, AttributeValue));
@@ -268,10 +267,10 @@ namespace UE::Dataflow::Editor
 				SNew(SSegmentedControl<EDataflowEditorToolEditOperation>)
 				.ToolTipText(LOCTEXT("BrushTooltip",
 					"Select the operation to apply when using the brush.\n"
-					"Add: applies the current value plus the Strength value to the new value.\n"
-					"Replace: applies the current value minus the Strength value to the new value.\n"
-					"Multiply: applies the current value multiplied by the Strength value to the new value.\n"
-					"Relax: applies the average of the connected (by edge) vertex value to the new vertex value, blended by the Strength."))
+					"Add: Applies the current value plus the value to the new value.\n"
+					"Replace: applies the current value minus the value to the new value.\n"
+					"Multiply: applies the current value multiplied by the value to the new value.\n"
+					"Relax: applies the average of the connected (by edge) vertex value to the new vertex value."))
 				.Value_Lambda([this]()
 					{
 						return ToolProperties.IsValid() ? ToolProperties->BrushProperties.BrushMode : EDataflowEditorToolEditOperation::Add;
@@ -349,30 +348,15 @@ namespace UE::Dataflow::Editor
 		.ValueContent()
 		[
 			SNew(SSpinBox<float>)
-			.MinValue_Lambda([this]() 
-				{ 
-					return Tool.IsValid() ? (Tool->GetBrushMinRadius()) : MinBrushRadius;
-				})
-			.MaxSliderValue_Lambda([this]()
-				{
-					return Tool.IsValid() ? (Tool->GetBrushMaxRadius()) : MaxBrushRadius;
-				})
+			.MinValue_Lambda([this]() { return Tool.IsValid() ? (Tool->GetBrushMinRadius()) : MinBrushRadius; })
+			.MaxSliderValue_Lambda([this]() { return Tool.IsValid() ? (Tool->GetBrushMaxRadius()) : MaxBrushRadius; })
 			.SupportDynamicSliderMaxValue(true)
-			.Value_Lambda([this]()
-				{
-					// BrushSize is adaptive ( [0,1] range ) , we need to scale it to the max value 
-					const float MaxRadius = Tool.IsValid() ? (Tool->GetBrushMaxRadius()) : MaxBrushRadius;
-					return (ToolProperties.IsValid()) ? (ToolProperties->BrushProperties.BrushSize * MaxRadius) : MaxRadius;
-				})
+			.Value_Lambda([this]() { return (Tool.IsValid()) ? Tool->GetCurrentBrushRadius() : MaxBrushRadius; })
 			.OnValueChanged_Lambda([this](float NewValue)
 				{
-					if (ToolProperties.IsValid())
+					if (Tool.IsValid())
 					{
-						// BrushSize is adaptive ( [0,1] range ), we need to divide by max value to set it 
-						const float MaxRadius = FMath::Max((Tool.IsValid() ? (Tool->GetBrushMaxRadius()) : MaxBrushRadius), MinBrushRadius);
-						ToolProperties->BrushProperties.BrushSize = (NewValue / MaxRadius);
-						FPropertyChangedEvent PropertyChangedEvent(UBrushBaseProperties::StaticClass()->FindPropertyByName(GET_MEMBER_NAME_CHECKED(FDataflowEditorVertexAttributePaintToolBrushProperties, BrushSize)));
-						ToolProperties->PostEditChangeProperty(PropertyChangedEvent);
+						Tool->SetBrushRadius(NewValue);
 					}
 				})
 			.OnValueCommitted_Lambda([this](float NewValue, ETextCommit::Type CommitType)
@@ -404,7 +388,7 @@ namespace UE::Dataflow::Editor
 			.MaxValue(1.0f)
 			.MaxSliderValue(1.f)
 			.Value(1.0f)
-			.SupportDynamicSliderMaxValue(true)
+			.SupportDynamicSliderMaxValue(false)
 			.Value_Lambda([this]()
 				{
 					return ToolProperties.IsValid() ? ToolProperties->BrushProperties.AttributeValue : 1.f;
@@ -430,6 +414,44 @@ namespace UE::Dataflow::Editor
 
 	///////////////////////////////////////////////////////////////////////////////////////////////////////////////
 
+	void FVertexAttributePaintToolDetailCustomization::AddBrushFalloffRow(IDetailCategoryBuilder& BrushCategory) const
+	{
+		BrushCategory.AddCustomRow(LOCTEXT("BrushFalloffCategory", "Brush Falloff"), false)
+			.NameContent()
+			[
+				SNew(STextBlock)
+					.Text(LOCTEXT("BrushFalloffLabel", "Falloff"))
+					.Font(FAppStyle::GetFontStyle(TEXT("PropertyWindow.NormalFont")))
+					.ToolTipText(LOCTEXT("BrushFalloffTooltip", "Amount of falloff in the [0,1] range"))
+			]
+			.ValueContent()
+			[
+				SNew(SSpinBox<float>)
+					.MinValue(0)
+					.MaxValue(1.0f)
+					.MaxSliderValue(1.f)
+					.Value(1.0f)
+					.SupportDynamicSliderMaxValue(false)
+					.Value_Lambda([this]() { return (Tool.IsValid()) ? Tool->GetBrushFalloff() : 1.0f; })
+					.OnValueChanged_Lambda([this](float NewValue)
+						{
+							if (Tool.IsValid())
+							{
+								Tool->SetBrushFalloff(NewValue);
+							}
+						})
+					.OnValueCommitted_Lambda([this](float NewValue, ETextCommit::Type CommitType)
+						{
+							if (ToolProperties.IsValid())
+							{
+								// TODO(ccaillaud) ==>> ToolProperties->SaveConfig();
+							}
+						})
+			];
+	}
+
+	///////////////////////////////////////////////////////////////////////////////////////////////////////////////
+
 	void FVertexAttributePaintToolDetailCustomization::AddBrushUI(IDetailLayoutBuilder& DetailBuilder) const
 	{
 		IDetailCategoryBuilder& BrushCategory = DetailBuilder.EditCategory("Brush", FText::GetEmpty(), ECategoryPriority::Important);
@@ -445,6 +467,9 @@ namespace UE::Dataflow::Editor
 
 			// Brush value field
 			AddBrushValueRow(BrushCategory);
+
+			// Brush value field
+			//AddBrushFalloffRow(BrushCategory);
 		}
 	}
 
diff --git a/Engine/Plugins/Experimental/Dataflow/Source/DataflowEditor/Private/DataflowEditorTools/DataflowEditorVertexAttributePaintToolDetailCustomization.h b/Engine/Plugins/Experimental/Dataflow/Source/DataflowEditor/Private/DataflowEditorTools/DataflowEditorVertexAttributePaintToolDetailCustomization.h
--- a/Engine/Plugins/Experimental/Dataflow/Source/DataflowEditor/Private/DataflowEditorTools/DataflowEditorVertexAttributePaintToolDetailCustomization.h
+++ b/Engine/Plugins/Experimental/Dataflow/Source/DataflowEditor/Private/DataflowEditorTools/DataflowEditorVertexAttributePaintToolDetailCustomization.h
@@ -39,6 +39,7 @@ namespace UE::Dataflow::Editor
 		void AddBrushFalloffModeRow(IDetailCategoryBuilder& BrushCategory) const;
 		void AddBrushRadiusRow(IDetailCategoryBuilder& BrushCategory) const;
 		void AddBrushValueRow(IDetailCategoryBuilder& BrushCategory) const;
+		void AddBrushFalloffRow(IDetailCategoryBuilder& BrushCategory) const;
 
 		void AddSelectionUI(IDetailLayoutBuilder& DetailBuilder) const;
 		TSharedRef<SWidget> MakeSelectionElementsToolbar() const;
diff --git a/Engine/Plugins/Experimental/Dataflow/Source/DataflowEditor/Public/DataflowEditorTools/DataflowEditorVertexAttributePaintTool.h b/Engine/Plugins/Experimental/Dataflow/Source/DataflowEditor/Public/DataflowEditorTools/DataflowEditorVertexAttributePaintTool.h
--- a/Engine/Plugins/Experimental/Dataflow/Source/DataflowEditor/Public/DataflowEditorTools/DataflowEditorVertexAttributePaintTool.h
+++ b/Engine/Plugins/Experimental/Dataflow/Source/DataflowEditor/Public/DataflowEditorTools/DataflowEditorVertexAttributePaintTool.h
@@ -93,10 +93,6 @@ struct FDataflowEditorVertexAttributePaintToolBrushProperties
 	UPROPERTY()
 	EDataflowEditorToolEditOperation BrushMode = EDataflowEditorToolEditOperation::Replace;
 
-	/** Adaptive size of brush - Relative to the model size */
-	UPROPERTY()
-	float BrushSize = 0.25f;
-
 	/** Value to paint on the mesh (dependent on the BrushMode) */
 	UPROPERTY()
 	float AttributeValue = 1.0;
@@ -371,14 +367,15 @@ public:
 	UE_API void ApplyVisibilityFilter(const TArray<int32>& Vertices, TArray<int32>& VisibleVertices);
 	UE_API void ApplyVisibilityFilter(TSet<int32>& Triangles, TArray<int32>& ROIBuffer, TArray<int32>& OutputBuffer);
 
-	// we override these so we can update the separate BrushSize property added for this tool
-	UE_API virtual void IncreaseBrushRadiusAction();
-	UE_API virtual void DecreaseBrushRadiusAction();
-	UE_API virtual void IncreaseBrushRadiusSmallStepAction();
-	UE_API virtual void DecreaseBrushRadiusSmallStepAction();
-
+	using UMeshSculptToolBase::GetCurrentBrushRadius;
+	UE_API float GetBrushAdaptiveSize() const;
 	UE_API float GetBrushMinRadius() const;
 	UE_API float GetBrushMaxRadius() const;
+	UE_API void SetBrushRadius(float NewRadius);
+	UE_API void SetBrushAdaptiveSize(float AdaptiveSize);
+
+	UE_API float GetBrushFalloff() const;
+	UE_API void SetBrushFalloff(float NewFalloff);
 
 	UE_API void SetColorMode(EDataflowEditorToolColorMode NewColorMode);
 
