Add a cvar to force dataflow to always be editable ( will hide the lock/.unlock button when enabled )
- usage : p.Dataflow.Editor.ForceGraphEditable 1

#Synced-CL 48321711

diff --git a/Engine/Plugins/Experimental/Dataflow/Source/DataflowEditor/Private/Dataflow/DataflowEditorToolkit.cpp b/Engine/Plugins/Experimental/Dataflow/Source/DataflowEditor/Private/Dataflow/DataflowEditorToolkit.cpp
--- a/Engine/Plugins/Experimental/Dataflow/Source/DataflowEditor/Private/Dataflow/DataflowEditorToolkit.cpp
+++ b/Engine/Plugins/Experimental/Dataflow/Source/DataflowEditor/Private/Dataflow/DataflowEditorToolkit.cpp
@@ -87,6 +87,11 @@ bool bDataflowLogNodeEvaluation = false;
 FAutoConsoleVariableRef CVARDataflowLogNodeEvaluation(TEXT("p.Dataflow.Editor.LogNodeEvaluation"), bDataflowLogNodeEvaluation,
 	TEXT("When enabled, logs the begin and end evaluation of Dataflow nodes - if the graph is large this may cause evry large amounty of log messages to be displayed in the output window.[def:false]"));
 
+bool bDataflowForceGraphEditable = false;
+FAutoConsoleVariableRef CVARbDataflowForceGraphEditable(TEXT("p.Dataflow.Editor.ForceGraphEditable"), bDataflowForceGraphEditable,
+	TEXT("when enabled all graph will be editable and the Unlock/lock editing buttons will not be visible.[def:false]"));
+
+
 const FName FDataflowEditorToolkit::GraphCanvasTabId(TEXT("DataflowEditor_GraphCanvas"));
 const FName FDataflowEditorToolkit::SubGraphCanvasTabId(TEXT("DataflowEditor_SubGraphTab"));
 const FName FDataflowEditorToolkit::NodeDetailsTabId(TEXT("DataflowEditor_NodeDetails"));
@@ -1129,7 +1134,7 @@ void FDataflowEditorToolkit::AddDataflowActionWidget(FToolMenuSection& InSection
 		FExecuteAction::CreateSP(this, &FDataflowEditorToolkit::ToggleDataflowEditingLock),
 		FCanExecuteAction::CreateSP(this, &FDataflowEditorToolkit::IsExternalDataflow),
 		FIsActionChecked::CreateSP(this, &FDataflowEditorToolkit::IsEditable),
-		FIsActionButtonVisible::CreateSP(this, &FDataflowEditorToolkit::IsExternalDataflow)
+		FIsActionButtonVisible::CreateSP(this, &FDataflowEditorToolkit::IsEditingLockButtonVisible)
 	);
 
 	const FSlateIcon UnlockedIcon = FSlateIcon(FAppStyle::GetAppStyleSetName(), "Icons.Unlock");
@@ -1169,7 +1174,12 @@ bool FDataflowEditorToolkit::IsExternalDataflow() const
 
 bool FDataflowEditorToolkit::IsEditable() const
 {
-	return (IsEmbeddedDataflow() || bAllowExternalDataflowEditing);
+	return (IsEmbeddedDataflow() || bAllowExternalDataflowEditing || bDataflowForceGraphEditable);
+}
+
+bool FDataflowEditorToolkit::IsEditingLockButtonVisible() const
+{
+	return IsExternalDataflow() && !bDataflowForceGraphEditable;
 }
 
 void FDataflowEditorToolkit::ConvertToEmbeddedDataflow()
diff --git a/Engine/Plugins/Experimental/Dataflow/Source/DataflowEditor/Public/Dataflow/DataflowEditorToolkit.h b/Engine/Plugins/Experimental/Dataflow/Source/DataflowEditor/Public/Dataflow/DataflowEditorToolkit.h
--- a/Engine/Plugins/Experimental/Dataflow/Source/DataflowEditor/Public/Dataflow/DataflowEditorToolkit.h
+++ b/Engine/Plugins/Experimental/Dataflow/Source/DataflowEditor/Public/Dataflow/DataflowEditorToolkit.h
@@ -324,17 +324,19 @@ private:
 	/** return true if the Dataflow is not an asset and is an object embedded into the owner asset */
 	bool IsEmbeddedDataflow() const;
 
-	/** return true if the Dataflow is a seperate asset or no dataflow is assigned to the owner asset */
+	/** return true if the Dataflow is a separate asset or no dataflow is assigned to the owner asset */
 	bool IsExternalDataflow() const;
 
-	/** whether the dataflow is in an editable state */
+	/** whether the Lock/Unlock editing button is visible or not */
+	bool IsEditingLockButtonVisible() const;
+
 	bool IsEditable() const;
 
 	/** Duplicate the current external dataflow and embed it into the owner asset */
 	void ConvertToEmbeddedDataflow();
 
 	/** 
-	* Duplicate the embedded dataflow object and store it into a seperate asset  
+	* Duplicate the embedded dataflow object and store it into a separate asset  
 	* this methid is interactive and will prompt a dialog for the user to choose a path to save the duplicated Dataflow to 
 	*/
 	void ConvertToExternalDataflow();
