Dataflow : save config for camera settings

#Synced-CL 48328177

diff --git a/Engine/Plugins/Experimental/Dataflow/Source/DataflowEditor/Private/Dataflow/DataflowConstructionViewportClient.cpp b/Engine/Plugins/Experimental/Dataflow/Source/DataflowEditor/Private/Dataflow/DataflowConstructionViewportClient.cpp
--- a/Engine/Plugins/Experimental/Dataflow/Source/DataflowEditor/Private/Dataflow/DataflowConstructionViewportClient.cpp
+++ b/Engine/Plugins/Experimental/Dataflow/Source/DataflowEditor/Private/Dataflow/DataflowConstructionViewportClient.cpp
@@ -48,6 +48,8 @@ FDataflowConstructionViewportClient::FDataflowConstructionViewportClient(FEditor
 	if (UDataflowEditorOptions* const MutableOptions = GetMutableDefault<UDataflowEditorOptions>())
 	{
 		OnSettingsChangeHandle = MutableOptions->OnSettingChanged().AddLambda([this](UObject*, FPropertyChangedEvent& Event) { UpdateFromSettings(); });
+
+		SetCameraSpeedSettings(MutableOptions->ConstructionCameraSpeedSettings);
 	}
 
 	WidgetMode = UE::Widget::WM_None;
@@ -125,6 +127,7 @@ FDataflowConstructionViewportClient::~FDataflowConstructionViewportClient()
 	{
 		Options->ConstructionViewFOV = FOVAngle;
 		Options->bConstructionViewFixedExposure = ExposureSettings.bFixed;
+		Options->ConstructionCameraSpeedSettings = GetCameraSpeedSettings();
 		Options->SaveConfig();
 	}
 }
diff --git a/Engine/Plugins/Experimental/Dataflow/Source/DataflowEditor/Private/Dataflow/DataflowEditorOptions.h b/Engine/Plugins/Experimental/Dataflow/Source/DataflowEditor/Private/Dataflow/DataflowEditorOptions.h
--- a/Engine/Plugins/Experimental/Dataflow/Source/DataflowEditor/Private/Dataflow/DataflowEditorOptions.h
+++ b/Engine/Plugins/Experimental/Dataflow/Source/DataflowEditor/Private/Dataflow/DataflowEditorOptions.h
@@ -4,6 +4,7 @@
 
 #include "UObject/Object.h"
 #include "Engine/DeveloperSettings.h"
+#include "Settings/EditorViewportSettings.h"
 #include "DataflowEditorOptions.generated.h"
 
 UENUM(BlueprintType)
@@ -59,6 +60,13 @@ class UDataflowEditorOptions : public UDeveloperSettings
 	UPROPERTY(config)
 	FString SimulationProfileName;
 
+	/** The default camera speed settings for the viewport */
+	UPROPERTY(config)
+	FEditorViewportCameraSpeedSettings ConstructionCameraSpeedSettings;
+
+	UPROPERTY(config)
+	FEditorViewportCameraSpeedSettings SimulationCameraSpeedSettings;
+
 	/** Which mouse button controls camera panning in the Construction Viewport in 2D view mode */
 	UPROPERTY(Config, EditAnywhere, Category = UI)
 	EDataflowConstructionViewportMousePanButton ConstructionViewportMousePanButton;
diff --git a/Engine/Plugins/Experimental/Dataflow/Source/DataflowEditor/Private/Dataflow/DataflowSimulationViewportClient.cpp b/Engine/Plugins/Experimental/Dataflow/Source/DataflowEditor/Private/Dataflow/DataflowSimulationViewportClient.cpp
--- a/Engine/Plugins/Experimental/Dataflow/Source/DataflowEditor/Private/Dataflow/DataflowSimulationViewportClient.cpp
+++ b/Engine/Plugins/Experimental/Dataflow/Source/DataflowEditor/Private/Dataflow/DataflowSimulationViewportClient.cpp
@@ -40,6 +40,7 @@ FDataflowSimulationViewportClient::FDataflowSimulationViewportClient(FEditorMode
 		FOVAngle = Options->SimulationViewFOV;
 		ViewFOV = FOVAngle;
 		ExposureSettings.bFixed = Options->bSimulationViewFixedExposure;
+		SetCameraSpeedSettings(Options->SimulationCameraSpeedSettings);
 	}
 }
 
@@ -49,6 +50,7 @@ FDataflowSimulationViewportClient::~FDataflowSimulationViewportClient()
 	{
 		Options->SimulationViewFOV = FOVAngle;
 		Options->bSimulationViewFixedExposure = ExposureSettings.bFixed;
+		Options->SimulationCameraSpeedSettings = GetCameraSpeedSettings();
 		Options->SaveConfig();
 	}
 }
