Dataflow : when selecting a color ramp preset , make sure the view refreshes accordingly

#Synced-CL 48328177

diff --git a/Engine/Plugins/Experimental/Dataflow/Source/DataflowEditor/Private/DataflowEditorTools/DataflowEditorVertexAttributePaintToolDetailCustomization.cpp b/Engine/Plugins/Experimental/Dataflow/Source/DataflowEditor/Private/DataflowEditorTools/DataflowEditorVertexAttributePaintToolDetailCustomization.cpp
--- a/Engine/Plugins/Experimental/Dataflow/Source/DataflowEditor/Private/DataflowEditorTools/DataflowEditorVertexAttributePaintToolDetailCustomization.cpp
+++ b/Engine/Plugins/Experimental/Dataflow/Source/DataflowEditor/Private/DataflowEditorTools/DataflowEditorVertexAttributePaintToolDetailCustomization.cpp
@@ -217,7 +217,7 @@ namespace UE::Dataflow::Editor
 						{
 							if (ToolProperties->DisplayProperties.ColorRampPresets.IsValidIndex(PresetIndex))
 							{
-								ToolProperties->DisplayProperties.ColorRamp = ToolProperties->DisplayProperties.ColorRampPresets[PresetIndex];
+								ToolProperties->DisplayProperties.ColorRamp.SetFrom(ToolProperties->DisplayProperties.ColorRampPresets[PresetIndex]);
 								ToolProperties->SaveConfig();
 
 								FPropertyChangedEvent PropertyChangedEvent(FDataflowEditorVertexAttributePaintToolDisplayProperties::StaticStruct()->FindPropertyByName(GET_MEMBER_NAME_CHECKED(FDataflowEditorVertexAttributePaintToolDisplayProperties, ColorRamp)));
diff --git a/Engine/Source/Runtime/Experimental/Dataflow/Engine/Private/Dataflow/DataflowColorRamp.cpp b/Engine/Source/Runtime/Experimental/Dataflow/Engine/Private/Dataflow/DataflowColorRamp.cpp
--- a/Engine/Source/Runtime/Experimental/Dataflow/Engine/Private/Dataflow/DataflowColorRamp.cpp
+++ b/Engine/Source/Runtime/Experimental/Dataflow/Engine/Private/Dataflow/DataflowColorRamp.cpp
@@ -125,6 +125,16 @@ FLinearColor FDataflowColorCurveOwner::GetCurveColor(FRichCurveEditInfo CurveInf
 
 ////////////////////////////////////////////////////////////////////////////////
 
+void FDataflowLinearColorRamp::SetFrom(const FDataflowLinearColorRamp& Other)
+{
+	RichCurves[0] = Other.RichCurves[0];
+	RichCurves[1] = Other.RichCurves[1];
+	RichCurves[2] = Other.RichCurves[2];
+	RichCurves[3] = Other.RichCurves[3];
+	OnCurveChanged(GetCurves());
+}
+
+
 bool FDataflowLinearColorRamp::IsEmpty() const
 {
 	return (RichCurves[0].IsEmpty()
diff --git a/Engine/Source/Runtime/Experimental/Dataflow/Engine/Public/Dataflow/DataflowColorRamp.h b/Engine/Source/Runtime/Experimental/Dataflow/Engine/Public/Dataflow/DataflowColorRamp.h
--- a/Engine/Source/Runtime/Experimental/Dataflow/Engine/Public/Dataflow/DataflowColorRamp.h
+++ b/Engine/Source/Runtime/Experimental/Dataflow/Engine/Public/Dataflow/DataflowColorRamp.h
@@ -60,6 +60,7 @@ public:
 	UE_API void SetColorAtTime(float Time, const FLinearColor& Color, bool bOnlyRBG);
 	UE_API bool IsEmpty() const;
 	UE_API void Reset();
+	UE_API void SetFrom(const FDataflowLinearColorRamp& Other);
 
 	DECLARE_MULTICAST_DELEGATE_OneParam(FOnColorCurveChanged, TArray<FRichCurve*> /* ChangedCurves */);
 	FOnColorCurveChanged OnColorCurveChangedDelegate;
