Dataflow : change the color of the dataflow icon in thumbnail depending on the embedded state

#Synced-CL 48458398

diff --git a/Engine/Plugins/Experimental/Dataflow/Source/DataflowEditor/Private/Dataflow/DataflowEditorModule.cpp b/Engine/Plugins/Experimental/Dataflow/Source/DataflowEditor/Private/Dataflow/DataflowEditorModule.cpp
--- a/Engine/Plugins/Experimental/Dataflow/Source/DataflowEditor/Private/Dataflow/DataflowEditorModule.cpp
+++ b/Engine/Plugins/Experimental/Dataflow/Source/DataflowEditor/Private/Dataflow/DataflowEditorModule.cpp
@@ -61,6 +61,9 @@ namespace UE::Dataflow::Private
 	static const FName DataflowColorRampName = TEXT("DataflowColorRamp");
 	static const FName DataflowLinearColorRampName = TEXT("DataflowLinearColorRamp");
 
+	static const FName HasDataflowRegistryTag = TEXT("HasDataflow");
+	static const FName IsDataflowEmbeddedRegistryTag = TEXT("IsDataflowEmbedded");
+
 	class FDataflowEditorWeightMapPaintToolActionCommands : public TInteractiveToolCommands<FDataflowEditorWeightMapPaintToolActionCommands>
 	{
 	public:
@@ -156,9 +159,11 @@ namespace UE::Dataflow::Private
 
 	TSharedRef<SWidget> OnGenerateDataflowAssetBindingIcons(const FAssetData& AssetData)
 	{
-		static const FName HasDataflowTag = TEXT("HasDataflow");
 		bool bHasDataflowAsset = false;
-		AssetData.GetTagValue(HasDataflowTag, bHasDataflowAsset);
+		AssetData.GetTagValue(HasDataflowRegistryTag, bHasDataflowAsset);
+		bool bIsDataflowEmbedded = false;
+		AssetData.GetTagValue(IsDataflowEmbeddedRegistryTag, bIsDataflowEmbedded);
+
 		if (bHasDataflowAsset)
 		{
 			const FSlateBrush* DataflowAssetIconBrush = FSlateIcon(FAppStyle::GetAppStyleSetName(), "GraphEditor.EventGraph_16x").GetIcon();
@@ -180,6 +185,10 @@ namespace UE::Dataflow::Private
 						[
 							SNew(SImage)
 							.Image(DataflowAssetIconBrush)
+							.ColorAndOpacity(bIsDataflowEmbedded 
+								? FSlateColor(FLinearColor::White)
+								: FSlateColor(FLinearColor::Yellow)
+							)
 						]
 					]
 				];
@@ -189,13 +198,11 @@ namespace UE::Dataflow::Private
 
 	TSharedRef<SWidget> OnGenerateDataflowAssetBindingTooltip(const FAssetData& AssetData)
 	{
-		static const FName HasDataflowTag = TEXT("HasDataflow");
 		bool bHasDataflowAsset = false;
-		AssetData.GetTagValue(HasDataflowTag, bHasDataflowAsset);
+		AssetData.GetTagValue(HasDataflowRegistryTag, bHasDataflowAsset);
 
-		static const FName IsDataflowEmbeddedTag = TEXT("IsDataflowEmbedded");
 		bool bIsDataflowEmbedded = false;
-		AssetData.GetTagValue(IsDataflowEmbeddedTag, bIsDataflowEmbedded);
+		AssetData.GetTagValue(IsDataflowEmbeddedRegistryTag, bIsDataflowEmbedded);
 
 		if (bHasDataflowAsset)
 		{
