Dataflow editor : fix potential crash in weightmap tool, when editing the color ramp
- Make sure to unsubscribe from the color ramp changes on shutdown

#Synced-CL 48458398

diff --git a/Engine/Plugins/Experimental/Dataflow/Source/DataflowEditor/Private/DataflowEditorTools/DataflowEditorVertexAttributePaintTool.cpp b/Engine/Plugins/Experimental/Dataflow/Source/DataflowEditor/Private/DataflowEditorTools/DataflowEditorVertexAttributePaintTool.cpp
--- a/Engine/Plugins/Experimental/Dataflow/Source/DataflowEditor/Private/DataflowEditorTools/DataflowEditorVertexAttributePaintTool.cpp
+++ b/Engine/Plugins/Experimental/Dataflow/Source/DataflowEditor/Private/DataflowEditorTools/DataflowEditorVertexAttributePaintTool.cpp
@@ -817,6 +817,7 @@ void UDataflowEditorVertexAttributePaintTool::Shutdown(EToolShutdownType Shutdow
 
 	if (ToolProperties)
 	{
+		ToolProperties->DisplayProperties.ColorRamp.OnColorCurveChangedDelegate.RemoveAll(this);
 		ToolProperties->SaveProperties(this);
 	}
 
