Dataflow weight map tool : Fix the auto focus when in mesh brush mode

#Synced-CL 48458398

diff --git a/Engine/Plugins/Experimental/Dataflow/Source/DataflowEditor/Private/DataflowEditorTools/DataflowEditorVertexAttributePaintTool.cpp b/Engine/Plugins/Experimental/Dataflow/Source/DataflowEditor/Private/DataflowEditorTools/DataflowEditorVertexAttributePaintTool.cpp
--- a/Engine/Plugins/Experimental/Dataflow/Source/DataflowEditor/Private/DataflowEditorTools/DataflowEditorVertexAttributePaintTool.cpp
+++ b/Engine/Plugins/Experimental/Dataflow/Source/DataflowEditor/Private/DataflowEditorTools/DataflowEditorVertexAttributePaintTool.cpp
@@ -29,6 +29,7 @@
 #include "Framework/Notifications/NotificationManager.h"
 #include "HAL/PlatformApplicationMisc.h"
 #include "Logging/LogMacros.h"
+#include "Math/Box.h"
 #include "ModelingToolTargetUtil.h"
 #include "Polygon2.h"
 #include "PreviewMesh.h"
@@ -876,6 +877,32 @@ void UDataflowEditorVertexAttributePaintTool::OnPropertyModified(UObject* Proper
 	CalculateBrushRadius();
 }
 
+FBox UDataflowEditorVertexAttributePaintTool::GetWorldSpaceFocusBox()
+{
+	if (ToolProperties->EditingMode == EDataflowEditorToolEditMode::Mesh)
+	{
+		if (PreviewMesh && PreviewMesh->GetMesh())
+		{
+			TArray<int32> SelectedVertices = GetSelectedVertices();
+			if (!SelectedVertices.IsEmpty())
+			{
+				const FDynamicMesh3* Mesh = PreviewMesh->GetMesh();
+				const FTransform3d Transform(PreviewMesh->GetTransform());
+
+				FBox MeshBounds(EForceInit::ForceInit);
+				for (const int32 VertexID : SelectedVertices)
+				{
+					MeshBounds += Transform.TransformPosition(Mesh->GetVertex(VertexID));
+				}
+				return static_cast<FBox>(MeshBounds);
+			}
+			return FBox(PreviewMesh->GetMesh()->GetBounds());
+		}
+
+	}
+	return Super::GetWorldSpaceFocusBox();
+}
+
 float UDataflowEditorVertexAttributePaintTool::GetBrushMinRadius() const
 {
 	return BrushProperties? BrushProperties->BrushSize.WorldSizeRange.Min * 0.5f: 0.001f;
diff --git a/Engine/Plugins/Experimental/Dataflow/Source/DataflowEditor/Public/DataflowEditorTools/DataflowEditorVertexAttributePaintTool.h b/Engine/Plugins/Experimental/Dataflow/Source/DataflowEditor/Public/DataflowEditorTools/DataflowEditorVertexAttributePaintTool.h
--- a/Engine/Plugins/Experimental/Dataflow/Source/DataflowEditor/Public/DataflowEditorTools/DataflowEditorVertexAttributePaintTool.h
+++ b/Engine/Plugins/Experimental/Dataflow/Source/DataflowEditor/Public/DataflowEditorTools/DataflowEditorVertexAttributePaintTool.h
@@ -342,6 +342,8 @@ public:
 
 	UE_API virtual void OnPropertyModified(UObject* PropertySet, FProperty* Property) override;
 
+	UE_API virtual FBox GetWorldSpaceFocusBox() override;
+
 	UE_API void SetBrushMode(EDataflowEditorToolEditOperation BrushMode);
 	UE_API bool IsInBrushMode() const;
 	UE_API bool IsVolumetricBrush() const;
