Dataflow : group customization : adjust the font and padding to be more similar to other part of the detail views

#Synced-CL 48601912

diff --git a/Engine/Plugins/Experimental/Dataflow/Source/DataflowEditor/Private/Dataflow/PropertyListCustomization.cpp b/Engine/Plugins/Experimental/Dataflow/Source/DataflowEditor/Private/Dataflow/PropertyListCustomization.cpp
--- a/Engine/Plugins/Experimental/Dataflow/Source/DataflowEditor/Private/Dataflow/PropertyListCustomization.cpp
+++ b/Engine/Plugins/Experimental/Dataflow/Source/DataflowEditor/Private/Dataflow/PropertyListCustomization.cpp
@@ -7,6 +7,7 @@
 #include "Widgets/Input/SComboButton.h"
 #include "Widgets/Views/SListView.h"
 #include "Widgets/Views/STableRow.h"
+#include "DetailLayoutBuilder.h"
 #include "DetailWidgetRow.h"
 
 namespace UE::Dataflow
@@ -155,8 +156,11 @@ namespace UE::Dataflow
 		if (Item)
 		{
 			return SNew(STableRow<TSharedPtr<FString>>, OwnerTable)
+				.Padding(4)
 			[
-				SNew(STextBlock).Text(*Item)
+				SNew(STextBlock)
+				.Text(*Item)
+				.Font(IDetailLayoutBuilder::GetDetailFont())
 			];
 		}
 		else
