Dataflow : make geometry collection renerable type more robust

#Synced-CL 48601912

diff --git a/Engine/Plugins/Experimental/Dataflow/Source/DataflowEditor/Private/DataflowRendering/DataflowGeometryCollectionRenderableType.cpp b/Engine/Plugins/Experimental/Dataflow/Source/DataflowEditor/Private/DataflowRendering/DataflowGeometryCollectionRenderableType.cpp
--- a/Engine/Plugins/Experimental/Dataflow/Source/DataflowEditor/Private/DataflowRendering/DataflowGeometryCollectionRenderableType.cpp
+++ b/Engine/Plugins/Experimental/Dataflow/Source/DataflowEditor/Private/DataflowRendering/DataflowGeometryCollectionRenderableType.cpp
@@ -67,7 +67,7 @@ namespace UE::Dataflow::Private
 							DynamicMeshes[Index] = CachedDynamicMeshes[Index];
 						});
 				}
-				else
+				if (DynamicMeshes.Num() != TransformAttribute->Num())
 				{
 					// TODO(dataflow) : need a version of ConvertGeometryCollectionToDynamicMesh that only takes a managed array collection
 					if (TUniquePtr<FGeometryCollection> GeomCollection = TUniquePtr<FGeometryCollection>(Collection.NewCopy<FGeometryCollection>()))
