Dataflow : enable dithering of occluded selection by default
The option can be toggle in the Show menu > Developper > Make Composite Editor Primitives Opaque

#Synced-CL 48601912

diff --git a/Engine/Plugins/Experimental/Dataflow/Source/DataflowEditor/Private/Dataflow/DataflowEditorToolkit.cpp b/Engine/Plugins/Experimental/Dataflow/Source/DataflowEditor/Private/Dataflow/DataflowEditorToolkit.cpp
--- a/Engine/Plugins/Experimental/Dataflow/Source/DataflowEditor/Private/Dataflow/DataflowEditorToolkit.cpp
+++ b/Engine/Plugins/Experimental/Dataflow/Source/DataflowEditor/Private/Dataflow/DataflowEditorToolkit.cpp
@@ -750,8 +750,8 @@ void FDataflowEditorToolkit::PostInitAssetEditor()
 		Client->EngineShowFlags.SetAntiAliasing(true);
 		Client->EngineShowFlags.SetMotionBlur(false);
 
-		// Disable the dithering of occluded portions of gizmos.
-		Client->EngineShowFlags.SetOpaqueCompositeEditorPrimitives(true);
+		// Enable the dithering of occluded selection and portions of gizmos.
+		Client->EngineShowFlags.SetOpaqueCompositeEditorPrimitives(false);
 
 		// Disable hardware occlusion queries, which make it harder to use vertex shaders to pull materials
 		// toward camera for z ordering because non-translucent materials start occluding themselves (once
