Dataflow smooth tool : fix avergae computation of the value of neighbours vertices that was not including the actual original value of the central vertex

#Synced-CL 48601912

diff --git a/Engine/Plugins/Experimental/Dataflow/Source/DataflowEditor/Private/DataflowEditorTools/DataflowEditorWeightMapPaintBrushOps.cpp b/Engine/Plugins/Experimental/Dataflow/Source/DataflowEditor/Private/DataflowEditorTools/DataflowEditorWeightMapPaintBrushOps.cpp
--- a/Engine/Plugins/Experimental/Dataflow/Source/DataflowEditor/Private/DataflowEditorTools/DataflowEditorWeightMapPaintBrushOps.cpp
+++ b/Engine/Plugins/Experimental/Dataflow/Source/DataflowEditor/Private/DataflowEditorTools/DataflowEditorWeightMapPaintBrushOps.cpp
@@ -118,8 +118,10 @@ void FDataflowWeightMapSmoothBrushOp::ApplyStampByVerticesStatic(
 		(int32 BufferIndex)
 		{
 			const int32 VertexIndex = Vertices[BufferIndex];
-			double Accumulator = 0.0f;
-			int32 NumNeighbors = 0;
+
+			// average include the vertex we are processing to avoid loosing information
+			double Accumulator = VertexWeightValuesConst[BufferIndex];
+			int32 NumNeighbors = 1;
 
 			for (const int32 NeighborIndex : Mesh->VtxVerticesItr(VertexIndex))
 			{
@@ -131,15 +133,7 @@ void FDataflowWeightMapSmoothBrushOp::ApplyStampByVerticesStatic(
 				}
 			}
 
-			if (NumNeighbors > 0)
-			{
-				OneRingAverages[BufferIndex] = Accumulator / NumNeighbors;
-			}
-			else
-			{
-				// Don't change the vertex value if it has no neighbors
-				OneRingAverages[BufferIndex] = VertexWeightValuesConst[BufferIndex];
-			}
+			OneRingAverages[BufferIndex] = Accumulator / NumNeighbors;
 		}, 
 		EParallelForFlags::None
 	);
