Dataflow : Fbox and arrays now display information when watching outputs of that type

#Synced-CL 48601912

diff --git a/Engine/Plugins/Experimental/Dataflow/Source/DataflowEditor/Private/Dataflow/DataflowSNode.cpp b/Engine/Plugins/Experimental/Dataflow/Source/DataflowEditor/Private/Dataflow/DataflowSNode.cpp
--- a/Engine/Plugins/Experimental/Dataflow/Source/DataflowEditor/Private/Dataflow/DataflowSNode.cpp
+++ b/Engine/Plugins/Experimental/Dataflow/Source/DataflowEditor/Private/Dataflow/DataflowSNode.cpp
@@ -38,7 +38,21 @@ namespace UE::Dataflow::Private
 		{
 			OutString += TEXT("/!\\ ");
 		}
-		if (FDataflowSelectionTypePolicy::SupportsTypeStatic(Output.GetType()))
+		if (Output.GetType() == "FBox")
+		{
+			const FBox DefaultValue(EForceInit::ForceInit);
+			const FBox& Value = Output.ReadValue(Context, DefaultValue);
+			if (Value.IsValid)
+			{
+				OutString += FString::Printf(TEXT("%s : %s\n"), *Output.GetName().ToString(), *Value.ToCompactString());
+			}
+			else
+			{
+				OutString += FString::Printf(TEXT("%s : FBox not set\n"), *Output.GetName().ToString());
+			}
+		
+		}
+		else if (FDataflowSelectionTypePolicy::SupportsTypeStatic(Output.GetType()))
 		{
 			const FDataflowSelection DefaultValue;
 			const FDataflowSelection& Value = Output.ReadValue(Context, DefaultValue);
@@ -61,6 +75,11 @@ namespace UE::Dataflow::Private
 			const FName ObjectName = Value ? Value->GetFName() : NullObject;
 			OutString += FString::Printf(TEXT("%s : %s\n"), *Output.GetName().ToString(), *ObjectName.ToString());
 		}
+		else if (FDataflowArrayTypePolicy::SupportsTypeStatic(Output.GetType()))
+		{
+			UE::Dataflow::FContextValue ContextValue(Context, Output.CacheKey());
+			OutString += FString::Printf(TEXT("%s : Array Size : %d\n"), *Output.GetName().ToString(), ContextValue.Num());
+		}
 		else
 		{
 			OutString += FString::Printf(TEXT("%s : (Cannot watch this type)\n"), *Output.GetName().ToString());
