Dataflow : update the name of the Embed / make external commands

#Synced-CL 48601912

diff --git a/Engine/Plugins/Experimental/Dataflow/Source/DataflowEditor/Private/Dataflow/DataflowEditorCommands.cpp b/Engine/Plugins/Experimental/Dataflow/Source/DataflowEditor/Private/Dataflow/DataflowEditorCommands.cpp
--- a/Engine/Plugins/Experimental/Dataflow/Source/DataflowEditor/Private/Dataflow/DataflowEditorCommands.cpp
+++ b/Engine/Plugins/Experimental/Dataflow/Source/DataflowEditor/Private/Dataflow/DataflowEditorCommands.cpp
@@ -80,8 +80,8 @@ void FDataflowEditorCommandsImpl::RegisterCommands()
 	UI_COMMAND(ConvertToForEachSubGraph, "Convert to For Each Subgraph", "Convert the subgraph to be able to iterate through an array.", EUserInterfaceActionType::Button, FInputChord());
 	UI_COMMAND(CollapseToSubGraph, "Collapse to Sub-graph", "Collapse the selection into a subgraph", EUserInterfaceActionType::Button, FInputChord());
 
-	UI_COMMAND(ConvertToEmbeddedDataflow, "Convert to Embedded", "Convert this Dataflow to be embedded in the current asset", EUserInterfaceActionType::Button, FInputChord());
-	UI_COMMAND(ConvertToExternalDataflow, "Convert to External", "Convert this Dataflow to be an external Dataflow asset that can be shared by other asset", EUserInterfaceActionType::Button, FInputChord());
+	UI_COMMAND(ConvertToEmbeddedDataflow, "Embed Dataflow", "Embed the Dataflow graph in the current asset", EUserInterfaceActionType::Button, FInputChord());
+	UI_COMMAND(ConvertToExternalDataflow, "Extract and Reference Dataflow...", "Convert this Dataflow to be an external Dataflow asset that can be shared by other asset and assign it to this asset", EUserInterfaceActionType::Button, FInputChord());
 	UI_COMMAND(ToggleDataflowEditingLock, "Lock/Unlock Editing", "Toggle editing for the current Dataflow graph.", EUserInterfaceActionType::ToggleButton, FInputChord());
 
 	for (const TPair<FName, TUniquePtr<UE::Dataflow::IDataflowConstructionViewMode>>& NameAndMode : UE::Dataflow::FRenderingViewModeFactory::GetInstance().GetViewModes())
