Chaos Cloth Asset - Made the Cloth and Outfit Asset have their Dataflow asset embedded.

[RN] Chaos Cloth and Outfit Asset - Converted assets to use an embedded version of the Dataflow.

#Synced-CL 48653274

diff --git a/Engine/Plugins/ChaosClothAssetEditor/Source/ChaosClothAssetTools/ChaosClothAssetTools.Build.cs b/Engine/Plugins/ChaosClothAssetEditor/Source/ChaosClothAssetTools/ChaosClothAssetTools.Build.cs
--- a/Engine/Plugins/ChaosClothAssetEditor/Source/ChaosClothAssetTools/ChaosClothAssetTools.Build.cs
+++ b/Engine/Plugins/ChaosClothAssetEditor/Source/ChaosClothAssetTools/ChaosClothAssetTools.Build.cs
@@ -26,6 +26,7 @@ public class ChaosClothAssetTools : ModuleRules
 				"ChaosClothEditor",  // For Chaos::FSimulationEditorExtender
 				"ChaosClothAsset",
 				"ChaosClothAssetEngine",
+				"DataflowEditor",
 				"DataflowEngine",
 				"SlateCore",
 				"UnrealEd",
diff --git a/Engine/Plugins/ChaosClothAssetEditor/Source/ChaosClothAssetTools/Private/ChaosClothAsset/ClothAssetFactory.cpp b/Engine/Plugins/ChaosClothAssetEditor/Source/ChaosClothAssetTools/Private/ChaosClothAsset/ClothAssetFactory.cpp
--- a/Engine/Plugins/ChaosClothAssetEditor/Source/ChaosClothAssetTools/Private/ChaosClothAsset/ClothAssetFactory.cpp
+++ b/Engine/Plugins/ChaosClothAssetEditor/Source/ChaosClothAssetTools/Private/ChaosClothAsset/ClothAssetFactory.cpp
@@ -1,15 +1,24 @@
 // Copyright Epic Games, Inc. All Rights Reserved.
 
 #include "ChaosClothAsset/ClothAssetFactory.h"
-#include "AssetRegistry/AssetRegistryModule.h"
 #include "ChaosClothAsset/ClothAsset.h"
-#include "ChaosClothAsset/CollectionClothFacade.h"
-#include "Dataflow/DataflowObject.h"
-#include "GeometryCollection/ManagedArrayCollection.h"
-#include "Misc/PackageName.h"
+#include "Dataflow/AssetDefinition_DataflowAsset.h"
+#include "HAL/IConsoleManager.h"
 
 #include UE_INLINE_GENERATED_CPP_BY_NAME(ClothAssetFactory)
 
+#define LOCTEXT_NAMESPACE "ChaosClothAssetFactory"
+
+DEFINE_LOG_CATEGORY_STATIC(LogChaosClothAssetFactory, Log, All);
+
+namespace UE::Chaos::ClothAsset::Private
+{
+	static TAutoConsoleVariable<bool> CVarCreateWithEmbeddedDataflow(
+		TEXT("p.ChaosClothAsset.CreateWithEmbeddedDataflow"),
+		true,
+		TEXT("Set how to create new Cloth Assets, either with embedded (default) or external (legacy) Dataflow assets.\n"));
+}
+
 UChaosClothAssetFactory::UChaosClothAssetFactory(const FObjectInitializer& ObjectInitializer)
 	: Super(ObjectInitializer)
 {
@@ -22,38 +31,15 @@ UObject* UChaosClothAssetFactory::FactoryCreateNew(UClass* Class, UObject* Paren
 {
 	using namespace UE::Chaos::ClothAsset;
 
-	UChaosClothAsset* const ClothAsset = NewObject<UChaosClothAsset>(Parent, Class, Name, Flags | RF_Transactional | RF_Public | RF_Standalone);
-	ClothAsset->MarkPackageDirty();
-
-	// Create a new Dataflow asset
-	const FString DataflowPath = FPackageName::GetLongPackagePath(ClothAsset->GetOutermost()->GetName());
-	const FString ClothAssetName = ClothAsset->GetName();
-	FString DataflowName = FString(TEXT("DF_")) + (ClothAssetName.StartsWith(TEXT("CA_")) ? ClothAssetName.RightChop(3) : ClothAssetName);
-	FString DataflowPackageName = FPaths::Combine(DataflowPath, DataflowName);
-	if (FindPackage(nullptr, *DataflowPackageName))
-	{
-		// If a Dataflow asset already exists with this name, make a unique name from it to avoid clobbering it
-		MakeUniqueObjectName(nullptr, UPackage::StaticClass(), FName(DataflowPackageName)).ToString(DataflowPackageName);
-		DataflowName = FPaths::GetBaseFilename(DataflowPackageName);
-	}
-	UPackage* const DataflowPackage = CreatePackage(*DataflowPackageName);
-
-	// Load the cloth template into the new Dataflow asset
-	UDataflow* const Template = LoadObject<UDataflow>(DataflowPackage, TEXT("/ChaosClothAssetEditor/ClothAssetTemplate.ClothAssetTemplate"));
-	UDataflow* const Dataflow = DuplicateObject(Template, DataflowPackage, FName(DataflowName));
-
-	Dataflow->MarkPackageDirty();
-
-	// Notify the asset registry
-	FAssetRegistryModule::AssetCreated(Dataflow);
-
-	// Set the Dataflow to the cloth asset
-	ClothAsset->SetDataflow(Dataflow);
-
-	return ClothAsset;
+	static const FString TemplatePath = TEXT("/ChaosClothAssetEditor/ClothAssetTemplate.ClothAssetTemplate");
+	static const TCHAR* AssetPrefix = TEXT("CA_");
+	const bool bEmbedDataflow = Private::CVarCreateWithEmbeddedDataflow.GetValueOnAnyThread();
+	return UE::DataflowAssetDefinitionHelpers::FactoryCreateNew(Class, Parent, Name, Flags, &TemplatePath, bEmbedDataflow, AssetPrefix);
 }
 
 FString UChaosClothAssetFactory::GetDefaultNewAssetName() const
 {
 	return FString(TEXT("CA_NewChaosClothAsset"));
 }
+
+#undef LOCTEXT_NAMESPACE
diff --git a/Engine/Plugins/Experimental/ChaosOutfitAsset/Source/ChaosOutfitAssetEditor/Private/ChaosOutfitAsset/OutfitAssetFactory.cpp b/Engine/Plugins/Experimental/ChaosOutfitAsset/Source/ChaosOutfitAssetEditor/Private/ChaosOutfitAsset/OutfitAssetFactory.cpp
--- a/Engine/Plugins/Experimental/ChaosOutfitAsset/Source/ChaosOutfitAssetEditor/Private/ChaosOutfitAsset/OutfitAssetFactory.cpp
+++ b/Engine/Plugins/Experimental/ChaosOutfitAsset/Source/ChaosOutfitAssetEditor/Private/ChaosOutfitAsset/OutfitAssetFactory.cpp
@@ -1,16 +1,25 @@
 // Copyright Epic Games, Inc. All Rights Reserved.
 
 #include "ChaosOutfitAsset/OutfitAssetFactory.h"
-#include "AssetRegistry/AssetRegistryModule.h"
 #include "ChaosOutfitAsset/OutfitAsset.h"
-#include "Dataflow/DataflowObject.h"
+#include "Dataflow/AssetDefinition_DataflowAsset.h"
 #include "Dialog/SMessageDialog.h"
-#include "Misc/PackageName.h"
+#include "HAL/IConsoleManager.h"
 
 #include UE_INLINE_GENERATED_CPP_BY_NAME(OutfitAssetFactory)
 
 #define LOCTEXT_NAMESPACE "OutfitAssetFactory"
 
+DEFINE_LOG_CATEGORY_STATIC(LogChaosOutfitAssetFactory, Log, All);
+
+namespace UE::Chaos::OutfitAsset::Private
+{
+	static TAutoConsoleVariable<bool> CVarCreateWithEmbeddedDataflow(
+		TEXT("p.ChaosOutfitAsset.CreateWithEmbeddedDataflow"),
+		true,
+		TEXT("Set how to create new Outfit Assets, either with embedded (default) or external (legacy) Dataflow assets.\n"));
+}
+
 UChaosOutfitAssetFactory::UChaosOutfitAssetFactory(const FObjectInitializer& ObjectInitializer)
 	: Super(ObjectInitializer)
 {
@@ -21,21 +30,7 @@ UChaosOutfitAssetFactory::UChaosOutfitAssetFactory(const FObjectInitializer& Obj
 
 UObject* UChaosOutfitAssetFactory::FactoryCreateNew(UClass* Class, UObject* Parent, FName Name, EObjectFlags Flags, UObject* /*Context*/, FFeedbackContext* /*Warn*/)
 {
-	UChaosOutfitAsset* const OutfitAsset = NewObject<UChaosOutfitAsset>(Parent, Class, Name, Flags | RF_Transactional | RF_Public | RF_Standalone);
-	OutfitAsset->MarkPackageDirty();
-
-	// Create a new Dataflow asset
-	const FString DataflowPath = FPackageName::GetLongPackagePath(OutfitAsset->GetOutermost()->GetName());
-	const FString OutfitAssetName = OutfitAsset->GetName();
-	FString DataflowName = FString(TEXT("DF_")) + (OutfitAssetName.StartsWith(TEXT("OA_")) ? OutfitAssetName.RightChop(3) : OutfitAssetName);
-	FString DataflowPackageName = FPaths::Combine(DataflowPath, DataflowName);
-	if (FindPackage(nullptr, *DataflowPackageName))
-	{
-		// If a Dataflow asset already exists with this name, make a unique name from it to avoid clobbering it
-		MakeUniqueObjectName(nullptr, UPackage::StaticClass(), FName(DataflowPackageName)).ToString(DataflowPackageName);
-		DataflowName = FPaths::GetBaseFilename(DataflowPackageName);
-	}
-	UPackage* const DataflowPackage = CreatePackage(*DataflowPackageName);
+	using namespace UE::Chaos::OutfitAsset;
 
 	// Select the template
 	TSharedRef<SMessageDialog> SelectTemplateMessageDialog = SNew(SMessageDialog)
@@ -53,47 +48,32 @@ UObject* UChaosOutfitAssetFactory::FactoryCreateNew(UClass* Class, UObject* Pare
 					.SetToolTipText(LOCTEXT("SelectSimpleOutfitTemplateTooltip", "Add a Dataflow with a simple Cloth Asset aggregator graph. Allows to simulate multiple Cloth Assets from the same ChaosClothComponent."))
 					.SetPrimary(true),
 				SMessageDialog::FButton(LOCTEXT("SelectResizableOutfitTemplate", "Resizable Outfit"))
-					.SetToolTipText(LOCTEXT("SelectResizableOutfitTemplateTooltip", "Add a Dataflow that builds a single resizable garment from multiple Cloth Assets. Body sizes will have to be provided in addition to the mutiple Cloth Assets."))
+					.SetToolTipText(LOCTEXT("SelectResizableOutfitTemplateTooltip", "Add a Dataflow that builds a single resizable garment from multiple Cloth Assets. Body sizes will have to be provided in addition to the multiple Cloth Assets."))
 					.SetPrimary(false),
 				SMessageDialog::FButton(LOCTEXT("SelectResizingGraph", "Resizing Graph"))
 					.SetToolTipText(LOCTEXT("SelectResizingGraphTooltip", "Resizing graph to test a resizable Outfit Asset."))
 					.SetPrimary(false),
 			});
 
-	FString OutfitAssetTemplate;
+	const FString* TemplatePath;
+	static const FString EmptyOutfitTemplatePath(TEXT("/ChaosOutfitAsset/EmptyOutfitAssetTemplate.EmptyOutfitAssetTemplate"));
+	static const FString SimpleOutfitTemplatePath(TEXT("/ChaosOutfitAsset/OutfitAssetTemplate.OutfitAssetTemplate"));
+	static const FString ResizableOutfitTemplatePath(TEXT("/ChaosOutfitAsset/MakeResizableOutfitTemplate.MakeResizableOutfitTemplate"));
+	static const FString ResizingOutfitTemplatePath(TEXT("/ChaosOutfitAsset/ResizeOutfitTemplate.ResizeOutfitTemplate"));
 	switch (SelectTemplateMessageDialog->ShowModal())
 	{
 	default:
-	case 0:  // No Dataflow
-		break;
-	case 1:  // Empty Dataflow
-		OutfitAssetTemplate = TEXT("/ChaosOutfitAsset/EmptyOutfitAssetTemplate.EmptyOutfitAssetTemplate");
-		break;
-	case 2:  // Simple Outfit
-		OutfitAssetTemplate = TEXT("/ChaosOutfitAsset/OutfitAssetTemplate.OutfitAssetTemplate");
-		break;
-	case 3:  // Resizable Outfit
-		OutfitAssetTemplate = TEXT("/ChaosOutfitAsset/MakeResizableOutfitTemplate.MakeResizableOutfitTemplate");
-		break;
-	case 4:  // Resizing Graph
-		OutfitAssetTemplate = TEXT("/ChaosOutfitAsset/ResizeOutfitTemplate.ResizeOutfitTemplate");
-		break;
-	}
-
-	// Load the cloth template into the new Dataflow asset
-	UDataflow* const Template = OutfitAssetTemplate.IsEmpty() ? nullptr: LoadObject<UDataflow>(DataflowPackage, *OutfitAssetTemplate);
-	if (UDataflow* const Dataflow = Template ? DuplicateObject(Template, DataflowPackage, FName(DataflowName)) : nullptr)
-	{
-		Dataflow->MarkPackageDirty();
-
-		// Notify the asset registry
-		FAssetRegistryModule::AssetCreated(Dataflow);
-
-		// Set the Dataflow to the cloth asset
-		OutfitAsset->SetDataflow(Dataflow);
+	case 0: TemplatePath = nullptr; break;  // No Dataflow
+	case 1: TemplatePath = &EmptyOutfitTemplatePath; break;  // Empty Dataflow with terminal node
+	case 2: TemplatePath = &SimpleOutfitTemplatePath; break;  // Simple Outfit
+	case 3: TemplatePath = &ResizableOutfitTemplatePath; break;  // Resizable Outfit
+	case 4: TemplatePath = &ResizingOutfitTemplatePath; break;  // Resizing Graph
 	}
 
-	return OutfitAsset;
+	// Create a new asset using the provided template path
+	static const TCHAR* AssetPrefix = TEXT("OA_");
+	const bool bEmbedDataflow = Private::CVarCreateWithEmbeddedDataflow.GetValueOnAnyThread();
+	return UE::DataflowAssetDefinitionHelpers::FactoryCreateNew(Class, Parent, Name, Flags, TemplatePath, bEmbedDataflow, AssetPrefix);
 }
 
 FString UChaosOutfitAssetFactory::GetDefaultNewAssetName() const
diff --git a/Engine/Plugins/Experimental/Dataflow/Source/DataflowEditor/Private/Dataflow/AssetDefinition_DataflowAsset.cpp b/Engine/Plugins/Experimental/Dataflow/Source/DataflowEditor/Private/Dataflow/AssetDefinition_DataflowAsset.cpp
--- a/Engine/Plugins/Experimental/Dataflow/Source/DataflowEditor/Private/Dataflow/AssetDefinition_DataflowAsset.cpp
+++ b/Engine/Plugins/Experimental/Dataflow/Source/DataflowEditor/Private/Dataflow/AssetDefinition_DataflowAsset.cpp
@@ -3,20 +3,21 @@
 #include "Dataflow/AssetDefinition_DataflowAsset.h"
 
 #include "AssetRegistry/AssetRegistryModule.h"
+#include "ContentBrowserMenuContexts.h"
 #include "ContentBrowserModule.h"
 #include "Dataflow/DataflowEditor.h"
 #include "Dataflow/DataflowEditorModule.h"
 #include "Dataflow/DataflowObject.h"
 #include "Dataflow/DataflowInstance.h"
 #include "Dialog/SMessageDialog.h"
+#include "Framework/Notifications/NotificationManager.h"
 #include "IContentBrowserSingleton.h"
 #include "Math/Color.h"
 #include "Misc/FileHelper.h"
 #include "ThumbnailRendering/SceneThumbnailInfo.h"
-
-#include "ContentBrowserMenuContexts.h"
 #include "ToolMenu.h"
 #include "ToolMenus.h"
+#include "Widgets/Notifications/SNotificationList.h"
 
 
 #include UE_INLINE_GENERATED_CPP_BY_NAME(AssetDefinition_DataflowAsset)
@@ -24,7 +25,6 @@
 
 #define LOCTEXT_NAMESPACE "AssetActions_DataflowAsset"
 
-
 bool bCanEditDataflow = false;
 FAutoConsoleVariableRef CVarDataflowIsEditable(TEXT("p.Dataflow.IsEditable"), bCanEditDataflow, TEXT("Whether to allow edits of the dataflow [def:true]"));
 
@@ -364,6 +364,119 @@ namespace UE::DataflowAssetDefinitionHelpers
 				}));
 	}
 
+	UObject* FactoryCreateNew(
+		UClass* AssetClass,
+		UObject* Parent,
+		FName Name,
+		EObjectFlags Flags,
+		const FString* DataflowTemplatePath,
+		bool bEmbedDataflow,
+		const TCHAR* AssetPrefix)
+	{
+		UObject* const NewAsset = NewObject<UObject>(Parent, AssetClass, Name, Flags | RF_Transactional | RF_Public | RF_Standalone);
+		if (!NewAsset)
+		{
+			FNotificationInfo NotificationInfo(FText::Format(LOCTEXT("ErrorCreatingNewAsset",
+				"Failed to create a new asset {0}."), FText::FromName(Name)));
+			NotificationInfo.ExpireDuration = 10.f;
+			FSlateNotificationManager::Get().AddNotification(NotificationInfo);
+			UE_LOG(LogDataflowEditor, Error, TEXT("%s"), *NotificationInfo.Text.Get().ToString());
+			return nullptr;
+		}
+		NewAsset->MarkPackageDirty();
+
+		if (!DataflowTemplatePath || !ensureMsgf(AssetClass->ImplementsInterface(UDataflowInstanceInterface::StaticClass()), TEXT("Only use UE::DataflowAssetDefinitionHelpers::FactoryCreateNew() with Dataflow asset types.")))
+		{
+			return NewAsset;
+		}
+
+		UDataflow* DataflowTemplate = nullptr;
+		if (!DataflowTemplatePath->IsEmpty())
+		{
+			DataflowTemplate = LoadObject<UDataflow>(nullptr, *DataflowTemplatePath);
+			if (!DataflowTemplate)
+			{
+				FNotificationInfo NotificationInfo(FText::Format(LOCTEXT("ErrorLoadingTemplate",
+					"Failed to load Dataflow template [{0}].\n"
+					"An empty Dataflow will be used to create new asset {1}."), FText::FromString(*DataflowTemplatePath), FText::FromName(Name)));
+				NotificationInfo.ExpireDuration = 10.f;
+				FSlateNotificationManager::Get().AddNotification(NotificationInfo);
+				UE_LOG(LogDataflowEditor, Error, TEXT("%s"), *NotificationInfo.Text.Get().ToString());
+			}
+		}
+		if (!DataflowTemplate)
+		{
+			DataflowTemplate = NewObject<UDataflow>();  // Use empty Dataflow instead
+		}
+
+		if (bEmbedDataflow)
+		{
+			// Duplicate it and embed it into the new asset
+			if (UDataflow* const Dataflow = DuplicateObject(DataflowTemplate, NewAsset, TEXT("EmbeddedDataflow")))
+			{
+				// Mark as transactional
+				Dataflow->SetFlags(RF_Transactional);
+
+				// Set the Dataflow to the new asset
+				CastChecked<IDataflowInstanceInterface>(NewAsset)->GetDataflowInstance().SetDataflowAsset(Dataflow);
+			}
+			else
+			{
+				FNotificationInfo NotificationInfo(FText::Format(LOCTEXT("ErrorEmbeddingDataflow",
+					"Failed to embed a Dataflow into new asset {0}."), FText::FromName(Name)));
+				NotificationInfo.ExpireDuration = 10.f;
+				FSlateNotificationManager::Get().AddNotification(NotificationInfo);
+				UE_LOG(LogDataflowEditor, Error, TEXT("%s"), *NotificationInfo.Text.Get().ToString());
+			}
+		}
+		else
+		{
+			// Create a new Dataflow asset
+			static const TCHAR* DataflowAssetPrefix = TEXT("DF_");
+			const FString DataflowPath = FPackageName::GetLongPackagePath(NewAsset->GetOutermost()->GetName());
+			const FString OutfitAssetName = NewAsset->GetName();
+			FString DataflowName = FString(DataflowAssetPrefix) +
+				(AssetPrefix && OutfitAssetName.StartsWith(AssetPrefix) ? OutfitAssetName.RightChop(FCString::Strlen(AssetPrefix)) : OutfitAssetName);
+			FString DataflowPackageName = FPaths::Combine(DataflowPath, DataflowName);
+			if (FindPackage(nullptr, *DataflowPackageName))
+			{
+				// If a Dataflow asset already exists with this name, make a unique name from it to avoid clobbering it
+				MakeUniqueObjectName(nullptr, UPackage::StaticClass(), FName(DataflowPackageName)).ToString(DataflowPackageName);
+				DataflowName = FPaths::GetBaseFilename(DataflowPackageName);
+			}
+			if (UPackage* const DataflowPackage = CreatePackage(*DataflowPackageName))
+			{
+				// Load the outfit template into the new Dataflow asset
+				if (UDataflow* const Dataflow = DuplicateObject(DataflowTemplate, DataflowPackage, FName(DataflowName)))
+				{
+					// Mark as a transactional asset and notify creation
+					Dataflow->SetFlags(RF_Transactional | RF_Public | RF_Standalone);
+					Dataflow->MarkPackageDirty();
+					FAssetRegistryModule::AssetCreated(Dataflow);
+
+					// Set the Dataflow to the new asset
+					CastChecked<IDataflowInstanceInterface>(NewAsset)->GetDataflowInstance().SetDataflowAsset(Dataflow);
+				}
+				else
+				{
+					FNotificationInfo NotificationInfo(FText::Format(LOCTEXT("ErrorCreatingDataflow",
+						"Failed to create a new Dataflow for new asset {0}."), FText::FromName(Name)));
+					NotificationInfo.ExpireDuration = 10.f;
+					FSlateNotificationManager::Get().AddNotification(NotificationInfo);
+					UE_LOG(LogDataflowEditor, Error, TEXT("%s"), *NotificationInfo.Text.Get().ToString());
+				}
+			}
+			else
+			{
+				FNotificationInfo NotificationInfo(FText::Format(LOCTEXT("ErrorCreatingPackage",
+					"Failed to create a new Dataflow package for new asset {0}."), FText::FromName(Name)));
+				NotificationInfo.ExpireDuration = 10.f;
+				FSlateNotificationManager::Get().AddNotification(NotificationInfo);
+				UE_LOG(LogDataflowEditor, Error, TEXT("%s"), *NotificationInfo.Text.Get().ToString());
+			}
+		}
+		return NewAsset;
+	}
 }
 
 
diff --git a/Engine/Plugins/Experimental/Dataflow/Source/DataflowEditor/Public/Dataflow/AssetDefinition_DataflowAsset.h b/Engine/Plugins/Experimental/Dataflow/Source/DataflowEditor/Public/Dataflow/AssetDefinition_DataflowAsset.h
--- a/Engine/Plugins/Experimental/Dataflow/Source/DataflowEditor/Public/Dataflow/AssetDefinition_DataflowAsset.h
+++ b/Engine/Plugins/Experimental/Dataflow/Source/DataflowEditor/Public/Dataflow/AssetDefinition_DataflowAsset.h
@@ -9,6 +9,7 @@
 #define UE_API DATAFLOWEDITOR_API
 
 class UDataflow;
+class UDataflowInstanceInterface;
 
 namespace UE::DataflowAssetDefinitionHelpers
 {
@@ -16,7 +17,7 @@ namespace UE::DataflowAssetDefinitionHelpers
 	// @return true is the asset is successfully created ( in that case OutDataflowAsset will containt the pointer to this asset )
 	DATAFLOWEDITOR_API bool CreateNewDataflowAsset(const UObject* BoundAsset, UObject*& OutDataflowAsset);
 
-	// Take an embedded dataflow asset from BoundAsset and make it shareable by being it's own seperate asset
+	// Take an embedded dataflow asset from BoundAsset and make it shareable by being its own separate asset
 	// If the Dataflow asset was already shareable , do nothing 
 	// @return true if the operation was a success ( in that case OutDataflowAsset will containt the pointer to the shareable Dataflow asset )
 	DATAFLOWEDITOR_API bool MakeShareableDataflowAsset(UObject* BoundAsset, UDataflow*& OutDataflowAsset);
@@ -43,6 +44,18 @@ namespace UE::DataflowAssetDefinitionHelpers
 
 	// Register a entry on the asset menu that will Re-Evaluate the assoiated dataflow for the selected asset
 	DATAFLOWEDITOR_API void RegisterReEvaluateDataflowAssetMenu(UClass* AssetClass);
+
+	// Create a new Dataflow based asset using the specified Dataflow template, an asset with an empty Dataflow if an empty path is provided, or an unbound asset if the path is nullptr.
+	// @note The Dataflow template is always duplicated to use with the new asset, even when not embedded.
+	// @return A new asset of type AssetClass that references a new Dataflow asset of the specified template.
+	DATAFLOWEDITOR_API UObject* FactoryCreateNew(
+		UClass* AssetClass,
+		UObject* Parent,
+		FName Name,
+		EObjectFlags Flags,
+		const FString* DataflowTemplatePath,
+		bool bEmbedDataflow,
+		const TCHAR* AssetPrefix);
 }
 
 #define DATAFLOW_REGISTER_ASSET_MENUS(AssetClassType) \
