Fix for crash on shutdown in DataflowEditorModule caused by a dangling UObject reference

#Synced-CL 48786356

diff --git a/Engine/Plugins/Experimental/Dataflow/Source/DataflowEditor/Private/Dataflow/DataflowEditorModule.cpp b/Engine/Plugins/Experimental/Dataflow/Source/DataflowEditor/Private/Dataflow/DataflowEditorModule.cpp
--- a/Engine/Plugins/Experimental/Dataflow/Source/DataflowEditor/Private/Dataflow/DataflowEditorModule.cpp
+++ b/Engine/Plugins/Experimental/Dataflow/Source/DataflowEditor/Private/Dataflow/DataflowEditorModule.cpp
@@ -323,7 +323,11 @@ void FDataflowEditorModule::ShutdownModule()
 		PropertyModule->UnregisterCustomPropertyTypeLayout(DataflowFreezeActionsName);
 		PropertyModule->UnregisterCustomPropertyTypeLayout(DataflowColorRampName);
 		PropertyModule->UnregisterCustomPropertyTypeLayout(DataflowLinearColorRampName);
-		PropertyModule->UnregisterCustomClassLayout(UDataflowEditorVertexAttributePaintToolProperties::StaticClass()->GetFName());
+
+		if (UObjectInitialized())
+		{
+			PropertyModule->UnregisterCustomClassLayout(UDataflowEditorVertexAttributePaintToolProperties::StaticClass()->GetFName());
+		}
 
 		UE::Dataflow::UnregisterAnyTypeCustomizations(*PropertyModule);
 	}
