[Dataflow] Fix copy/paste ignoring properties that were tagged as Instanced in their UPROPERTY data, causing instanced objects to all point at the same target.

#Synced-CL 48785580

diff --git a/Engine/Plugins/Experimental/Dataflow/Source/DataflowEditor/Private/Dataflow/DataflowAssetEditUtils.cpp b/Engine/Plugins/Experimental/Dataflow/Source/DataflowEditor/Private/Dataflow/DataflowAssetEditUtils.cpp
--- a/Engine/Plugins/Experimental/Dataflow/Source/DataflowEditor/Private/Dataflow/DataflowAssetEditUtils.cpp
+++ b/Engine/Plugins/Experimental/Dataflow/Source/DataflowEditor/Private/Dataflow/DataflowAssetEditUtils.cpp
@@ -635,7 +635,7 @@ namespace UE::Dataflow
 
 					FString ContentString;
 					const TUniquePtr<FDataflowNode> DefaultElement;
-					DataflowNode->TypedScriptStruct()->ExportText(ContentString, DataflowNode.Get(), DataflowNode.Get(), nullptr, PPF_None, nullptr);
+					DataflowNode->TypedScriptStruct()->ExportText(ContentString, DataflowNode.Get(), DataflowNode.Get(), UDataflow::GetDataflowAssetFromEdGraph(EdNode->GetGraph()), PPF_Copy, nullptr);
 
 					FDataflowNodeData NodeData
 					{
@@ -757,7 +757,7 @@ namespace UE::Dataflow
 					// load properties to DataflowNode
 					if (!NodeData.Properties.IsEmpty())
 					{
-						DataflowNode->TypedScriptStruct()->ImportText(*NodeData.Properties, DataflowNode.Get(), nullptr, EPropertyPortFlags::PPF_None, nullptr, DataflowNode->TypedScriptStruct()->GetName(), true);
+						DataflowNode->TypedScriptStruct()->ImportText(*NodeData.Properties, DataflowNode.Get(), &DataflowAsset, EPropertyPortFlags::PPF_InstanceSubobjects, nullptr, DataflowNode->TypedScriptStruct()->GetName(), true);
 					}
 
 					// if we are pasting a dataflow variable, let's create the variable if needed
