Dataflow : fix select vertex -> relax operation not working consistently
- issue was that the stamp.power is not always set properly when applying by selection
- changed the method of applying operation on selection by avoid passing the stampd and instead passing individual tool value, center and radius

#Synced-CL 48816559

diff --git a/Engine/Plugins/Experimental/Dataflow/Source/DataflowEditor/Private/DataflowEditorTools/DataflowEditorVertexAttributePaintTool.cpp b/Engine/Plugins/Experimental/Dataflow/Source/DataflowEditor/Private/DataflowEditorTools/DataflowEditorVertexAttributePaintTool.cpp
--- a/Engine/Plugins/Experimental/Dataflow/Source/DataflowEditor/Private/DataflowEditorTools/DataflowEditorVertexAttributePaintTool.cpp
+++ b/Engine/Plugins/Experimental/Dataflow/Source/DataflowEditor/Private/DataflowEditorTools/DataflowEditorVertexAttributePaintTool.cpp
@@ -1685,9 +1685,11 @@ void UDataflowEditorVertexAttributePaintTool::ApplyValueToSelection(EDataflowEdi
 		ROIWeightValueBuffer.SetNum(VertexROI.Num(), EAllowShrinking::No);
 		SyncWeightBufferWithMesh(Mesh);
 
-		FDataflowVertexAttributePaintBrushOp::ApplyStampByVerticesStatic(
-			Mesh, CurrentStamp, VertexROI, ROIWeightValueBuffer, 
-			Operation, InValue, /*bApplyRadiusLimit*/false);
+		FDataflowVertexAttributePaintBrushOp::ApplyToVerticesStatic(
+			Mesh, VertexROI, ROIWeightValueBuffer, 
+			Operation, InValue, 
+			/*bApplyRadiusLimit*/false, CurrentStamp.LocalFrame.Origin, CurrentStamp.Radius
+		);
 
 		SyncMeshWithWeightBuffer(Mesh);
 	}
diff --git a/Engine/Plugins/Experimental/Dataflow/Source/DataflowEditor/Private/DataflowEditorTools/DataflowEditorWeightMapPaintBrushOps.cpp b/Engine/Plugins/Experimental/Dataflow/Source/DataflowEditor/Private/DataflowEditorTools/DataflowEditorWeightMapPaintBrushOps.cpp
--- a/Engine/Plugins/Experimental/Dataflow/Source/DataflowEditor/Private/DataflowEditorTools/DataflowEditorWeightMapPaintBrushOps.cpp
+++ b/Engine/Plugins/Experimental/Dataflow/Source/DataflowEditor/Private/DataflowEditorTools/DataflowEditorWeightMapPaintBrushOps.cpp
@@ -90,12 +90,25 @@ void FDataflowWeightMapSmoothBrushOp::ApplyStampByVertices(
 }
 
 
-void FDataflowWeightMapSmoothBrushOp::ApplyStampByVerticesStatic(
+/*static*/ void FDataflowWeightMapSmoothBrushOp::ApplyStampByVerticesStatic(
 	const FDynamicMesh3* Mesh,
 	const FSculptBrushStamp& Stamp,
 	const TArray<int32>& Vertices,
 	TArray<double>& VertexWeightValues,
 	bool bApplyRadiusLimit)
+{
+	FDataflowWeightMapSmoothBrushOp::ApplyToVerticesStatic(Mesh, Vertices, VertexWeightValues, Stamp.Power, bApplyRadiusLimit, Stamp.LocalFrame.Origin, Stamp.Radius);
+}
+
+/*static*/ void FDataflowWeightMapSmoothBrushOp::ApplyToVerticesStatic(
+	const FDynamicMesh3* Mesh,
+	const TArray<int32>& Vertices,
+	TArray<double>& VertexWeightValues,
+	const double SmoothingFactor,
+	bool bApplyRadiusLimit,
+	const FVector Center,
+	const float Radius
+)
 {
 	// Converted from FClothPaintTool_Smooth::SmoothVertices
 	const int32 NumVertices = Vertices.Num();
@@ -140,21 +153,21 @@ void FDataflowWeightMapSmoothBrushOp::ApplyStampByVerticesStatic(
 
 	// Blend vertex value with its average one-ring value
 	ParallelFor(NumVertices,
-		[&OneRingAverages, &VertexWeightValues, bApplyRadiusLimit, &Stamp, &Vertices, &Mesh](int32 BufferIndex)
+		[&OneRingAverages, &VertexWeightValues, bApplyRadiusLimit, &Center, &Radius, &SmoothingFactor, &Vertices, &Mesh](int32 BufferIndex)
 		{
 			const double Diff = OneRingAverages[BufferIndex] - VertexWeightValues[BufferIndex];
 			if (bApplyRadiusLimit)
 			{
-				const FVector3d& StampPos = Stamp.LocalFrame.Origin;
+				const FVector3d& StampPos = Center;
 				const int32 VertexIndex = Vertices[BufferIndex];
 				const FVector3d VertexPos = Mesh->GetVertex(VertexIndex);
 				const double DistanceSquared = (VertexPos - StampPos).SquaredLength();
-				if (DistanceSquared >= Stamp.Radius * Stamp.Radius)
+				if (DistanceSquared >= Radius * Radius)
 				{
 					return;
 				}
 			}
-			VertexWeightValues[BufferIndex] += Stamp.Power * Diff;
+			VertexWeightValues[BufferIndex] += SmoothingFactor * Diff;
 		}, 
 		EParallelForFlags::None
 	);
@@ -184,27 +197,40 @@ void FDataflowVertexAttributePaintBrushOp::ApplyStampByVertices(
 	EDataflowEditorToolEditOperation EditOperation,
 	const double ToolValue,
 	bool bApplyRadiusLimit)
+{
+	FDataflowVertexAttributePaintBrushOp::ApplyToVerticesStatic(Mesh, Vertices, OutValues, EditOperation, ToolValue, bApplyRadiusLimit, Stamp.LocalFrame.Origin, Stamp.Radius);
+}
+
+/*static*/ void FDataflowVertexAttributePaintBrushOp::ApplyToVerticesStatic(
+	const FDynamicMesh3* Mesh,
+	const TArray<int32>& Vertices,
+	TArray<double>& OutValues,
+	EDataflowEditorToolEditOperation EditOperation,
+	const double ToolValue,
+	bool bApplyRadiusLimit,
+	const FVector Center,
+	const float Radius
+)
 {
 	// special case for Relax option 
 	if (EditOperation == EDataflowEditorToolEditOperation::Relax)
 	{
-		FDataflowWeightMapSmoothBrushOp::ApplyStampByVerticesStatic(Mesh, Stamp, Vertices, OutValues, bApplyRadiusLimit);
+		FDataflowWeightMapSmoothBrushOp::ApplyToVerticesStatic(Mesh, Vertices, OutValues, ToolValue, bApplyRadiusLimit, Center, Radius);
 		return;
 	}
 
-	// other operations
 	const int32 NumVertices = Vertices.Num();
-	ParallelFor(NumVertices, 
-		[bApplyRadiusLimit, ToolValue, EditOperation , &Stamp, &Mesh, &Vertices, &OutValues]
+	ParallelFor(NumVertices,
+		[bApplyRadiusLimit, ToolValue, EditOperation, &Center, &Radius, &Mesh, &Vertices, &OutValues]
 		(int32 BufferIndex)
 		{
 			if (bApplyRadiusLimit)
 			{
-				const FVector3d& StampPos = Stamp.LocalFrame.Origin;
+				const FVector3d& StampPos = Center;
 				const int32 MeshVertIdx = Vertices[BufferIndex];
 				const FVector3d VertexPos = Mesh->GetVertex(MeshVertIdx);
 				const double DistanceSquared = (VertexPos - StampPos).SquaredLength();
-				if (DistanceSquared >= Stamp.Radius * Stamp.Radius)
+				if (DistanceSquared >= Radius * Radius)
 				{
 					return;
 				}
diff --git a/Engine/Plugins/Experimental/Dataflow/Source/DataflowEditor/Private/DataflowEditorTools/DataflowEditorWeightMapPaintBrushOps.h b/Engine/Plugins/Experimental/Dataflow/Source/DataflowEditor/Private/DataflowEditorTools/DataflowEditorWeightMapPaintBrushOps.h
--- a/Engine/Plugins/Experimental/Dataflow/Source/DataflowEditor/Private/DataflowEditorTools/DataflowEditorWeightMapPaintBrushOps.h
+++ b/Engine/Plugins/Experimental/Dataflow/Source/DataflowEditor/Private/DataflowEditorTools/DataflowEditorWeightMapPaintBrushOps.h
@@ -159,6 +159,15 @@ public:
 		const TArray<int32>& Vertices,
 		TArray<double>& NewAttributesOut,
 		bool bApplyRadiusLimit);
+
+	static void ApplyToVerticesStatic(
+		const FDynamicMesh3* Mesh,
+		const TArray<int32>& Vertices,
+		TArray<double>& NewAttributesOut,
+		const double SmoothingFactor,
+		bool bApplyRadiusLimit,
+		const FVector Center,
+		const float Radius);
 };
 
 ////////////////////////////////////////////////////////////////////////////////////////////////////////////////////
@@ -208,4 +217,14 @@ public:
 		EDataflowEditorToolEditOperation EditOperation,
 		const double ToolValue,
 		bool bApplyRadiusLimit);
+
+	static void ApplyToVerticesStatic(
+		const FDynamicMesh3* Mesh,
+		const TArray<int32>& Vertices,
+		TArray<double>& NewAttributesOut,
+		EDataflowEditorToolEditOperation EditOperation,
+		const double ToolValue,
+		bool bApplyRadiusLimit,
+		const FVector Center, 
+		const float Radius);
 };
