Dataflow weight map paint tool : add option to toggle backface painting or not  ( on by default )  and make brush disappear when not on the mesh and not being actively dragged

#Synced-CL 48816559

diff --git a/Engine/Plugins/Experimental/Dataflow/Source/DataflowEditor/Private/DataflowEditorTools/DataflowEditorVertexAttributePaintTool.cpp b/Engine/Plugins/Experimental/Dataflow/Source/DataflowEditor/Private/DataflowEditorTools/DataflowEditorVertexAttributePaintTool.cpp
--- a/Engine/Plugins/Experimental/Dataflow/Source/DataflowEditor/Private/DataflowEditorTools/DataflowEditorVertexAttributePaintTool.cpp
+++ b/Engine/Plugins/Experimental/Dataflow/Source/DataflowEditor/Private/DataflowEditorTools/DataflowEditorVertexAttributePaintTool.cpp
@@ -2016,6 +2016,13 @@ bool UDataflowEditorVertexAttributePaintTool::UpdateBrushPosition(const FRay& Wo
 	return bHit;
 }
 
+FInputRayHit UDataflowEditorVertexAttributePaintTool::BeginHoverSequenceHitTest(const FInputDeviceRay& Ray)
+{
+	FInputRayHit RayHit = Super::BeginHoverSequenceHitTest(Ray);
+	bBrushIsHoverMesh = RayHit.bHit;
+	return RayHit;
+}
+
 bool UDataflowEditorVertexAttributePaintTool::OnUpdateHover(const FInputDeviceRay& DevicePos)
 {
 	if (ensure(InStroke() == false))
@@ -2052,7 +2059,7 @@ void UDataflowEditorVertexAttributePaintTool::DrawHUD(FCanvas* Canvas, IToolsCon
 	// when the user adjust the brush interactively we hide the value on brush
 	if (!bInteractiveBrushIsAdjusting && IsInBrushMode())
 	{
-		if (Canvas && RenderAPI && BrushIndicator)
+		if (Canvas && RenderAPI && BrushIndicator && bBrushIsHoverMesh)
 		{
 			// Display the value under the brush
 			const float BrushRadius = BrushIndicator->BrushRadius;
@@ -2128,7 +2135,7 @@ void UDataflowEditorVertexAttributePaintTool::OnTick(float DeltaTime)
 	}
 
 	ConfigureIndicator(IsVolumetricBrush());
-	SetIndicatorVisibility(!bIsLasso && !bIsMeshEditMode);
+	SetIndicatorVisibility(!bIsLasso && !bIsMeshEditMode && bBrushIsHoverMesh);
 
 	SCOPE_CYCLE_COUNTER(VertexAttributePaintTool_Tick);
 
diff --git a/Engine/Plugins/Experimental/Dataflow/Source/DataflowEditor/Private/DataflowEditorTools/DataflowEditorVertexAttributePaintToolDetailCustomization.cpp b/Engine/Plugins/Experimental/Dataflow/Source/DataflowEditor/Private/DataflowEditorTools/DataflowEditorVertexAttributePaintToolDetailCustomization.cpp
--- a/Engine/Plugins/Experimental/Dataflow/Source/DataflowEditor/Private/DataflowEditorTools/DataflowEditorVertexAttributePaintToolDetailCustomization.cpp
+++ b/Engine/Plugins/Experimental/Dataflow/Source/DataflowEditor/Private/DataflowEditorTools/DataflowEditorVertexAttributePaintToolDetailCustomization.cpp
@@ -639,6 +639,36 @@ namespace UE::Dataflow::Editor
 
 	///////////////////////////////////////////////////////////////////////////////////////////////////////////////
 
+	void FVertexAttributePaintToolDetailCustomization::AddBrushHitBackFacesRow(IDetailCategoryBuilder& BrushCategory) const
+	{
+		BrushCategory.AddCustomRow(LOCTEXT("BrushHitBackFacesCategory", "Brush Falloff"), false)
+			.NameContent()
+			[
+				SNew(STextBlock)
+					.Text(LOCTEXT("BrushHitBackFacesLabel", "Hit Back faces"))
+					.Font(FAppStyle::GetFontStyle(TEXT("PropertyWindow.NormalFont")))
+					.ToolTipText(LOCTEXT("BrushHitBackFacesTooltip", "Allow the brush to hit the back side of the mesh"))
+			]
+			.ValueContent()
+			[
+				SNew(SCheckBox)
+				.IsChecked_Lambda([this]()
+					{
+						const bool bIsChecked = (Tool.IsValid() && Tool->BrushProperties && Tool->BrushProperties->bHitBackFaces);
+						return bIsChecked ? ECheckBoxState::Checked : ECheckBoxState::Unchecked;
+					})
+				.OnCheckStateChanged_Lambda([this](ECheckBoxState InCheckBoxState)
+					{
+						if (Tool.IsValid() && Tool->BrushProperties)
+						{
+							Tool->BrushProperties->bHitBackFaces = (InCheckBoxState == ECheckBoxState::Checked);
+						}
+					})
+			];
+	}
+
+	///////////////////////////////////////////////////////////////////////////////////////////////////////////////
+
 	void FVertexAttributePaintToolDetailCustomization::AddBrushUI(IDetailLayoutBuilder& DetailBuilder) const
 	{
 		IDetailCategoryBuilder& BrushCategory = DetailBuilder.EditCategory("Brush", FText::GetEmpty(), ECategoryPriority::Important);
@@ -657,6 +687,10 @@ namespace UE::Dataflow::Editor
 
 			// Brush value field
 			//AddBrushFalloffRow(BrushCategory);
+
+			// Brush backface hit option
+			AddBrushHitBackFacesRow(BrushCategory);
+
 		}
 	}
 
diff --git a/Engine/Plugins/Experimental/Dataflow/Source/DataflowEditor/Private/DataflowEditorTools/DataflowEditorVertexAttributePaintToolDetailCustomization.h b/Engine/Plugins/Experimental/Dataflow/Source/DataflowEditor/Private/DataflowEditorTools/DataflowEditorVertexAttributePaintToolDetailCustomization.h
--- a/Engine/Plugins/Experimental/Dataflow/Source/DataflowEditor/Private/DataflowEditorTools/DataflowEditorVertexAttributePaintToolDetailCustomization.h
+++ b/Engine/Plugins/Experimental/Dataflow/Source/DataflowEditor/Private/DataflowEditorTools/DataflowEditorVertexAttributePaintToolDetailCustomization.h
@@ -41,6 +41,8 @@ namespace UE::Dataflow::Editor
 		void AddBrushRadiusRow(IDetailCategoryBuilder& BrushCategory) const;
 		void AddBrushValueRow(IDetailCategoryBuilder& BrushCategory) const;
 		void AddBrushFalloffRow(IDetailCategoryBuilder& BrushCategory) const;
+		void AddBrushHitBackFacesRow(IDetailCategoryBuilder& BrushCategory) const;
+		
 
 		void AddSelectionUI(IDetailLayoutBuilder& DetailBuilder) const;
 		TSharedRef<SWidget> MakeSelectionElementsToolbar() const;
diff --git a/Engine/Plugins/Experimental/Dataflow/Source/DataflowEditor/Public/DataflowEditorTools/DataflowEditorVertexAttributePaintTool.h b/Engine/Plugins/Experimental/Dataflow/Source/DataflowEditor/Public/DataflowEditorTools/DataflowEditorVertexAttributePaintTool.h
--- a/Engine/Plugins/Experimental/Dataflow/Source/DataflowEditor/Public/DataflowEditorTools/DataflowEditorVertexAttributePaintTool.h
+++ b/Engine/Plugins/Experimental/Dataflow/Source/DataflowEditor/Public/DataflowEditorTools/DataflowEditorVertexAttributePaintTool.h
@@ -338,6 +338,7 @@ public:
 	virtual bool HasAccept() const override { return true; }
 	UE_API virtual bool CanAccept() const override;
 
+	UE_API virtual FInputRayHit BeginHoverSequenceHitTest(const FInputDeviceRay& Ray)  override;
 	UE_API virtual bool OnUpdateHover(const FInputDeviceRay& DevicePos) override;
 
 	UE_API virtual void OnPropertyModified(UObject* PropertySet, FProperty* Property) override;
@@ -493,6 +494,7 @@ protected:
 	UE_API bool UpdateBrushPosition(const FRay& WorldRay);
 
 	bool bPendingPickWeight = false;
+	bool bBrushIsHoverMesh = false;
 
 	TArray<int32> ROITriangleBuffer;
 	TArray<double> ROIWeightValueBuffer;
