Dataflow weight map tool : Add a Query section to the tool propertyies : Value at brush  + Query Type
- Also fixed issue with the ramp menu staying in greyscale and material modes

#Synced-CL 48816559

diff --git a/Engine/Plugins/Experimental/Dataflow/Source/DataflowEditor/Private/DataflowEditorTools/DataflowEditorVertexAttributePaintTool.cpp b/Engine/Plugins/Experimental/Dataflow/Source/DataflowEditor/Private/DataflowEditorTools/DataflowEditorVertexAttributePaintTool.cpp
--- a/Engine/Plugins/Experimental/Dataflow/Source/DataflowEditor/Private/DataflowEditorTools/DataflowEditorVertexAttributePaintTool.cpp
+++ b/Engine/Plugins/Experimental/Dataflow/Source/DataflowEditor/Private/DataflowEditorTools/DataflowEditorVertexAttributePaintTool.cpp
@@ -1986,8 +1986,20 @@ int32 UDataflowEditorVertexAttributePaintTool::FindHitTargetMeshTriangleConst(co
 
 void UDataflowEditorVertexAttributePaintTool::UpdateHitSculptMeshTriangle(int32 TriangleID, const FRay3d& LocalRay)
 {
-	// TODO: Figure out what the actual position on the triangle is when hit.
-	CurrentBaryCentricCoords = FVector3d(1.0 / 3.0, 1.0 / 3.0, 1.0 / 3.0);
+	using namespace UE::Geometry;
+
+	CurrentBaryCentricCoords = FVector3d::Zero();
+
+	const FDynamicMesh3* const Mesh = GetSculptMesh();
+
+	if (Mesh->IsTriangle(TriangleID))
+	{
+		FTriangle3d Triangle;
+		Mesh->GetTriVertices(TriangleID, Triangle.V[0], Triangle.V[1], Triangle.V[2]);
+		FIntrRay3Triangle3d Query(LocalRay, Triangle);
+		Query.Find();
+		CurrentBaryCentricCoords = Query.TriangleBaryCoords;
+	}
 }
 
 bool UDataflowEditorVertexAttributePaintTool::UpdateBrushPosition(const FRay& WorldRay)
@@ -2317,8 +2329,48 @@ void UDataflowEditorVertexAttributePaintTool::PrecomputeFilterData()
 
 float UDataflowEditorVertexAttributePaintTool::GetCurrentWeightValueUnderBrush() const
 {
-	const int32 VertexID = GetBrushNearestVertex();
-	return VertexData.GetValue(VertexID);
+	using namespace UE::Geometry;
+
+	float Value = 0;
+	if (ToolProperties)
+	{
+		switch (ToolProperties->BrushProperties.ValueQueryType)
+		{
+
+		case EDataflowEditorToolValueQueryType::Interpolated:
+		{
+			const int32 Tid = GetBrushTriangleID();
+			if (Tid != IndexConstants::InvalidID)
+			{
+				const FDynamicMesh3* Mesh = GetSculptMesh();
+				const FIndex3i Vertices = Mesh->GetTriangle(Tid);
+				Value = 0.0f;
+				for (int32 TriangleVertexIndex = 0; TriangleVertexIndex < 3; ++TriangleVertexIndex)
+				{
+					const float VertexWeight = VertexData.GetValue(Vertices[TriangleVertexIndex]);
+					Value += CurrentBaryCentricCoords[TriangleVertexIndex] * VertexWeight;
+				}
+			}
+		}
+		break;
+
+		case EDataflowEditorToolValueQueryType::NearestVertexFast:
+		{
+			const int32 VertexID = GetBrushNearestVertex();
+			Value = VertexData.GetValue(VertexID);
+		}
+		break;
+
+		case EDataflowEditorToolValueQueryType::NearestVertexAccurate:
+		{
+			const int32 VertexID = GetBrushNearestVertexAccurate();
+			Value = VertexData.GetValue(VertexID);
+		}
+
+		break;
+		}
+	}
+	return Value;
 }
 
 int32 UDataflowEditorVertexAttributePaintTool::GetBrushNearestVertex() const
@@ -2352,6 +2404,41 @@ int32 UDataflowEditorVertexAttributePaintTool::GetBrushNearestVertex() const
 	return Vertices[TriangleVertex];
 }
 
+
+int32 UDataflowEditorVertexAttributePaintTool::GetBrushNearestVertexAccurate() const
+{
+	using namespace UE::Geometry;
+
+	int32 NearestVertexIndex = INDEX_NONE;
+	const int32 Tid = GetBrushTriangleID();
+
+	if (Tid != IndexConstants::InvalidID)
+	{
+		const FDynamicMesh3* const Mesh = GetSculptMesh();
+
+		FVector3d PointOnSurface(0, 0, 0);
+		const FIndex3i Vertices = Mesh->GetTriangle(Tid);
+		for (int32 TriangleVertexIndex = 0; TriangleVertexIndex < 3; ++TriangleVertexIndex)
+		{
+			PointOnSurface += CurrentBaryCentricCoords[TriangleVertexIndex] * Mesh->GetVertex(Vertices[TriangleVertexIndex]);
+		}
+
+		double MinDist = UE_BIG_NUMBER;
+		for (const int32 VertexIndex : VertexROI)
+		{
+			const FVector3d& VertexPosition = Mesh->GetVertex(VertexIndex);
+			const double CurrDist = FVector3d::Distance(VertexPosition, PointOnSurface);
+			if (CurrDist < MinDist)
+			{
+				MinDist = CurrDist;
+				NearestVertexIndex = VertexIndex;
+			}
+		}
+	}
+
+	return NearestVertexIndex;
+}
+
 void UDataflowEditorVertexAttributePaintTool::UpdateBrushType(EDataflowEditorToolEditOperation BrushMode)
 {
 	static const FText BaseMessage = LOCTEXT("OnStartTool", "Hold [Shift] to Relax values, Hold [Ctrl] to Invert Value, [/] and S/D change Size (+Shift to small-step)");
diff --git a/Engine/Plugins/Experimental/Dataflow/Source/DataflowEditor/Private/DataflowEditorTools/DataflowEditorVertexAttributePaintToolDetailCustomization.cpp b/Engine/Plugins/Experimental/Dataflow/Source/DataflowEditor/Private/DataflowEditorTools/DataflowEditorVertexAttributePaintToolDetailCustomization.cpp
--- a/Engine/Plugins/Experimental/Dataflow/Source/DataflowEditor/Private/DataflowEditorTools/DataflowEditorVertexAttributePaintToolDetailCustomization.cpp
+++ b/Engine/Plugins/Experimental/Dataflow/Source/DataflowEditor/Private/DataflowEditorTools/DataflowEditorVertexAttributePaintToolDetailCustomization.cpp
@@ -15,6 +15,7 @@
 #include "Widgets/Input/SSegmentedControl.h"
 #include "Widgets/Input/SSpinBox.h"
 #include "Widgets/Input/SNumericEntryBox.h"
+#include "Widgets/Input/SComboBox.h"
 #include "Framework/MultiBox/MultiBoxBuilder.h"
 #include "Drawing/MeshElementsVisualizer.h"
 
@@ -28,6 +29,27 @@ namespace UE::Dataflow::Editor
 	float FVertexAttributePaintToolDetailCustomization::WeightEditVerticalPadding = 4.0f;
 	float FVertexAttributePaintToolDetailCustomization::WeightEditHorizontalPadding = 2.0f;
 
+	static const FName QueryTypeInterpolated(TEXT("Interpolated"));
+	static const FName QueryTypeVertexFast(TEXT("Nearest Vertex Fast"));
+	static const FName QueryTypeVertexAccurate(TEXT("Nearest Vertex Accurate"));
+
+	///////////////////////////////////////////////////////////////////////////////////////////////////////////////
+	FVertexAttributePaintToolDetailCustomization::FVertexAttributePaintToolDetailCustomization()
+	{
+		BrushQueryTypeNames.Add(QueryTypeInterpolated);
+		BrushQueryTypeNames.Add(QueryTypeVertexFast);
+		BrushQueryTypeNames.Add(QueryTypeVertexAccurate);
+
+		BrushQueryTypeEnumFromNames.Add(QueryTypeInterpolated, EDataflowEditorToolValueQueryType::Interpolated);
+		BrushQueryTypeEnumFromNames.Add(QueryTypeVertexFast, EDataflowEditorToolValueQueryType::NearestVertexFast);
+		BrushQueryTypeEnumFromNames.Add(QueryTypeVertexAccurate, EDataflowEditorToolValueQueryType::NearestVertexAccurate);
+
+		BrushQueryTypeNameFromEnum.Add(EDataflowEditorToolValueQueryType::Interpolated, QueryTypeInterpolated);
+		BrushQueryTypeNameFromEnum.Add(EDataflowEditorToolValueQueryType::NearestVertexFast, QueryTypeVertexFast);
+		BrushQueryTypeNameFromEnum.Add(EDataflowEditorToolValueQueryType::NearestVertexAccurate, QueryTypeVertexAccurate);
+
+	}
+
 	///////////////////////////////////////////////////////////////////////////////////////////////////////////////
 
 	TSharedRef<IDetailCustomization> FVertexAttributePaintToolDetailCustomization::MakeInstance()
@@ -331,13 +353,19 @@ namespace UE::Dataflow::Editor
 				return GradientEditor;
 			};
 
-		auto MakeColorRampOptionMenuWidget = [this]()
+		auto MakeColorRampOptionMenuWidget = [this](TWeakObjectPtr<UDataflowEditorVertexAttributePaintToolProperties> WeakToolProperties)
 		{
 			return SNew(SComboButton)
 				.ComboButtonStyle(&FAppStyle::Get().GetWidgetStyle<FComboButtonStyle>("SimpleComboButtonWithIcon"))
 				.ForegroundColor(FSlateColor::UseStyle())
 				.ToolTipText(LOCTEXT("ColorRampOptionMenuToolTip", "Color Ramp Options Menu"))
-				.OnGetMenuContent(this, &FVertexAttributePaintToolDetailCustomization::BuildColorRampMenu);
+				.OnGetMenuContent(this, &FVertexAttributePaintToolDetailCustomization::BuildColorRampMenu)
+				.Visibility_Lambda([WeakToolProperties]()
+					{
+						return (WeakToolProperties.IsValid() && (WeakToolProperties->DisplayProperties.ColorMode == EDataflowEditorToolColorMode::Ramp))
+							? EVisibility::Visible
+							: EVisibility::Collapsed;
+					});
 		};
 
 		if (ToolProperties.IsValid())
@@ -348,7 +376,7 @@ namespace UE::Dataflow::Editor
 			TSharedRef<SColorGradientEditor> CustomRampWidget =
 				MakeColorRampWidget(EDataflowEditorToolColorMode::Ramp, &ToolProperties->DisplayProperties.ColorRamp, ToolProperties);
 
-			TSharedRef<SComboButton> CustomRampWidgetOptionMenu = MakeColorRampOptionMenuWidget();
+			TSharedRef<SComboButton> CustomRampWidgetOptionMenu = MakeColorRampOptionMenuWidget(ToolProperties);
 
 			TSharedRef<SColorGradientEditor> MaterialRampWidget =
 				MakeColorRampWidget(EDataflowEditorToolColorMode::FullMaterial, &ToolProperties->DisplayProperties.WhiteColorRamp, ToolProperties);
@@ -420,6 +448,17 @@ namespace UE::Dataflow::Editor
 			}
 		}
 
+		// MeQuery category
+		IDetailCategoryBuilder& QueryCategory = DetailBuilder.EditCategory("Query", FText::GetEmpty(), ECategoryPriority::Important);
+		QueryCategory.InitiallyCollapsed(false);
+		{
+			// Value at brush
+			AddBrushValueAtBrushRow(QueryCategory);
+
+			// Value at brush query method
+			AddBrushValueQueryMethodRow(QueryCategory);
+		}
+
 		// Mesh Display category
 		IDetailCategoryBuilder& MeshDisplayCategory = DetailBuilder.EditCategory("MeshDisplay", FText::GetEmpty(), ECategoryPriority::Important);
 		MeshDisplayCategory.InitiallyCollapsed(false);
@@ -639,6 +678,75 @@ namespace UE::Dataflow::Editor
 
 	///////////////////////////////////////////////////////////////////////////////////////////////////////////////
 
+	void FVertexAttributePaintToolDetailCustomization::AddBrushValueAtBrushRow(IDetailCategoryBuilder& BrushCategory) const
+	{
+		BrushCategory.AddCustomRow(LOCTEXT("ValueAtBrushCategory", "ValueAtBrush"), false)
+			.NameContent()
+			[
+				SNew(STextBlock)
+					.Text(LOCTEXT("ValueAtBrushLabel", "Value At Brush"))
+					.Font(FAppStyle::GetFontStyle(TEXT("PropertyWindow.NormalFont")))
+					.ToolTipText(LOCTEXT("ValueAtBrushTooltip", "The value where the brush is loclated. The exact method used to compute it can be changed in the combobox below"))
+			]
+			.ValueContent()
+			[
+				SNew(STextBlock)
+					.Justification(ETextJustify::Left)
+					.IsEnabled(false)
+					.Text_Lambda([this]()
+						{
+							return FText::AsNumber(ToolProperties.IsValid() ? ToolProperties->BrushProperties.ValueAtBrush: 1.f);
+						})
+			];
+	}
+
+	void FVertexAttributePaintToolDetailCustomization::AddBrushValueQueryMethodRow(IDetailCategoryBuilder& BrushCategory) const
+	{
+		BrushCategory.AddCustomRow(LOCTEXT("ValueQueryMethodBrushCategory", "ValueAtBrush"), false)
+			.NameContent()
+			[
+				SNew(STextBlock)
+					.Text(LOCTEXT("ValueQueryMethodBrushLabel", "Value At Brush Query Type"))
+					.Font(FAppStyle::GetFontStyle(TEXT("PropertyWindow.NormalFont")))
+					.ToolTipText(LOCTEXT("ValueQueryMethodBrushTooltip", "The method used to query the value under the brush"))
+			]
+			.ValueContent()
+			[
+				SNew(SComboBox<FName>)
+					.ComboBoxStyle(&FAppStyle::Get().GetWidgetStyle<FComboBoxStyle>(TEXT("ComboBox")))
+					.OptionsSource(&BrushQueryTypeNames)
+					.HasDownArrow(false)
+					.InitiallySelectedItem(QueryTypeVertexFast)
+					.ContentPadding(0.f)
+					.OnGenerateWidget_Lambda([this](FName Item)
+						{
+							return SNew(STextBlock)
+								.Margin(FMargin(2,1))
+								.Justification(ETextJustify::Left)
+								.Text(FText::FromString(Item.ToString()));
+						})
+					.OnSelectionChanged_Lambda([this](const FName InSelection, ESelectInfo::Type InSelectInfo)
+						{
+							if (ToolProperties.IsValid())
+							{
+								ToolProperties->BrushProperties.ValueQueryType = BrushQueryTypeEnumFromNames.FindRef(InSelection, EDataflowEditorToolValueQueryType::NearestVertexFast);
+							}
+						})
+					.Content()
+					[
+						SNew(STextBlock)
+							.Margin(FMargin(2, 1))
+							.Text_Lambda([this]()
+								{
+									const FName CurrentItemName = ToolProperties.IsValid()
+										? BrushQueryTypeNameFromEnum.FindRef(ToolProperties->BrushProperties.ValueQueryType, QueryTypeVertexFast)
+										: QueryTypeVertexFast;
+									return FText::FromName(CurrentItemName);
+								})
+					]
+			];
+	}
+
 	void FVertexAttributePaintToolDetailCustomization::AddBrushHitBackFacesRow(IDetailCategoryBuilder& BrushCategory) const
 	{
 		BrushCategory.AddCustomRow(LOCTEXT("BrushHitBackFacesCategory", "Brush Falloff"), false)
diff --git a/Engine/Plugins/Experimental/Dataflow/Source/DataflowEditor/Private/DataflowEditorTools/DataflowEditorVertexAttributePaintToolDetailCustomization.h b/Engine/Plugins/Experimental/Dataflow/Source/DataflowEditor/Private/DataflowEditorTools/DataflowEditorVertexAttributePaintToolDetailCustomization.h
--- a/Engine/Plugins/Experimental/Dataflow/Source/DataflowEditor/Private/DataflowEditorTools/DataflowEditorVertexAttributePaintToolDetailCustomization.h
+++ b/Engine/Plugins/Experimental/Dataflow/Source/DataflowEditor/Private/DataflowEditorTools/DataflowEditorVertexAttributePaintToolDetailCustomization.h
@@ -13,12 +13,14 @@ class SWidget;
 class UDataflowEditorVertexAttributePaintTool;
 class UDataflowEditorVertexAttributePaintToolProperties;
 
+enum class EDataflowEditorToolValueQueryType : uint8;
+
 namespace UE::Dataflow::Editor
 {
 	class FVertexAttributePaintToolDetailCustomization : public IDetailCustomization
 	{
 	public:
-
+		FVertexAttributePaintToolDetailCustomization();
 		virtual ~FVertexAttributePaintToolDetailCustomization() override;
 
 		static TSharedRef<IDetailCustomization> MakeInstance();
@@ -41,6 +43,8 @@ namespace UE::Dataflow::Editor
 		void AddBrushRadiusRow(IDetailCategoryBuilder& BrushCategory) const;
 		void AddBrushValueRow(IDetailCategoryBuilder& BrushCategory) const;
 		void AddBrushFalloffRow(IDetailCategoryBuilder& BrushCategory) const;
+		void AddBrushValueAtBrushRow(IDetailCategoryBuilder& BrushCategory) const;
+		void AddBrushValueQueryMethodRow(IDetailCategoryBuilder& BrushCategory) const;
 		void AddBrushHitBackFacesRow(IDetailCategoryBuilder& BrushCategory) const;
 		
 
@@ -66,6 +70,10 @@ namespace UE::Dataflow::Editor
 		TWeakObjectPtr<UDataflowEditorVertexAttributePaintToolProperties> ToolProperties;
 		TWeakObjectPtr<UDataflowEditorVertexAttributePaintTool> Tool;
 
+		TArray<FName> BrushQueryTypeNames;
+		TMap<FName, EDataflowEditorToolValueQueryType> BrushQueryTypeEnumFromNames;
+		TMap<EDataflowEditorToolValueQueryType, FName> BrushQueryTypeNameFromEnum;
+
 		static float WeightSliderWidths;
 		static float WeightEditingLabelsPercent;
 		static float WeightEditVerticalPadding;
diff --git a/Engine/Plugins/Experimental/Dataflow/Source/DataflowEditor/Public/DataflowEditorTools/DataflowEditorToolEnums.h b/Engine/Plugins/Experimental/Dataflow/Source/DataflowEditor/Public/DataflowEditorTools/DataflowEditorToolEnums.h
--- a/Engine/Plugins/Experimental/Dataflow/Source/DataflowEditor/Public/DataflowEditorTools/DataflowEditorToolEnums.h
+++ b/Engine/Plugins/Experimental/Dataflow/Source/DataflowEditor/Public/DataflowEditorTools/DataflowEditorToolEnums.h
@@ -50,3 +50,17 @@ enum class EDataflowEditorToolMirrorDirection : uint8
 	PositiveToNegative,
 	NegativeToPositive,
 };
+
+/** Value Query mode  */
+UENUM()
+enum class EDataflowEditorToolValueQueryType : uint8
+{
+	/** Value is interpolated from triangle vertices */
+	Interpolated UMETA(DisplayName = "Interpolated"),
+
+	/** Return the value at the closest vertex of the triangle under the mouse cursor */
+	NearestVertexFast UMETA(DisplayName = "Nearest Vertex (Fast)"),
+
+	/** Return the value of the nearest vertex inside the brush radius, even if the vertex is not on the triangle under the mouse cursor */
+	NearestVertexAccurate UMETA(DisplayName = "Nearest Vertex (Accurate)")
+};
\ No newline at end of file
diff --git a/Engine/Plugins/Experimental/Dataflow/Source/DataflowEditor/Public/DataflowEditorTools/DataflowEditorVertexAttributePaintTool.h b/Engine/Plugins/Experimental/Dataflow/Source/DataflowEditor/Public/DataflowEditorTools/DataflowEditorVertexAttributePaintTool.h
--- a/Engine/Plugins/Experimental/Dataflow/Source/DataflowEditor/Public/DataflowEditorTools/DataflowEditorVertexAttributePaintTool.h
+++ b/Engine/Plugins/Experimental/Dataflow/Source/DataflowEditor/Public/DataflowEditorTools/DataflowEditorVertexAttributePaintTool.h
@@ -120,6 +120,10 @@ struct FDataflowEditorVertexAttributePaintToolBrushProperties
 	/** The weight value at the brush indicator */
 	UPROPERTY(VisibleAnywhere, Transient, Category = Query, meta = (NoResetToDefault))
 	double ValueAtBrush = 0;
+
+	/** method used for querying the value at brush */
+	UPROPERTY(EditAnywhere, Category = Filters)
+	EDataflowEditorToolValueQueryType ValueQueryType = EDataflowEditorToolValueQueryType::NearestVertexFast;
 };
 
 USTRUCT()
@@ -472,6 +476,7 @@ protected:
 	UE_API float GetCurrentWeightValueUnderBrush() const;
 	FVector3d CurrentBaryCentricCoords;
 	UE_API int32 GetBrushNearestVertex() const;
+	UE_API int32 GetBrushNearestVertexAccurate() const;
 
 	TSet<int32> AccumulatedTriangleROI;
 	bool bUndoUpdatePending = false;
