Dataflow : when embedding a dataflow, make sure to mark the host asset dirty

#Synced-CL 48897267

diff --git a/Engine/Plugins/Experimental/Dataflow/Source/DataflowEditor/Private/Dataflow/DataflowEditorToolkit.cpp b/Engine/Plugins/Experimental/Dataflow/Source/DataflowEditor/Private/Dataflow/DataflowEditorToolkit.cpp
--- a/Engine/Plugins/Experimental/Dataflow/Source/DataflowEditor/Private/Dataflow/DataflowEditorToolkit.cpp
+++ b/Engine/Plugins/Experimental/Dataflow/Source/DataflowEditor/Private/Dataflow/DataflowEditorToolkit.cpp
@@ -1189,6 +1189,7 @@ void FDataflowEditorToolkit::ConvertToEmbeddedDataflow()
 			UDataflow* NewDataflowAsset = nullptr;
 			if (UE::DataflowAssetDefinitionHelpers::EmbedDataflowAsset(BoundAsset, NewDataflowAsset))
 			{
+				BoundAsset->Modify();
 				ResetEditor();
 			}
 		}
