Dataflow : Add falloff support for vertex paint tools

#Synced-CL 49009211

diff --git a/Engine/Plugins/Experimental/Dataflow/Source/DataflowEditor/Private/DataflowEditorTools/DataflowEditorVertexAttributePaintTool.cpp b/Engine/Plugins/Experimental/Dataflow/Source/DataflowEditor/Private/DataflowEditorTools/DataflowEditorVertexAttributePaintTool.cpp
--- a/Engine/Plugins/Experimental/Dataflow/Source/DataflowEditor/Private/DataflowEditorTools/DataflowEditorVertexAttributePaintTool.cpp
+++ b/Engine/Plugins/Experimental/Dataflow/Source/DataflowEditor/Private/DataflowEditorTools/DataflowEditorVertexAttributePaintTool.cpp
@@ -791,7 +791,7 @@ void UDataflowEditorVertexAttributePaintTool::Setup()
 
 	SetBrushMode(ToolProperties->BrushProperties.BrushMode);
 
-	//SetPrimaryFalloffType(EMeshSculptFalloffType::Smooth);
+	SetPrimaryFalloffType(EMeshSculptFalloffType::Smooth);
 
 }
 
@@ -977,12 +977,14 @@ void UDataflowEditorVertexAttributePaintTool::SetBrushMode(EDataflowEditorToolEd
 		FDataflowEditorVertexAttributePaintToolBrushConfig BrushConfigToSave;
 		BrushConfigToSave.BrushSize = GetBrushAdaptiveSize();
 		BrushConfigToSave.Value = ToolProperties->BrushProperties.AttributeValue;
+		BrushConfigToSave.Falloff = GetBrushFalloff();
 		ToolProperties->SetBrushConfig(CurrentBrushMode, BrushConfigToSave);
 
 		// Load the new brush config 
 		const FDataflowEditorVertexAttributePaintToolBrushConfig& NewBrushConfig = ToolProperties->GetBrushConfig(NewBrushMode);
 		SetBrushAdaptiveSize(NewBrushConfig.BrushSize);
 		ToolProperties->BrushProperties.AttributeValue = NewBrushConfig.Value;
+		SetBrushFalloff(NewBrushConfig.Falloff);
 
 		// Send notification for the changed properties
 		FPropertyChangedEvent BrushModePropertyChangedEvent(UBrushBaseProperties::StaticClass()->FindPropertyByName(GET_MEMBER_NAME_CHECKED(FDataflowEditorVertexAttributePaintToolBrushProperties, BrushMode)));
@@ -1023,6 +1025,7 @@ void UDataflowEditorVertexAttributePaintTool::OnBeginStroke(const FRay& WorldRay
 		}
 		PaintBrushOpProperties->EditOperation = ToolProperties->BrushProperties.BrushMode;
 		PaintBrushOpProperties->Strength = 1.0f; // Attribute painting only need the attribute value 
+		PaintBrushOpProperties->SetFalloff(UMeshSculptToolBase::BrushProperties->BrushFalloffAmount);
 	}
 
 	// initialize first "Last Stamp", so that we can assume all stamps in stroke have a valid previous stamp
@@ -1222,6 +1225,7 @@ bool UDataflowEditorVertexAttributePaintTool::UpdateStampPosition(const FRay& Wo
 	CurrentStamp.WorldFrame = GetBrushFrameWorld();
 	CurrentStamp.LocalFrame = GetBrushFrameLocal();
 	CurrentStamp.Power = GetActivePressure() * GetCurrentBrushStrength();
+	CurrentStamp.Falloff = GetCurrentBrushFalloff();
 
 	CurrentStamp.PrevLocalFrame = LastStamp.LocalFrame;
 	CurrentStamp.PrevWorldFrame = LastStamp.WorldFrame;
diff --git a/Engine/Plugins/Experimental/Dataflow/Source/DataflowEditor/Private/DataflowEditorTools/DataflowEditorVertexAttributePaintToolDetailCustomization.cpp b/Engine/Plugins/Experimental/Dataflow/Source/DataflowEditor/Private/DataflowEditorTools/DataflowEditorVertexAttributePaintToolDetailCustomization.cpp
--- a/Engine/Plugins/Experimental/Dataflow/Source/DataflowEditor/Private/DataflowEditorTools/DataflowEditorVertexAttributePaintToolDetailCustomization.cpp
+++ b/Engine/Plugins/Experimental/Dataflow/Source/DataflowEditor/Private/DataflowEditorTools/DataflowEditorVertexAttributePaintToolDetailCustomization.cpp
@@ -794,7 +794,7 @@ namespace UE::Dataflow::Editor
 			AddBrushValueRow(BrushCategory);
 
 			// Brush value field
-			//AddBrushFalloffRow(BrushCategory);
+			AddBrushFalloffRow(BrushCategory);
 
 			// Brush backface hit option
 			AddBrushHitBackFacesRow(BrushCategory);
diff --git a/Engine/Plugins/Experimental/Dataflow/Source/DataflowEditor/Private/DataflowEditorTools/DataflowEditorWeightMapPaintBrushOps.cpp b/Engine/Plugins/Experimental/Dataflow/Source/DataflowEditor/Private/DataflowEditorTools/DataflowEditorWeightMapPaintBrushOps.cpp
--- a/Engine/Plugins/Experimental/Dataflow/Source/DataflowEditor/Private/DataflowEditorTools/DataflowEditorWeightMapPaintBrushOps.cpp
+++ b/Engine/Plugins/Experimental/Dataflow/Source/DataflowEditor/Private/DataflowEditorTools/DataflowEditorWeightMapPaintBrushOps.cpp
@@ -86,7 +86,7 @@ void FDataflowWeightMapSmoothBrushOp::ApplyStampByVertices(
 	const TArray<int32>& Vertices,
 	TArray<double>& VertexWeightValues)
 {
-	ApplyStampByVerticesStatic(Mesh, Stamp, Vertices, VertexWeightValues, bApplyRadiusLimit);
+	ApplyStampByVerticesStatic(Mesh, Stamp, Vertices, VertexWeightValues, bApplyRadiusLimit, Falloff);
 }
 
 
@@ -95,9 +95,18 @@ void FDataflowWeightMapSmoothBrushOp::ApplyStampByVertices(
 	const FSculptBrushStamp& Stamp,
 	const TArray<int32>& Vertices,
 	TArray<double>& VertexWeightValues,
-	bool bApplyRadiusLimit)
+	bool bApplyRadiusLimit,
+	TSharedPtr<FMeshSculptFallofFunc> FalloffFunc)
 {
-	FDataflowWeightMapSmoothBrushOp::ApplyToVerticesStatic(Mesh, Vertices, VertexWeightValues, Stamp.Power, bApplyRadiusLimit, Stamp.LocalFrame.Origin, Stamp.Radius);
+	auto FalloffFunction = [&Stamp, &FalloffFunc](const FVector& Position)
+		{
+			if (FalloffFunc)
+			{
+				return (float)FalloffFunc->Evaluate(Stamp, Position);
+			}
+			return 1.0f;
+		};
+	FDataflowWeightMapSmoothBrushOp::ApplyToVerticesStatic(Mesh, Vertices, VertexWeightValues, Stamp.Power, bApplyRadiusLimit, Stamp.LocalFrame.Origin, Stamp.Radius, FalloffFunction);
 }
 
 /*static*/ void FDataflowWeightMapSmoothBrushOp::ApplyToVerticesStatic(
@@ -107,7 +116,8 @@ void FDataflowWeightMapSmoothBrushOp::ApplyStampByVertices(
 	const double SmoothingFactor,
 	bool bApplyRadiusLimit,
 	const FVector Center,
-	const float Radius
+	const float Radius,
+	FalloffFunction FalloffFunc
 )
 {
 	// Converted from FClothPaintTool_Smooth::SmoothVertices
@@ -153,9 +163,9 @@ void FDataflowWeightMapSmoothBrushOp::ApplyStampByVertices(
 
 	// Blend vertex value with its average one-ring value
 	ParallelFor(NumVertices,
-		[&OneRingAverages, &VertexWeightValues, bApplyRadiusLimit, &Center, &Radius, &SmoothingFactor, &Vertices, &Mesh](int32 BufferIndex)
+		[&OneRingAverages, &VertexWeightValues, bApplyRadiusLimit, &Center, &Radius, &SmoothingFactor, &Vertices, &Mesh, &FalloffFunc](int32 BufferIndex)
 		{
-			const double Diff = OneRingAverages[BufferIndex] - VertexWeightValues[BufferIndex];
+			float FalloffScalar = 1.0;
 			if (bApplyRadiusLimit)
 			{
 				const FVector3d& StampPos = Center;
@@ -166,8 +176,13 @@ void FDataflowWeightMapSmoothBrushOp::ApplyStampByVertices(
 				{
 					return;
 				}
+				if (FalloffFunc.IsSet())
+				{
+					FalloffScalar = FalloffFunc(VertexPos);
+				}
 			}
-			VertexWeightValues[BufferIndex] += SmoothingFactor * Diff;
+			const float NewValue = FMath::Lerp(VertexWeightValues[BufferIndex], OneRingAverages[BufferIndex], SmoothingFactor * FalloffScalar);
+			VertexWeightValues[BufferIndex] = FMath::Clamp(NewValue, 0.f, 1.f);
 		}, 
 		EParallelForFlags::None
 	);
@@ -185,7 +200,7 @@ void FDataflowVertexAttributePaintBrushOp::ApplyStampByVertices(
 
 	if (UDataflowVertexAttributePaintBrushOpProps* Props = GetPropertySetAs<UDataflowVertexAttributePaintBrushOpProps>())
 	{
-		ApplyStampByVerticesStatic(Mesh, Stamp, Vertices, OutValues, Props->EditOperation, Props->GetAttribute(), bApplyRadiusLimit);
+		ApplyStampByVerticesStatic(Mesh, Stamp, Vertices, OutValues, Props->EditOperation, Props->GetAttribute(), bApplyRadiusLimit, Falloff);
 	}
 }
 
@@ -196,9 +211,18 @@ void FDataflowVertexAttributePaintBrushOp::ApplyStampByVertices(
 	TArray<double>& OutValues,
 	EDataflowEditorToolEditOperation EditOperation,
 	const double ToolValue,
-	bool bApplyRadiusLimit)
+	bool bApplyRadiusLimit,
+	TSharedPtr<FMeshSculptFallofFunc> FalloffFunc)
 {
-	FDataflowVertexAttributePaintBrushOp::ApplyToVerticesStatic(Mesh, Vertices, OutValues, EditOperation, ToolValue, bApplyRadiusLimit, Stamp.LocalFrame.Origin, Stamp.Radius);
+	auto FalloffFunction = [&Stamp, &FalloffFunc](const FVector& Position)
+		{
+			if (FalloffFunc)
+			{
+				return (float)FalloffFunc->Evaluate(Stamp, Position);
+			}
+			return 1.0f;
+		};
+	FDataflowVertexAttributePaintBrushOp::ApplyToVerticesStatic(Mesh, Vertices, OutValues, EditOperation, ToolValue, bApplyRadiusLimit, Stamp.LocalFrame.Origin, Stamp.Radius, FalloffFunction);
 }
 
 /*static*/ void FDataflowVertexAttributePaintBrushOp::ApplyToVerticesStatic(
@@ -209,7 +233,8 @@ void FDataflowVertexAttributePaintBrushOp::ApplyStampByVertices(
 	const double ToolValue,
 	bool bApplyRadiusLimit,
 	const FVector Center,
-	const float Radius
+	const float Radius,
+	FalloffFunction FalloffFunc
 )
 {
 	// special case for Relax option 
@@ -221,9 +246,10 @@ void FDataflowVertexAttributePaintBrushOp::ApplyStampByVertices(
 
 	const int32 NumVertices = Vertices.Num();
 	ParallelFor(NumVertices,
-		[bApplyRadiusLimit, ToolValue, EditOperation, &Center, &Radius, &Mesh, &Vertices, &OutValues]
+		[bApplyRadiusLimit, ToolValue, EditOperation, &Center, &Radius, &Mesh, &Vertices, &OutValues, &FalloffFunc]
 		(int32 BufferIndex)
 		{
+			float FalloffScalar = 1.0;
 			if (bApplyRadiusLimit)
 			{
 				const FVector3d& StampPos = Center;
@@ -234,6 +260,10 @@ void FDataflowVertexAttributePaintBrushOp::ApplyStampByVertices(
 				{
 					return;
 				}
+				if (FalloffFunc.IsSet())
+				{
+					FalloffScalar = FalloffFunc(VertexPos);
+				}
 			}
 
 			const double CurrentValue = OutValues[BufferIndex];
@@ -251,12 +281,13 @@ void FDataflowVertexAttributePaintBrushOp::ApplyStampByVertices(
 				NewValue = CurrentValue * ToolValue;
 				break;
 			case EDataflowEditorToolEditOperation::Invert:
-				NewValue = 1.0 - CurrentValue;
+				NewValue = 1.0 - CurrentValue; // no falloff
 				break;
 			case EDataflowEditorToolEditOperation::Relax:
 				ensure(false); // already handled by SmoothBrush.ApplyStampByVertices
 			}
 
+			NewValue = FMath::Lerp(OutValues[BufferIndex], NewValue, FalloffScalar);
 			OutValues[BufferIndex] = FMath::Clamp(NewValue, 0.0, 1.0);
 		},
 		EParallelForFlags::None
diff --git a/Engine/Plugins/Experimental/Dataflow/Source/DataflowEditor/Private/DataflowEditorTools/DataflowEditorWeightMapPaintBrushOps.h b/Engine/Plugins/Experimental/Dataflow/Source/DataflowEditor/Private/DataflowEditorTools/DataflowEditorWeightMapPaintBrushOps.h
--- a/Engine/Plugins/Experimental/Dataflow/Source/DataflowEditor/Private/DataflowEditorTools/DataflowEditorWeightMapPaintBrushOps.h
+++ b/Engine/Plugins/Experimental/Dataflow/Source/DataflowEditor/Private/DataflowEditorTools/DataflowEditorWeightMapPaintBrushOps.h
@@ -13,6 +13,9 @@ class FDataflowMeshVertexWeightMapEditBrushOp : public FMeshSculptBrushOp
 {
 public:
 
+	// falloff as a function of position 
+	using FalloffFunction = TFunction<float(const FVector&)>;
+
 	// not supported for this kind of brush op
 	virtual void ApplyStamp(
 		const FDynamicMesh3* Mesh,
@@ -90,12 +93,16 @@ public:
 	UPROPERTY(EditAnywhere, Category = PaintBrush, meta = (DisplayName = "Strength", UIMin = "0.0", UIMax = "10.", ClampMin = "0.0", ClampMax = "10."))
 	float Strength = 0.5;
 
+	UPROPERTY(EditAnywhere, Category = PaintBrush, meta = (UIMin = "0.0", UIMax = "1.", ClampMin = "0.0", ClampMax = "1."))
+	float Falloff = 0;
+
 	virtual double GetAttribute() { return AttributeValue; }
 
 	virtual float GetStrength() override { return Strength; }
 	virtual void SetStrength(float NewStrength) override { Strength = NewStrength;  }
 
-	virtual float GetFalloff() override { return 0.0f; }
+	virtual void SetFalloff(float NewFalloff) { Falloff = FMath::Clamp(NewFalloff, 0.f, 1.f); }
+	virtual float GetFalloff() override { return Falloff; }
 };
 
 
@@ -158,7 +165,8 @@ public:
 		const FSculptBrushStamp& Stamp,
 		const TArray<int32>& Vertices,
 		TArray<double>& NewAttributesOut,
-		bool bApplyRadiusLimit);
+		bool bApplyRadiusLimit,
+		TSharedPtr<FMeshSculptFallofFunc> FalloffFunc = {});
 
 	static void ApplyToVerticesStatic(
 		const FDynamicMesh3* Mesh,
@@ -167,7 +175,8 @@ public:
 		const double SmoothingFactor,
 		bool bApplyRadiusLimit,
 		const FVector Center,
-		const float Radius);
+		const float Radius, 
+		FalloffFunction FalloffFunc = {});
 };
 
 ////////////////////////////////////////////////////////////////////////////////////////////////////////////////////
@@ -186,10 +195,15 @@ public:
 	UPROPERTY(EditAnywhere, Category = PaintBrush, meta = (DisplayName = "Strength", UIMin = "0.0", UIMax = "10.", ClampMin = "0.0", ClampMax = "10."))
 	float Strength = 1.0;
 
+	UPROPERTY(EditAnywhere, Category = PaintBrush, meta = (UIMin = "0.0", UIMax = "1.", ClampMin = "0.0", ClampMax = "1."))
+	float Falloff = 0;
+
 	virtual double GetAttribute() { return AttributeValue; }
 	virtual float GetStrength() override { return Strength; }
 	virtual void SetStrength(float NewStrength) override { Strength = NewStrength; }
-	virtual float GetFalloff() override { return 0.0f; }
+
+	virtual void SetFalloff(float NewFalloff) { Falloff = FMath::Clamp(NewFalloff, 0.f, 1.f); }
+	virtual float GetFalloff() override { return Falloff; }
 
 	EDataflowEditorToolEditOperation EditOperation = EDataflowEditorToolEditOperation::Add;
 };
@@ -197,7 +211,6 @@ public:
 class FDataflowVertexAttributePaintBrushOp : public FDataflowMeshVertexWeightMapEditBrushOp
 {
 public:
-
 	virtual bool IgnoreZeroMovements() const override
 	{
 		return false;
@@ -216,7 +229,8 @@ public:
 		TArray<double>& NewAttributesOut,
 		EDataflowEditorToolEditOperation EditOperation,
 		const double ToolValue,
-		bool bApplyRadiusLimit);
+		bool bApplyRadiusLimit,
+		TSharedPtr<FMeshSculptFallofFunc> FalloffFunc = {});
 
 	static void ApplyToVerticesStatic(
 		const FDynamicMesh3* Mesh,
@@ -226,5 +240,6 @@ public:
 		const double ToolValue,
 		bool bApplyRadiusLimit,
 		const FVector Center, 
-		const float Radius);
+		const float Radius, 
+		FalloffFunction FalloffFunc = {});
 };
diff --git a/Engine/Plugins/Experimental/Dataflow/Source/DataflowEditor/Public/DataflowEditorTools/DataflowEditorVertexAttributePaintTool.h b/Engine/Plugins/Experimental/Dataflow/Source/DataflowEditor/Public/DataflowEditorTools/DataflowEditorVertexAttributePaintTool.h
--- a/Engine/Plugins/Experimental/Dataflow/Source/DataflowEditor/Public/DataflowEditorTools/DataflowEditorVertexAttributePaintTool.h
+++ b/Engine/Plugins/Experimental/Dataflow/Source/DataflowEditor/Public/DataflowEditorTools/DataflowEditorVertexAttributePaintTool.h
@@ -139,6 +139,9 @@ public:
 	/** Value to paint on the mesh (dependent on the BrushMode) */
 	UPROPERTY(Config)
 	float Value = 1.0;
+
+	UPROPERTY(Config)
+	float Falloff = 0.0;
 };
 
 ///////////////////////////////////////////////////////////////////////////////////////////////////////////////////////////////////////////////////
