Dataflow : add Texture to attribute transfer node
- works in all workflows ( including Cloth )

#Synced-CL 49064124

diff --git a/Engine/Plugins/Experimental/Dataflow/Source/DataflowNodes/DataflowNodes.Build.cs b/Engine/Plugins/Experimental/Dataflow/Source/DataflowNodes/DataflowNodes.Build.cs
--- a/Engine/Plugins/Experimental/Dataflow/Source/DataflowNodes/DataflowNodes.Build.cs
+++ b/Engine/Plugins/Experimental/Dataflow/Source/DataflowNodes/DataflowNodes.Build.cs
@@ -28,7 +28,8 @@ namespace UnrealBuildTool.Rules
 			PrivateDependencyModuleNames.AddRange(
 				new string[]
 				{
-					"GeometryFramework"
+					"GeometryFramework",
+					"ImageCore"
 				}
 			);
 
diff --git a/Engine/Plugins/Experimental/Dataflow/Source/DataflowNodes/Private/Dataflow/DataflowNodesPlugin.cpp b/Engine/Plugins/Experimental/Dataflow/Source/DataflowNodes/Private/Dataflow/DataflowNodesPlugin.cpp
--- a/Engine/Plugins/Experimental/Dataflow/Source/DataflowNodes/Private/Dataflow/DataflowNodesPlugin.cpp
+++ b/Engine/Plugins/Experimental/Dataflow/Source/DataflowNodes/Private/Dataflow/DataflowNodesPlugin.cpp
@@ -14,6 +14,7 @@
 #include "Dataflow/DataflowCollectionEditSkeletonBonesNode.h"
 #include "Dataflow/DataflowCollectionSetSkinningSkeletalMesh.h"
 #include "Dataflow/DataflowSkeletonAssetTerminalNode.h"
+#include "Dataflow/DataflowTextureToAttributeNode.h"
 
 #define LOCTEXT_NAMESPACE "DataflowNodes"
 
@@ -49,6 +50,7 @@ void IDataflowNodesPlugin::StartupModule()
 	UE::Dataflow::RegisterStaticMeshNodes();
 	UE::Dataflow::RegisterSelectionNodes();
 	UE::Dataflow::RegisterContextOverridesNodes();
+	UE::Dataflow::RegisterTextureToAttributeNodes();
 	UE::Dataflow::DataflowCollectionAttributeKeyNodes();
 	DATAFLOW_NODE_REGISTER_CREATION_FACTORY(FDataflowCollectionAddScalarVertexPropertyNode);
 	DATAFLOW_NODE_REGISTER_CREATION_FACTORY(FDataflowCollectionEditSkinWeightsNode);
diff --git a/Engine/Plugins/Experimental/Dataflow/Source/DataflowNodes/Private/Dataflow/DataflowTextureToAttributeNode.cpp b/Engine/Plugins/Experimental/Dataflow/Source/DataflowNodes/Private/Dataflow/DataflowTextureToAttributeNode.cpp
new file mode 100644
--- /dev/null
+++ b/Engine/Plugins/Experimental/Dataflow/Source/DataflowNodes/Private/Dataflow/DataflowTextureToAttributeNode.cpp
@@ -0,0 +1,173 @@
+// Copyright Epic Games, Inc. All Rights Reserved.
+
+#include "Dataflow/DataflowTextureToAttributeNode.h"
+#include "Dataflow/DataflowNodeFactory.h"
+#include "Engine/Texture2D.h"
+#include "GeometryCollection/GeometryCollection.h"
+#include "GeometryCollection/Facades/CollectionMeshFacade.h"
+#include "ImageCore.h"
+
+#include UE_INLINE_GENERATED_CPP_BY_NAME(DataflowTextureToAttributeNode)
+
+#define LOCTEXT_NAMESPACE "DataflowTextureToAttributeNode"
+
+namespace UE::Dataflow
+{
+	void RegisterTextureToAttributeNodes()
+	{
+		DATAFLOW_NODE_REGISTER_CREATION_FACTORY(FDataflowTextureToAttributeNode);
+	}
+
+	namespace Private
+	{
+		static bool GetImageFromTexture2D(UTexture2D& InTexture, FImage& OutImage)
+		{
+			if (const FSharedImageConstRef CpuCopy = InTexture.GetCPUCopy())
+			{
+				CpuCopy->CopyTo(OutImage, ERawImageFormat::RGBA32F, EGammaSpace::Linear);
+				return true;
+			}
+#if WITH_EDITORONLY_DATA
+			if (InTexture.Source.IsValid())
+			{
+				FImage MipImage;
+				if (InTexture.Source.GetMipImage(MipImage, 0))
+				{
+					MipImage.CopyTo(OutImage, ERawImageFormat::RGBA32F, EGammaSpace::Linear);
+					return true;
+				}
+			}
+#endif
+			return false;
+		}
+	}
+};
+
+
+
+FDataflowTextureToAttributeNode::FDataflowTextureToAttributeNode(const UE::Dataflow::FNodeParameters& InParam, FGuid InGuid)
+	: FDataflowNode(InParam, InGuid)
+{
+	RegisterInputConnection(&Collection);
+	RegisterInputConnection(&Texture);
+	RegisterInputConnection(&UVChannel);
+	RegisterInputConnection(&AttributeName);
+	RegisterOutputConnection(&Collection, &Collection);
+}
+
+void FDataflowTextureToAttributeNode::Evaluate(UE::Dataflow::FContext& Context, const FDataflowOutput* Out) const
+{
+	if (Out->IsA(&Collection))
+	{
+		const FName InAttributeName = FName(GetValue(Context, &AttributeName, AttributeName));
+		if (InAttributeName.IsNone())
+		{
+			Context.Error(LOCTEXT("NoAttributeName_Msg", "Attribute name is empty, data transfer from texture will be skipped"));
+			SafeForwardInput(Context, &Collection, &Collection);
+			return;
+		}
+
+		FManagedArrayCollection InCollection = GetValue(Context, &Collection, Collection);
+
+		const FName InVertexGroupName = VertexGroup.Name;
+		if (!InCollection.HasGroup(InVertexGroupName))
+		{
+			Context.Error(LOCTEXT("InvalidVertexGroup_Msg", "Vertex group does not exists, data transfer from texture will be skipped"));
+			SafeForwardInput(Context, &Collection, &Collection);
+			return;
+		}
+
+		const TObjectPtr<UTexture2D> InTexture = GetValue(Context, &Texture, Texture);
+		if (!InTexture)
+		{
+			Context.Warning(LOCTEXT("NoTexture_Msg", "Input texture is null, no data will be transfered and the attribute will be inititialized to zeros"));
+		}
+
+		// Create or get the new attribute 
+		TManagedArray<float>& TargetAttribute = InCollection.AddAttribute<float>(InAttributeName, InVertexGroupName);
+		TargetAttribute.Fill(0);
+
+		// get the Uv attribute ( could be single or array of Uvs depending on the type of collection )
+		const int32 InUVChannel = GetValue(Context, &UVChannel, UVChannel);
+		GeometryCollection::Facades::FCollectionMeshFacade MeshFacade(InCollection);
+		const TManagedArray<FVector2f>* UvAttribute = MeshFacade.FindUVLayer(InUVChannel);
+
+		// cannot really include cloth here so we need to directlky serach for the UV attribute in the render vertices ( eventually we may converge to use only geometry collections )
+		const TManagedArray<TArray<FVector2f>>* UvArrayAttribute = InCollection.FindAttribute<TArray<FVector2f>>("RenderUVs", InVertexGroupName);
+		const int32 UvArrayNumChannels = (UvArrayAttribute && (UvArrayAttribute->Num() > 0)) ? (*UvArrayAttribute)[0].Num() : 0;
+
+		if (!UvAttribute && (InUVChannel<0 || InUVChannel >= UvArrayNumChannels))
+		{
+			Context.Warning(LOCTEXT("InvalidUVChannel_Msg", "Invalid UV channel, no data will be transfered and the attribute will be inititialized to zeros"));
+		}
+		else if (InTexture)
+		{
+			FImage Image;
+			const bool bGotImage = UE::Dataflow::Private::GetImageFromTexture2D(*InTexture, Image);
+			if (!bGotImage || Image.SizeX == 0 || Image.SizeY == 0)
+			{
+				Context.Warning(LOCTEXT("TextureError_Msg", "Failed to get the texture data or empty image, no data will be transfered and the attribute will be inititialized to zeros"));
+			}
+			else
+			{
+				TArrayView64<FLinearColor> Pixels = Image.AsRGBA32F();
+
+				auto GetUVs = [&UvAttribute, &UvArrayAttribute, &InUVChannel](int32 Index)
+					{
+						if (UvAttribute && UvAttribute->IsValidIndex(Index))
+						{
+							return (*UvAttribute)[Index];
+						}
+						else if (UvArrayAttribute && UvArrayAttribute->IsValidIndex(Index))
+						{
+							const TArray<FVector2f>& UVs = (*UvArrayAttribute)[Index];
+							return UVs.IsValidIndex(InUVChannel)? UVs[InUVChannel]: FVector2f::ZeroVector;
+						}
+						return FVector2f::ZeroVector;
+					};
+
+				const int32 MaxX = Image.SizeX - 1;
+				const int32 MaxY = Image.SizeY - 1;
+
+				const int32 NumVertices = TargetAttribute.Num();
+				for (int32 Index = 0; Index < NumVertices; ++Index)
+				{
+					const FVector2f UvCoord = GetUVs(Index);
+					const FVector2f NormUvCoord
+					{ 
+						FMath::Fractional(FMath::Fractional(UvCoord.X) + 1.0f), // [0,1[
+						FMath::Fractional(FMath::Fractional(UvCoord.Y) + 1.0f), // [0,1[
+					};
+					const FVector2f PixelCoord
+					{
+						NormUvCoord.X * (float)(Image.SizeX) - 0.5f,
+						NormUvCoord.Y * (float)(Image.SizeY) - 0.5f,
+					};
+					const FVector2f PixelFraction
+					{
+						FMath::Fractional(PixelCoord.X),
+						FMath::Fractional(PixelCoord.Y),
+					};
+
+					const int32 PixelX0 = FMath::Clamp(int32(PixelCoord.X), 0, MaxX);
+					const int32 PixelX1 = FMath::Clamp(PixelX0 + 1, 0, MaxX);
+					const int32 PixelY0 = FMath::Clamp(int32(PixelCoord.Y), 0, MaxY);
+					const int32 PixelY1 = FMath::Clamp(PixelY0 + 1, 0, MaxY);
+					const FLinearColor Pixel00 = Pixels[PixelX0 + PixelY0 * Image.SizeX];
+					const FLinearColor Pixel10 = Pixels[PixelX1 + PixelY0 * Image.SizeX];
+					const FLinearColor Pixel01 = Pixels[PixelX0 + PixelY1 * Image.SizeX];
+					const FLinearColor Pixel11 = Pixels[PixelX1 + PixelY1 * Image.SizeX];
+
+					const FLinearColor Lerp0 = FMath::Lerp(Pixel00, Pixel10, PixelFraction.X);
+					const FLinearColor Lerp1 = FMath::Lerp(Pixel01, Pixel11, PixelFraction.X);
+					TargetAttribute[Index] = FMath::Lerp(Lerp0, Lerp1, PixelFraction.Y).R;
+				}
+			}
+		}
+
+		SetValue(Context, MoveTemp(InCollection), &Collection);
+	}
+}
+
+#undef LOCTEXT_NAMESPACE 
+
diff --git a/Engine/Plugins/Experimental/Dataflow/Source/DataflowNodes/Private/Dataflow/DataflowTextureToAttributeNode.h b/Engine/Plugins/Experimental/Dataflow/Source/DataflowNodes/Private/Dataflow/DataflowTextureToAttributeNode.h
new file mode 100644
--- /dev/null
+++ b/Engine/Plugins/Experimental/Dataflow/Source/DataflowNodes/Private/Dataflow/DataflowTextureToAttributeNode.h
@@ -0,0 +1,55 @@
+// Copyright Epic Games, Inc. All Rights Reserved.
+
+#pragma once 
+
+#include "Dataflow/DataflowNode.h"
+#include "GeometryCollection/ManagedArrayCollection.h"
+#include "Dataflow/DataflowCollectionAddScalarVertexPropertyNode.h"
+
+#include "DataflowTextureToAttributeNode.generated.h"
+
+class UTexture2D;
+
+/**
+* Transfer data from a texture to a an attribute in a collection
+* The texture is sampled at the specific UVs (UVChannel parameter) and is transfered to an float attribute in the selected vertex group 
+* Sampling is done using bilinear filtering
+*/
+USTRUCT()
+struct FDataflowTextureToAttributeNode : public FDataflowNode
+{
+	GENERATED_USTRUCT_BODY()
+	DATAFLOW_NODE_DEFINE_INTERNAL(FDataflowTextureToAttributeNode, "TextureToAttribute", "Dataflow", "Transfer UVs Projection")
+
+public:
+	FDataflowTextureToAttributeNode(const UE::Dataflow::FNodeParameters& InParam, FGuid InGuid = FGuid::NewGuid());
+
+private:
+	/** Collection to transfer the attribute data to */
+	UPROPERTY(meta = (DataflowInput, DataflowOutput, DataflowPassthrough = Collection))
+	FManagedArrayCollection Collection;
+
+	/** Texture to transfer the attribute from */
+	UPROPERTY(EditAnywhere, Category = Source, meta = (DataflowInput))
+	TObjectPtr<UTexture2D> Texture = nullptr;
+
+	/** UV channel to use to sample the texture */
+	UPROPERTY(EditAnywhere, Category = Source, meta = (DataflowInput))
+	int32 UVChannel = 0;
+
+	/** Name of the attribute to to transfer the texture color to */
+	UPROPERTY(EditAnywhere, Category = Target, meta = (DataflowInput))
+	FString AttributeName;
+
+	/** Vertex group to use */
+	UPROPERTY(EditAnywhere, Category = Target)
+	FScalarVertexPropertyGroup VertexGroup;
+
+	virtual void Evaluate(UE::Dataflow::FContext& Context, const FDataflowOutput* Out) const override;
+};
+
+namespace UE::Dataflow
+{
+	void RegisterTextureToAttributeNodes();
+}
+
