Dataflow editor : Fix Add/Multiplier slider in Paint tool
- Using the slider is now seen as a unique undo transaction
- Change the range of multiply operation to be between 0 and 2 with a default to 1
- Also fix the brush radius not being properly restored when clicking away and back to the tool

#Synced-CL 49114577

diff --git a/Engine/Plugins/Experimental/Dataflow/Source/DataflowEditor/Private/DataflowEditorTools/DataflowEditorVertexAttributePaintTool.cpp b/Engine/Plugins/Experimental/Dataflow/Source/DataflowEditor/Private/DataflowEditorTools/DataflowEditorVertexAttributePaintTool.cpp
--- a/Engine/Plugins/Experimental/Dataflow/Source/DataflowEditor/Private/DataflowEditorTools/DataflowEditorVertexAttributePaintTool.cpp
+++ b/Engine/Plugins/Experimental/Dataflow/Source/DataflowEditor/Private/DataflowEditorTools/DataflowEditorVertexAttributePaintTool.cpp
@@ -789,6 +789,7 @@ void UDataflowEditorVertexAttributePaintTool::Setup()
 	//DynamicMeshComponent->FastNotifyVertexAttributesUpdated(EMeshRenderAttributeFlags::VertexColors);
 	GetToolManager()->PostInvalidation();
 
+	LoadBrushConfig(ToolProperties->BrushProperties.BrushMode);
 	SetBrushMode(ToolProperties->BrushProperties.BrushMode);
 
 	SetPrimaryFalloffType(EMeshSculptFalloffType::Smooth);
@@ -816,8 +817,11 @@ void UDataflowEditorVertexAttributePaintTool::Shutdown(EToolShutdownType Shutdow
 		MeshElementsDisplay->Disconnect();
 	}
 
+	// make sure the tool properties are saved for the next session
 	if (ToolProperties)
 	{
+		SaveBrushConfig(ToolProperties->BrushProperties.BrushMode);
+
 		ToolProperties->DisplayProperties.ColorRamp.OnColorCurveChangedDelegate.RemoveAll(this);
 		ToolProperties->SaveProperties(this);
 	}
@@ -966,6 +970,31 @@ void UDataflowEditorVertexAttributePaintTool::SetBrushFalloff(float NewFalloff)
 	}
 }
 
+void UDataflowEditorVertexAttributePaintTool::SaveBrushConfig(EDataflowEditorToolEditOperation BrushMode)
+{
+	if (ToolProperties)
+	{
+		// Save the current brush config 
+		FDataflowEditorVertexAttributePaintToolBrushConfig BrushConfigToSave;
+		BrushConfigToSave.BrushSize = GetBrushAdaptiveSize();
+		BrushConfigToSave.Value = ToolProperties->BrushProperties.AttributeValue;
+		BrushConfigToSave.Falloff = GetBrushFalloff();
+		ToolProperties->SetBrushConfig(BrushMode, BrushConfigToSave);
+	}
+}
+
+void UDataflowEditorVertexAttributePaintTool::LoadBrushConfig(EDataflowEditorToolEditOperation BrushMode)
+{
+	const FDataflowEditorVertexAttributePaintToolBrushConfig& NewBrushConfig = ToolProperties->GetBrushConfig(BrushMode);
+	SetBrushAdaptiveSize(NewBrushConfig.BrushSize);
+	ToolProperties->BrushProperties.AttributeValue = NewBrushConfig.Value;
+	SetBrushFalloff(NewBrushConfig.Falloff);
+
+	FPropertyChangedEvent ValuePropertyChangedEvent(UBrushBaseProperties::StaticClass()->FindPropertyByName(GET_MEMBER_NAME_CHECKED(FDataflowEditorVertexAttributePaintToolBrushProperties, AttributeValue)));
+	ToolProperties->PostEditChangeProperty(ValuePropertyChangedEvent);
+
+}
+
 void UDataflowEditorVertexAttributePaintTool::SetBrushMode(EDataflowEditorToolEditOperation NewBrushMode)
 {
 	if (ToolProperties && BrushProperties)
@@ -974,23 +1003,14 @@ void UDataflowEditorVertexAttributePaintTool::SetBrushMode(EDataflowEditorToolEd
 		ToolProperties->BrushProperties.BrushMode = NewBrushMode;
 
 		// Save the current brush config 
-		FDataflowEditorVertexAttributePaintToolBrushConfig BrushConfigToSave;
-		BrushConfigToSave.BrushSize = GetBrushAdaptiveSize();
-		BrushConfigToSave.Value = ToolProperties->BrushProperties.AttributeValue;
-		BrushConfigToSave.Falloff = GetBrushFalloff();
-		ToolProperties->SetBrushConfig(CurrentBrushMode, BrushConfigToSave);
+		SaveBrushConfig(CurrentBrushMode);
 
 		// Load the new brush config 
-		const FDataflowEditorVertexAttributePaintToolBrushConfig& NewBrushConfig = ToolProperties->GetBrushConfig(NewBrushMode);
-		SetBrushAdaptiveSize(NewBrushConfig.BrushSize);
-		ToolProperties->BrushProperties.AttributeValue = NewBrushConfig.Value;
-		SetBrushFalloff(NewBrushConfig.Falloff);
+		LoadBrushConfig(NewBrushMode);
 
 		// Send notification for the changed properties
 		FPropertyChangedEvent BrushModePropertyChangedEvent(UBrushBaseProperties::StaticClass()->FindPropertyByName(GET_MEMBER_NAME_CHECKED(FDataflowEditorVertexAttributePaintToolBrushProperties, BrushMode)));
 		ToolProperties->PostEditChangeProperty(BrushModePropertyChangedEvent);
-		FPropertyChangedEvent ValuePropertyChangedEvent(UBrushBaseProperties::StaticClass()->FindPropertyByName(GET_MEMBER_NAME_CHECKED(FDataflowEditorVertexAttributePaintToolBrushProperties, AttributeValue)));
-		ToolProperties->PostEditChangeProperty(ValuePropertyChangedEvent);
 
 		ToolProperties->SaveConfig();
 		BrushProperties->SaveConfig();
@@ -1669,7 +1689,7 @@ TArray<int32> UDataflowEditorVertexAttributePaintTool::GetSelectedVertices() con
 	return {};
 }
 
-void UDataflowEditorVertexAttributePaintTool::ApplyValueToSelection(EDataflowEditorToolEditOperation Operation, float InValue)
+void UDataflowEditorVertexAttributePaintTool::ApplyValueToSelection(EDataflowEditorToolEditOperation Operation, float InValue, bool bWithTransaction)
 {
 	if (!VertexData.IsValid())
 	{
@@ -1682,8 +1702,10 @@ void UDataflowEditorVertexAttributePaintTool::ApplyValueToSelection(EDataflowEdi
 		return;
 	}
 
-	BeginChange();
-
+	if (bWithTransaction)
+	{
+		BeginChange();
+	}
 	if (FDynamicMesh3* Mesh = GetSculptMesh())
 	{
 		ROIWeightValueBuffer.SetNum(VertexROI.Num(), EAllowShrinking::No);
@@ -1701,7 +1723,10 @@ void UDataflowEditorVertexAttributePaintTool::ApplyValueToSelection(EDataflowEdi
 	// update colors
 	UpdateVertexColorOverlay();
 	GetToolManager()->PostInvalidation();
-	EndChange();
+	if (bWithTransaction)
+	{
+		EndChange();
+	}
 }
 
 void UDataflowEditorVertexAttributePaintTool::PruneSelection(float Threshold)
diff --git a/Engine/Plugins/Experimental/Dataflow/Source/DataflowEditor/Private/DataflowEditorTools/DataflowEditorVertexAttributePaintToolDetailCustomization.cpp b/Engine/Plugins/Experimental/Dataflow/Source/DataflowEditor/Private/DataflowEditorTools/DataflowEditorVertexAttributePaintToolDetailCustomization.cpp
--- a/Engine/Plugins/Experimental/Dataflow/Source/DataflowEditor/Private/DataflowEditorTools/DataflowEditorVertexAttributePaintToolDetailCustomization.cpp
+++ b/Engine/Plugins/Experimental/Dataflow/Source/DataflowEditor/Private/DataflowEditorTools/DataflowEditorVertexAttributePaintToolDetailCustomization.cpp
@@ -1111,6 +1111,31 @@ namespace UE::Dataflow::Editor
 
 	void FVertexAttributePaintToolDetailCustomization::MakeSelectionAddMultiplySliderOperationRow(IDetailCategoryBuilder& EditValuesCategory) const
 	{
+		struct FSliderState
+		{
+			float Value = 0.f;
+		};
+
+		TSharedPtr<FSliderState> SliderState = MakeShared<FSliderState>();
+
+		auto GetCurrentValue = [SliderState]()
+			{
+				return SliderState->Value;
+			};
+
+		auto ResetValue = [SliderState](EDataflowEditorToolEditOperation Mode)
+			{
+				switch (Mode)
+				{
+				case EDataflowEditorToolEditOperation::Add:
+					SliderState->Value = 0.f;
+					break;
+				case EDataflowEditorToolEditOperation::Multiply:
+					SliderState->Value = 1.f;
+					break;
+				}
+			};
+
 		TSharedRef<SSegmentedControl<EDataflowEditorToolEditOperation>> AddMultiplyButtonsWidget =
 			SNew(SSegmentedControl<EDataflowEditorToolEditOperation>)
 			.ToolTipText(LOCTEXT("InteractiveEditModeTooltip",
@@ -1118,6 +1143,10 @@ namespace UE::Dataflow::Editor
 				"Multiply: applies the current value multiplied by the flood value to the new value.\n"
 				"This operation applies interactively while dragging the slider. It operates on the currently selected vertices."))
 			.Value(EDataflowEditorToolEditOperation::Add)
+			.OnValueChanged_Lambda([ResetValue](EDataflowEditorToolEditOperation Mode)
+				{
+					ResetValue(Mode);
+				})
 			+ SSegmentedControl<EDataflowEditorToolEditOperation>::Slot(EDataflowEditorToolEditOperation::Add)
 			.Text(LOCTEXT("BrushAddMode", "Add"))
 			+ SSegmentedControl<EDataflowEditorToolEditOperation>::Slot(EDataflowEditorToolEditOperation::Multiply)
@@ -1139,42 +1168,20 @@ namespace UE::Dataflow::Editor
 				switch (AddMultiplyButtonsWidget->GetValue())
 				{
 				case EDataflowEditorToolEditOperation::Add: return 1.f;
-				case EDataflowEditorToolEditOperation::Multiply: return 10.f;
+				case EDataflowEditorToolEditOperation::Multiply: return 2.f;
 				}
 				return 0.f;
 			};
 
-		struct FSliderState
-		{
-			float Value = 0.f;
-		};
-
-		TSharedPtr<FSliderState> SliderState = MakeShared<FSliderState>();
 
-		auto GetCurrentValue = [SliderState]()
-			{
-				return SliderState->Value;
-			};
-
-		auto ResetValue = [SliderState, AddMultiplyButtonsWidget]()
-			{
-				switch (AddMultiplyButtonsWidget->GetValue())
-				{
-				case EDataflowEditorToolEditOperation::Add: 
-					SliderState->Value = 0.f;
-					break;
-				case EDataflowEditorToolEditOperation::Multiply:
-					SliderState->Value = 1.f;
-					break;
-				}
-			};
 
 		auto ApplyOperation = [this, SliderState, AddMultiplyButtonsWidget](float NewValue)
 			{
 				if (Tool.IsValid())
 				{
 					SliderState->Value = NewValue;
-					Tool->ApplyValueToSelection(AddMultiplyButtonsWidget->GetValue(), NewValue);
+					// bWithTransaction = false because we already are tracking our own transaction ( see OnBeginSliderMovement_Lambda and OnEndSliderMovement_Lambda)
+					Tool->ApplyValueToSelection(AddMultiplyButtonsWidget->GetValue(), NewValue, /*bWithTransaction*/false);
 				}
 			};
 
@@ -1195,24 +1202,18 @@ namespace UE::Dataflow::Editor
 				})
 			.OnBeginSliderMovement_Lambda([this]()
 				{
-					// TODO(ccaillaud)
-					/*if (Tool.IsValid())
+					if (Tool.IsValid())
 					{
 						Tool->BeginChange();
-					}*/
+					}
 				})
-			.OnEndSliderMovement_Lambda([this, ResetValue](float)
+			.OnEndSliderMovement_Lambda([this, ResetValue, AddMultiplyButtonsWidget](float)
 				{
-					const FText TransactionLabel = LOCTEXT("FloodWeightChange", "Flood values on vertices.");
-					// TODO(ccaillaud)
-					/*if (Tool.IsValid())
+					if (Tool.IsValid())
 					{
-						Tool->EndChange(TransactionLabel);
-					}*/
-
-					// reset the value
-					//SliderState-
-					ResetValue();
+						Tool->EndChange();
+					}
+					ResetValue(AddMultiplyButtonsWidget->GetValue());
 				})
 			.ToolTipText(LOCTEXT("FloodWeightsToolTip", "Drag the slider to interactively adjust values on the selected vertices."))
 			;
diff --git a/Engine/Plugins/Experimental/Dataflow/Source/DataflowEditor/Public/DataflowEditorTools/DataflowEditorVertexAttributePaintTool.h b/Engine/Plugins/Experimental/Dataflow/Source/DataflowEditor/Public/DataflowEditorTools/DataflowEditorVertexAttributePaintTool.h
--- a/Engine/Plugins/Experimental/Dataflow/Source/DataflowEditor/Public/DataflowEditorTools/DataflowEditorVertexAttributePaintTool.h
+++ b/Engine/Plugins/Experimental/Dataflow/Source/DataflowEditor/Public/DataflowEditorTools/DataflowEditorVertexAttributePaintTool.h
@@ -52,6 +52,11 @@ namespace UE::DataflowEditorVertexAttributePaintTool::CVars
 	extern bool DataflowEditorUseNewWeightMapTool;
 }
 
+namespace UE::Dataflow::Editor
+{
+	class FVertexAttributePaintToolDetailCustomization;
+}
+
 DECLARE_STATS_GROUP(TEXT("Dataflow_VertexAttributePaintTool"), STATGROUP_VertexAttributePaintTool, STATCAT_Advanced);
 DECLARE_CYCLE_STAT(TEXT("Dataflow_VertexAttributePaintTool_UpdateROI"), VertexAttributePaintTool_UpdateROI, STATGROUP_VertexAttributePaintTool);
 DECLARE_CYCLE_STAT(TEXT("Dataflow_VertexAttributePaintTool_ApplyStamp"), VertexAttributePaintTool_ApplyStamp, STATGROUP_VertexAttributePaintTool);
@@ -287,14 +292,20 @@ public:
 			{
 				for (const int32 Idx : WeightToDynamicMesh[DynamicMeshToWeight[VertexIdx]])
 				{
-					ActiveWeightEditChangeTracker->SaveVertexOneRingTriangles(Idx, true);
+					if (ActiveWeightEditChangeTracker)
+					{
+						ActiveWeightEditChangeTracker->SaveVertexOneRingTriangles(Idx, true);
+					}
 					const float Value = Lambda(*ActiveWeightMap, Idx);
 					ActiveWeightMap->SetValue(Idx, &Value);
 				}
 			}
 			else
 			{
-				ActiveWeightEditChangeTracker->SaveVertexOneRingTriangles(VertexIdx, true);
+				if (ActiveWeightEditChangeTracker)
+				{
+					ActiveWeightEditChangeTracker->SaveVertexOneRingTriangles(VertexIdx, true);
+				}
 				const float Value = Lambda(*ActiveWeightMap, VertexIdx);
 				ActiveWeightMap->SetValue(VertexIdx, &Value);
 			}
@@ -406,7 +417,7 @@ public:
 	UE_API void PruneSelection(float Threshold);
 
 	// selection operations
-	UE_API void ApplyValueToSelection(EDataflowEditorToolEditOperation Operation, float InValue);
+	UE_API void ApplyValueToSelection(EDataflowEditorToolEditOperation Operation, float InValue, bool bWithTransaction = true);
 	UE_API void MirrorValues();
 
 	UE_API TArray<int32> GetSelectedVertices() const;
@@ -427,6 +438,9 @@ protected:
 
 	UE_API void OnColorRampChanged(TArray<FRichCurve*> Curves);
 
+	void SaveBrushConfig(EDataflowEditorToolEditOperation BrushMode);
+	void LoadBrushConfig(EDataflowEditorToolEditOperation BrushMode);
+
 	virtual bool SharesBrushPropertiesChanges() const override { return false; }
 	// end UMeshSculptToolBase API
 
@@ -510,6 +524,7 @@ protected:
 	UE_API bool SyncWeightBufferWithMesh(const FDynamicMesh3* Mesh);
 	UE_API void SetVerticesToWeightMap(const TSet<int32>& Vertices, double WeightValue);
 
+	friend class UE::Dataflow::Editor::FVertexAttributePaintToolDetailCustomization;
 	UE_API void BeginChange();
 	UE_API void EndChange();
 	UE_API void CancelChange();
