Dataflow : add visiblility filter to weight map tool
- Also fix the implementation to not miss some vertices based on the precision of the ray cast

#Synced-CL 49114577

diff --git a/Engine/Plugins/Experimental/Dataflow/Source/DataflowEditor/Private/DataflowEditorTools/DataflowEditorVertexAttributePaintTool.cpp b/Engine/Plugins/Experimental/Dataflow/Source/DataflowEditor/Private/DataflowEditorTools/DataflowEditorVertexAttributePaintTool.cpp
--- a/Engine/Plugins/Experimental/Dataflow/Source/DataflowEditor/Private/DataflowEditorTools/DataflowEditorVertexAttributePaintTool.cpp
+++ b/Engine/Plugins/Experimental/Dataflow/Source/DataflowEditor/Private/DataflowEditorTools/DataflowEditorVertexAttributePaintTool.cpp
@@ -1913,7 +1913,7 @@ void UDataflowEditorVertexAttributePaintTool::ApplyVisibilityFilter(TSet<int32>&
 	}
 
 	OutputBuffer.Reset();
-	ApplyVisibilityFilter(TempROIBuffer, OutputBuffer);
+	ApplyVisibilityFilter(ROIBuffer, OutputBuffer);
 
 	Vertices.Reset();
 	for (int32 VertexIdx : OutputBuffer)
@@ -1945,23 +1945,22 @@ void UDataflowEditorVertexAttributePaintTool::ApplyVisibilityFilter(const TArray
 			VisibilityFilterBuffer[idx] = true;
 			UE::Geometry::FVertexInfo VertexInfo;
 			Mesh->GetVertex(Vertices[idx], VertexInfo, true, false, false);
-			FVector3d Centroid = VertexInfo.Position;
-			FVector3d FaceNormal = (FVector3d)VertexInfo.Normal;
-			if (FaceNormal.Dot((Centroid - LocalEyePosition)) > 0)
+			const FVector3d Centroid = VertexInfo.Position;
+			const FVector3d FaceNormal = (FVector3d)VertexInfo.Normal;
+			const FVector3d CameraToVertexNormal = UE::Geometry::Normalized(Centroid - LocalEyePosition);
+			if (!FaceNormal.IsZero() && FaceNormal.Dot(CameraToVertexNormal) > 0)
 			{
 				VisibilityFilterBuffer[idx] = false;
 			}
-			if (ToolProperties->BrushProperties.VisibilityFilter == EDataflowEditorToolVisibilityType::Unoccluded)
+			else if (ToolProperties->BrushProperties.VisibilityFilter == EDataflowEditorToolVisibilityType::Unoccluded)
 			{
-				int32 HitTID = Octree.FindNearestHitObject(FRay3d(LocalEyePosition, UE::Geometry::Normalized(Centroid - LocalEyePosition)));
+				const FVector RayDirection = CameraToVertexNormal;
+				const FVector RayOrigin = LocalEyePosition;
+				const double MaxDistance = (LocalEyePosition - Centroid).Size() - UE_SMALL_NUMBER;
+				const int32 HitTID = Octree.FindNearestHitObject(FRay3d(RayOrigin, RayDirection), MaxDistance);
 				if (HitTID != IndexConstants::InvalidID && Mesh->IsTriangle(HitTID))
 				{
-					// Check to see if our vertex has been occulded by another triangle.
-					UE::Geometry::FIndex3i TriVertices = Mesh->GetTriangle(HitTID);
-					if (TriVertices[0] != Vertices[idx] && TriVertices[1] != Vertices[idx] && TriVertices[2] != Vertices[idx])
-					{
-						VisibilityFilterBuffer[idx] = false;
-					}
+					VisibilityFilterBuffer[idx] = false;
 				}
 			}
 		});
diff --git a/Engine/Plugins/Experimental/Dataflow/Source/DataflowEditor/Private/DataflowEditorTools/DataflowEditorVertexAttributePaintToolDetailCustomization.cpp b/Engine/Plugins/Experimental/Dataflow/Source/DataflowEditor/Private/DataflowEditorTools/DataflowEditorVertexAttributePaintToolDetailCustomization.cpp
--- a/Engine/Plugins/Experimental/Dataflow/Source/DataflowEditor/Private/DataflowEditorTools/DataflowEditorVertexAttributePaintToolDetailCustomization.cpp
+++ b/Engine/Plugins/Experimental/Dataflow/Source/DataflowEditor/Private/DataflowEditorTools/DataflowEditorVertexAttributePaintToolDetailCustomization.cpp
@@ -33,6 +33,31 @@ namespace UE::Dataflow::Editor
 	static const FName QueryTypeVertexFast(TEXT("Nearest Vertex Fast"));
 	static const FName QueryTypeVertexAccurate(TEXT("Nearest Vertex Accurate"));
 
+	static const FName VisibilityFilterNone(TEXT("All"));
+	static const FName VisibilityFilterUnoccluded(TEXT("Unoccluded"));
+
+	static EDataflowEditorToolVisibilityType GetVisibilityFilterEnumFromName(FName Name)
+	{
+		if (Name == VisibilityFilterUnoccluded) 
+		{ 
+			return EDataflowEditorToolVisibilityType::Unoccluded;
+		}
+		// VisibilityFilterNone & default 
+		return EDataflowEditorToolVisibilityType::None;
+	}
+
+	static FName GetVisibilityFilterNameFromEnum(EDataflowEditorToolVisibilityType Enum)
+	{
+		switch (Enum)
+		{
+		case EDataflowEditorToolVisibilityType::Unoccluded: 
+			return VisibilityFilterUnoccluded;
+		case EDataflowEditorToolVisibilityType::None:
+		default:
+			return VisibilityFilterNone;
+		}
+	}
+
 	///////////////////////////////////////////////////////////////////////////////////////////////////////////////
 	FVertexAttributePaintToolDetailCustomization::FVertexAttributePaintToolDetailCustomization()
 	{
@@ -48,6 +73,9 @@ namespace UE::Dataflow::Editor
 		BrushQueryTypeNameFromEnum.Add(EDataflowEditorToolValueQueryType::NearestVertexFast, QueryTypeVertexFast);
 		BrushQueryTypeNameFromEnum.Add(EDataflowEditorToolValueQueryType::NearestVertexAccurate, QueryTypeVertexAccurate);
 
+		VisibilityFilterNames.Add(VisibilityFilterNone);
+		VisibilityFilterNames.Add(VisibilityFilterUnoccluded);
+
 	}
 
 	///////////////////////////////////////////////////////////////////////////////////////////////////////////////
@@ -459,6 +487,12 @@ namespace UE::Dataflow::Editor
 			AddBrushValueQueryMethodRow(QueryCategory);
 		}
 
+		IDetailCategoryBuilder& FiltersCategory = DetailBuilder.EditCategory("Filters", FText::GetEmpty(), ECategoryPriority::Important);
+		FiltersCategory.InitiallyCollapsed(false);
+		{
+			AddVisibilityFilterRow(FiltersCategory);
+		}
+
 		// Mesh Display category
 		IDetailCategoryBuilder& MeshDisplayCategory = DetailBuilder.EditCategory("MeshDisplay", FText::GetEmpty(), ECategoryPriority::Important);
 		MeshDisplayCategory.InitiallyCollapsed(false);
@@ -715,7 +749,7 @@ namespace UE::Dataflow::Editor
 				SNew(SComboBox<FName>)
 					.ComboBoxStyle(&FAppStyle::Get().GetWidgetStyle<FComboBoxStyle>(TEXT("ComboBox")))
 					.OptionsSource(&BrushQueryTypeNames)
-					.HasDownArrow(false)
+					.HasDownArrow(true)
 					.InitiallySelectedItem(QueryTypeVertexFast)
 					.ContentPadding(0.f)
 					.OnGenerateWidget_Lambda([this](FName Item)
@@ -747,6 +781,54 @@ namespace UE::Dataflow::Editor
 			];
 	}
 
+	void FVertexAttributePaintToolDetailCustomization::AddVisibilityFilterRow(IDetailCategoryBuilder& BrushCategory) const
+	{
+		const EDataflowEditorToolVisibilityType InitialValue = (ToolProperties.IsValid()) ? ToolProperties->BrushProperties.VisibilityFilter :EDataflowEditorToolVisibilityType::None;
+			
+
+		BrushCategory.AddCustomRow(LOCTEXT("VisibilityFilterCategory", "VisibilityFilter"), false)
+			.NameContent()
+			[
+				SNew(STextBlock)
+					.Text(LOCTEXT("VisibilityFilterLabel", "Visibility Filter"))
+					.Font(FAppStyle::GetFontStyle(TEXT("PropertyWindow.NormalFont")))
+					.ToolTipText(LOCTEXT("VisibilityFilterTooltip", "Method to use to filter faces to paint on"))
+			]
+			.ValueContent()
+			[
+				SNew(SComboBox<FName>)
+					.ComboBoxStyle(&FAppStyle::Get().GetWidgetStyle<FComboBoxStyle>(TEXT("ComboBox")))
+					.OptionsSource(&VisibilityFilterNames)
+					.HasDownArrow(true)
+					.InitiallySelectedItem(GetVisibilityFilterNameFromEnum(InitialValue))
+					.ContentPadding(0.f)
+					.OnGenerateWidget_Lambda([this](FName Item)
+						{
+							return SNew(STextBlock)
+								.Margin(FMargin(2, 1))
+								.Justification(ETextJustify::Left)
+								.Text(FText::FromString(Item.ToString()));
+						})
+					.OnSelectionChanged_Lambda([this](const FName InSelection, ESelectInfo::Type InSelectInfo)
+						{
+							if (ToolProperties.IsValid())
+							{
+								ToolProperties->BrushProperties.VisibilityFilter = GetVisibilityFilterEnumFromName(InSelection);
+							}
+						})
+					.Content()
+					[
+						SNew(STextBlock)
+							.Margin(FMargin(2, 1))
+							.Text_Lambda([this]()
+								{
+									const EDataflowEditorToolVisibilityType CurrentValue = (ToolProperties.IsValid()) ? ToolProperties->BrushProperties.VisibilityFilter : EDataflowEditorToolVisibilityType::None;
+									return FText::FromName(GetVisibilityFilterNameFromEnum(CurrentValue));
+								})
+					]
+			];
+	}
+
 	void FVertexAttributePaintToolDetailCustomization::AddBrushHitBackFacesRow(IDetailCategoryBuilder& BrushCategory) const
 	{
 		BrushCategory.AddCustomRow(LOCTEXT("BrushHitBackFacesCategory", "Brush Hit Back Faces"), false)
diff --git a/Engine/Plugins/Experimental/Dataflow/Source/DataflowEditor/Private/DataflowEditorTools/DataflowEditorVertexAttributePaintToolDetailCustomization.h b/Engine/Plugins/Experimental/Dataflow/Source/DataflowEditor/Private/DataflowEditorTools/DataflowEditorVertexAttributePaintToolDetailCustomization.h
--- a/Engine/Plugins/Experimental/Dataflow/Source/DataflowEditor/Private/DataflowEditorTools/DataflowEditorVertexAttributePaintToolDetailCustomization.h
+++ b/Engine/Plugins/Experimental/Dataflow/Source/DataflowEditor/Private/DataflowEditorTools/DataflowEditorVertexAttributePaintToolDetailCustomization.h
@@ -46,6 +46,7 @@ namespace UE::Dataflow::Editor
 		void AddBrushValueAtBrushRow(IDetailCategoryBuilder& BrushCategory) const;
 		void AddBrushValueQueryMethodRow(IDetailCategoryBuilder& BrushCategory) const;
 		void AddBrushHitBackFacesRow(IDetailCategoryBuilder& BrushCategory) const;
+		void AddVisibilityFilterRow(IDetailCategoryBuilder& BrushCategory) const;
 		
 
 		void AddSelectionUI(IDetailLayoutBuilder& DetailBuilder) const;
@@ -74,6 +75,8 @@ namespace UE::Dataflow::Editor
 		TMap<FName, EDataflowEditorToolValueQueryType> BrushQueryTypeEnumFromNames;
 		TMap<EDataflowEditorToolValueQueryType, FName> BrushQueryTypeNameFromEnum;
 
+		TArray<FName> VisibilityFilterNames;
+
 		static float WeightSliderWidths;
 		static float WeightEditingLabelsPercent;
 		static float WeightEditVerticalPadding;
