Dataflow : Add Material view and rendering

#Synced-CL 49200819

diff --git a/Engine/Plugins/Experimental/Dataflow/Source/DataflowEditor/Private/Dataflow/DataflowEditorModule.cpp b/Engine/Plugins/Experimental/Dataflow/Source/DataflowEditor/Private/Dataflow/DataflowEditorModule.cpp
--- a/Engine/Plugins/Experimental/Dataflow/Source/DataflowEditor/Private/Dataflow/DataflowEditorModule.cpp
+++ b/Engine/Plugins/Experimental/Dataflow/Source/DataflowEditor/Private/Dataflow/DataflowEditorModule.cpp
@@ -44,6 +44,7 @@
 #include "DataflowRendering/DataflowPointRenderableType.h"
 #include "DataflowRendering/DataflowBoxRenderableType.h"
 #include "DataflowRendering/DataflowSphereRenderableType.h"
+#include "DataflowRendering/DataflowMaterialRenderableType.h"
 
 #include "ContentBrowserModule.h"
 #include "Modules/ModuleManager.h"
@@ -289,6 +290,7 @@ void FDataflowEditorModule::StartupModule()
 	UE::Dataflow::Private::RegisterPointRenderableTypes();
 	UE::Dataflow::Private::RegisterBoxRenderableTypes();
 	UE::Dataflow::Private::RegisterSphereRenderableTypes();
+	UE::Dataflow::Private::RegisterMaterialRenderableTypes();
 
 	// register an browser extension to show when an asset is bound to a dataflow asset
 	if (FContentBrowserModule* ContentBrowserModule = FModuleManager::Get().GetModulePtr<FContentBrowserModule>(TEXT("ContentBrowser")))
diff --git a/Engine/Plugins/Experimental/Dataflow/Source/DataflowEditor/Private/Dataflow/DataflowRenderingViewMode.cpp b/Engine/Plugins/Experimental/Dataflow/Source/DataflowEditor/Private/Dataflow/DataflowRenderingViewMode.cpp
--- a/Engine/Plugins/Experimental/Dataflow/Source/DataflowEditor/Private/Dataflow/DataflowRenderingViewMode.cpp
+++ b/Engine/Plugins/Experimental/Dataflow/Source/DataflowEditor/Private/Dataflow/DataflowRenderingViewMode.cpp
@@ -23,6 +23,8 @@ namespace UE::Dataflow
 		return ELevelViewportType::LVT_Perspective;
 	}
 
+	////////////////////////////////////////////////////////////////////////////////////////////////
+
 	FName FDataflowConstruction2DViewMode::Name = FName("2DView");
 
 	FName FDataflowConstruction2DViewMode::GetName() const
@@ -40,6 +42,8 @@ namespace UE::Dataflow
 		return FText(LOCTEXT("2DViewTooltipText", "Default 2D View"));
 	}
 	
+	////////////////////////////////////////////////////////////////////////////////////////////////
+
 	FName FDataflowConstruction3DViewMode::Name = FName("3DView");
 
 	FName FDataflowConstruction3DViewMode::GetName() const
@@ -57,6 +61,8 @@ namespace UE::Dataflow
 		return FText(LOCTEXT("3DViewTooltipText", "Default 3D View"));
 	}
 
+	////////////////////////////////////////////////////////////////////////////////////////////////
+
 	FName FDataflowConstructionUVViewMode::Name = FName("UVView");
 
 	FName FDataflowConstructionUVViewMode::GetName() const
@@ -74,6 +80,25 @@ namespace UE::Dataflow
 		return FText(LOCTEXT("UVViewTooltipText", "UV/Texture Coordinate View"));
 	}
 
+	////////////////////////////////////////////////////////////////////////////////////////////////
+
+	FName FDataflowConstructionMaterialViewMode::Name = FName("MaterialView");
+
+	FName FDataflowConstructionMaterialViewMode::GetName() const
+	{
+		return Name;
+	}
+
+	FText FDataflowConstructionMaterialViewMode::GetButtonText() const
+	{
+		return FText(LOCTEXT("MaterialViewButtonText", "MaterialView"));
+	}
+
+	FText FDataflowConstructionMaterialViewMode::GetTooltipText() const
+	{
+		return FText(LOCTEXT("MaterialViewTooltipText", "Material View"));
+	}
+
 	//
 	// Factory
 	//
@@ -84,6 +109,7 @@ namespace UE::Dataflow
 		ViewModeMap.Add(FDataflowConstruction2DViewMode::Name, MakeUnique<FDataflowConstruction2DViewMode>());
 		ViewModeMap.Add(FDataflowConstruction3DViewMode::Name, MakeUnique<FDataflowConstruction3DViewMode>());
 		ViewModeMap.Add(FDataflowConstructionUVViewMode::Name, MakeUnique<FDataflowConstructionUVViewMode>());
+		ViewModeMap.Add(FDataflowConstructionMaterialViewMode::Name, MakeUnique<FDataflowConstructionMaterialViewMode>());
 	}
 
 	FRenderingViewModeFactory& FRenderingViewModeFactory::GetInstance()
diff --git a/Engine/Plugins/Experimental/Dataflow/Source/DataflowEditor/Private/DataflowRendering/DataflowMaterialRenderableType.cpp b/Engine/Plugins/Experimental/Dataflow/Source/DataflowEditor/Private/DataflowRendering/DataflowMaterialRenderableType.cpp
new file mode 100644
--- /dev/null
+++ b/Engine/Plugins/Experimental/Dataflow/Source/DataflowEditor/Private/DataflowRendering/DataflowMaterialRenderableType.cpp
@@ -0,0 +1,108 @@
+// Copyright Epic Games, Inc. All Rights Reserved.
+
+#include "DataflowRendering/DataflowMaterialRenderableType.h"
+
+#include "Components/InstancedStaticMeshComponent.h"
+
+#include "Dataflow/DataflowRenderingViewMode.h"
+
+#include "DataflowRendering/DataflowRenderableComponents.h"
+#include "DataflowRendering/DataflowRenderableType.h"
+#include "DataflowRendering/DataflowRenderableTypeInstance.h"
+#include "DataflowRendering/DataflowRenderableTypeRegistry.h"
+#include "DataflowRendering/DataflowRenderableTypeUtils.h"
+
+#include "Components/StaticMeshComponent.h"
+#include "Materials/MaterialInterface.h"
+
+#include "Engine/StaticMesh.h"
+
+#include "UObject/ObjectPtr.h"
+
+namespace UE::Dataflow::Private
+{
+	///////////////////////////////////////////////////////////////////////////////////////////////////////////////
+
+	static UStaticMesh* LoadMaterialPreviewMesh()
+	{
+		return LoadObject<UStaticMesh>(nullptr, TEXT("/Script/Engine.StaticMesh'/Engine/EngineMeshes/SM_MatPreviewMesh_01.SM_MatPreviewMesh_01'"));
+	}
+
+	class FMaterialSurfaceRenderableType : public IRenderableType
+	{
+		UE_DATAFLOW_IRENDERABLE_PRIMARY_TYPE(TObjectPtr<UMaterialInterface>, Material);
+		UE_DATAFLOW_IRENDERABLE_RENDER_GROUP(Material);
+		UE_DATAFLOW_IRENDERABLE_VIEW_MODE(FDataflowConstructionMaterialViewMode);
+
+		virtual void GetPrimitiveComponents(const FRenderableTypeInstance& Instance, FRenderableComponents& OutComponents) const override
+		{
+			const TObjectPtr<UMaterialInterface> Material = GetMaterial(Instance, {});
+
+			static const FName MaterialComponentName = TEXT("Material");
+			if (UStaticMeshComponent* Component = OutComponents.AddNewComponent<UStaticMeshComponent>(MaterialComponentName))
+			{
+				//LoadMaterialPreviewMesh
+				Component->SetStaticMesh(LoadMaterialPreviewMesh());
+				Component->SetMaterial(0, Material);
+			}
+		}
+	};
+
+	///////////////////////////////////////////////////////////////////////////////////////////////////////////////
+
+	class FMaterialArraySurfaceRenderableType : public IRenderableType
+	{
+		UE_DATAFLOW_IRENDERABLE_PRIMARY_TYPE(TArray<TObjectPtr<UMaterialInterface>>, MaterialArray);
+		UE_DATAFLOW_IRENDERABLE_RENDER_GROUP(Material);
+		UE_DATAFLOW_IRENDERABLE_VIEW_MODE(FDataflowConstructionMaterialViewMode);
+
+		virtual void GetPrimitiveComponents(const FRenderableTypeInstance& Instance, FRenderableComponents& OutComponents) const override
+		{
+			const TArray<TObjectPtr<UMaterialInterface>>& MaterialArray = GetMaterialArray(Instance, {});
+			const int32 NumMaterials = MaterialArray.Num();
+			if (NumMaterials > 0)
+			{
+				UStaticMesh* MaterialPreviewMesh = LoadMaterialPreviewMesh();
+				const FVector PreviewBoundsSize = MaterialPreviewMesh->GetBoundingBox().GetSize();
+
+				const int32 NumMaterialPerDimension = FMath::Max(0, FMath::CeilToInt32(FMath::Sqrt((float)NumMaterials)));
+				for (int32 Y = 0; Y < NumMaterialPerDimension; ++Y)
+				{
+					for (int32 X = 0; X < NumMaterialPerDimension; ++X)
+					{
+						const int32 MaterialIndex = X + (Y * NumMaterialPerDimension);
+						if (MaterialArray.IsValidIndex(MaterialIndex))
+						{
+							const FName ComponentName = FName(FString::Printf(TEXT("Material [%01d]"), MaterialIndex));
+							if (UStaticMeshComponent* Component = OutComponents.AddNewComponent<UStaticMeshComponent>(ComponentName))
+							{
+								Component->SetStaticMesh(MaterialPreviewMesh);
+								Component->SetMaterial(0, MaterialArray[MaterialIndex]);
+
+								const double SpacingFactor = 1.1f;
+								const FVector Position
+								{
+									 PreviewBoundsSize.X * SpacingFactor * (double)X,
+									 PreviewBoundsSize.Y * SpacingFactor * (double)Y,
+									 0
+								};
+								Component->SetWorldTransform(FTransform(FQuat::Identity, Position));
+							}
+						}
+					}
+				}
+			}
+		}
+	};
+
+	///////////////////////////////////////////////////////////////////////////////////////////////////////////////
+
+	void RegisterMaterialRenderableTypes()
+	{
+		UE_DATAFLOW_REGISTER_RENDERABLE_TYPE(FMaterialSurfaceRenderableType);
+		UE_DATAFLOW_REGISTER_RENDERABLE_TYPE(FMaterialArraySurfaceRenderableType);
+	}
+
+	///////////////////////////////////////////////////////////////////////////////////////////////////////////////
+
+}
\ No newline at end of file
diff --git a/Engine/Plugins/Experimental/Dataflow/Source/DataflowEditor/Private/DataflowRendering/DataflowMaterialRenderableType.h b/Engine/Plugins/Experimental/Dataflow/Source/DataflowEditor/Private/DataflowRendering/DataflowMaterialRenderableType.h
new file mode 100644
--- /dev/null
+++ b/Engine/Plugins/Experimental/Dataflow/Source/DataflowEditor/Private/DataflowRendering/DataflowMaterialRenderableType.h
@@ -0,0 +1,8 @@
+// Copyright Epic Games, Inc. All Rights Reserved.
+
+#pragma once
+
+namespace UE::Dataflow::Private
+{
+	void RegisterMaterialRenderableTypes();
+}
diff --git a/Engine/Plugins/Experimental/Dataflow/Source/DataflowEditor/Public/Dataflow/DataflowRenderingViewMode.h b/Engine/Plugins/Experimental/Dataflow/Source/DataflowEditor/Public/Dataflow/DataflowRenderingViewMode.h
--- a/Engine/Plugins/Experimental/Dataflow/Source/DataflowEditor/Public/Dataflow/DataflowRenderingViewMode.h
+++ b/Engine/Plugins/Experimental/Dataflow/Source/DataflowEditor/Public/Dataflow/DataflowRenderingViewMode.h
@@ -89,6 +89,16 @@ namespace UE::Dataflow
 		DATAFLOWEDITOR_API virtual FText GetTooltipText() const override;
 	};
 
+	class FDataflowConstructionMaterialViewMode : public FDataflowConstruction3DViewModeBase
+	{
+	public:
+		DATAFLOWEDITOR_API static FName Name;
+		virtual ~FDataflowConstructionMaterialViewMode() = default;
+	private:
+		DATAFLOWEDITOR_API virtual FName GetName() const override;
+		DATAFLOWEDITOR_API virtual FText GetButtonText() const override;
+		DATAFLOWEDITOR_API virtual FText GetTooltipText() const override;
+	};
 
 	//
 	// ViewMode registry/factory
