Dataflow : make sure we select the right views for node using the new rendering system

#Synced-CL 49200819

diff --git a/Engine/Plugins/Experimental/Dataflow/Source/DataflowEditor/Private/Dataflow/DataflowEditorMode.cpp b/Engine/Plugins/Experimental/Dataflow/Source/DataflowEditor/Private/Dataflow/DataflowEditorMode.cpp
--- a/Engine/Plugins/Experimental/Dataflow/Source/DataflowEditor/Private/Dataflow/DataflowEditorMode.cpp
+++ b/Engine/Plugins/Experimental/Dataflow/Source/DataflowEditor/Private/Dataflow/DataflowEditorMode.cpp
@@ -955,7 +955,7 @@ void UDataflowEditorMode::SetConstructionViewModeForNode(const UDataflowEdNode*
 				for (const TPair<FName, TUniquePtr<UE::Dataflow::IDataflowConstructionViewMode>>& ViewMode : UE::Dataflow::FRenderingViewModeFactory::GetInstance().GetViewModes())
 				{
 					check(ViewMode.Value.IsValid());
-					const bool bCanRender = UE::Dataflow::CanRenderNodeOutput(*Node, *EditorContent, *ViewMode.Value);
+					const bool bCanRender = UE::Dataflow::IsViewModeSupported(*Node, *EditorContent, *ViewMode.Value);
 					if (bCanRender)
 					{
 						OutViewModeName = ViewMode.Key;
@@ -970,7 +970,7 @@ void UDataflowEditorMode::SetConstructionViewModeForNode(const UDataflowEdNode*
 		{
 			SetConstructionViewMode(NewViewMode);
 		}
-		else if (UE::Dataflow::CanRenderNodeOutput(*Node, *EditorContent, *GetConstructionViewMode()))        // Can we use the current view mode?
+		else if (UE::Dataflow::IsViewModeSupported(*Node, *EditorContent, *GetConstructionViewMode()))        // Can we use the current view mode?
 		{
 			// no need to switch
 		}
@@ -1008,7 +1008,7 @@ bool UDataflowEditorMode::CanChangeConstructionViewModeTo(const FName& NewViewMo
 						{
 							if (const UE::Dataflow::IDataflowConstructionViewMode* const ViewMode = UE::Dataflow::FRenderingViewModeFactory::GetInstance().GetViewMode(NewViewModeName))
 							{
-								if (UE::Dataflow::CanRenderNodeOutput(*SelectedDataflowEdNode, *EditorContent, *ViewMode))
+								if (UE::Dataflow::IsViewModeSupported(*SelectedDataflowEdNode, *EditorContent, *ViewMode))
 								{
 									return true;
 								}
diff --git a/Engine/Plugins/Experimental/Dataflow/Source/DataflowEditor/Private/Dataflow/DataflowEditorUtil.cpp b/Engine/Plugins/Experimental/Dataflow/Source/DataflowEditor/Private/Dataflow/DataflowEditorUtil.cpp
--- a/Engine/Plugins/Experimental/Dataflow/Source/DataflowEditor/Private/Dataflow/DataflowEditorUtil.cpp
+++ b/Engine/Plugins/Experimental/Dataflow/Source/DataflowEditor/Private/Dataflow/DataflowEditorUtil.cpp
@@ -11,6 +11,7 @@
 #include "Dataflow/DataflowObject.h"
 #include "Dataflow/DataflowRenderingFactory.h"
 #include "Dataflow/DataflowRenderingViewMode.h"
+#include "DataflowRendering/DataflowRenderableTypeInstance.h"
 #include "DynamicMesh/MeshNormals.h"
 #include "Engine/SkeletalMesh.h"
 #include "Animation/Skeleton.h"
@@ -122,6 +123,42 @@ namespace UE::Dataflow
 		return MakeShared<FEngineContext>(nullptr);
 	}
 
+	bool IsViewModeSupported(const UDataflowEdNode& EdNode, const UDataflowBaseContent& EditorContent, const IDataflowConstructionViewMode& ViewMode, bool bCheckDebugDraw)
+	{
+		if (const TSharedPtr<FEngineContext> Context = EditorContent.GetDataflowContext())
+		{
+			if (TSharedPtr<const FDataflowNode> DataflowNode = EdNode.GetDataflowNode())
+			{
+				// first check the rendering collection rendering system (legacy)
+				if (const FRenderingFactory* const Factory = FRenderingFactory::GetInstance())
+				{
+					for (const FRenderingParameter& Parameter : EdNode.GetRenderParameters())
+					{
+						if (Factory->CanRenderNodeOutput(FGraphRenderingState{ EdNode.GetDataflowNodeGuid(), DataflowNode.Get(), Parameter, *Context, ViewMode }))
+						{
+							return true;
+						}
+					}
+				}
+
+				// check if the node needs debug draw
+				if (bCheckDebugDraw && DataflowNode->CanDebugDrawViewMode(ViewMode.GetName()))
+				{
+					return true;
+				}
+
+				// check the component based rendering system (new)
+				TArray<FRenderableTypeInstance> RenderableInstances;
+				FRenderableTypeInstance::GetRenderableInstancesForNode(*Context, ViewMode, *DataflowNode, RenderableInstances);
+				if (RenderableInstances.Num() > 0)
+				{
+					return true;
+				}
+			}
+		}
+		return false;
+	}
+
 	bool CanRenderNodeOutput(const UDataflowEdNode& EdNode, const UDataflowBaseContent& EditorContent, const IDataflowConstructionViewMode& ViewMode, bool bCheckDebugDraw)
 	{
 		if (const TSharedPtr<FEngineContext> Context = EditorContent.GetDataflowContext())
@@ -185,7 +222,7 @@ namespace UE::Dataflow
 		{
 			if (const IDataflowConstructionViewMode* const ViewMode = FRenderingViewModeFactory::GetInstance().GetViewMode(ViewModeName))
 			{
-				if (CanRenderNodeOutput(EdNode, EditorContent, *ViewMode))
+				if (IsViewModeSupported(EdNode, EditorContent, *ViewMode))
 				{
 					OutValidViewModeNames.Add(ViewModeName);
 				}
diff --git a/Engine/Plugins/Experimental/Dataflow/Source/DataflowEditor/Private/Dataflow/DataflowEditorUtil.h b/Engine/Plugins/Experimental/Dataflow/Source/DataflowEditor/Private/Dataflow/DataflowEditorUtil.h
--- a/Engine/Plugins/Experimental/Dataflow/Source/DataflowEditor/Private/Dataflow/DataflowEditorUtil.h
+++ b/Engine/Plugins/Experimental/Dataflow/Source/DataflowEditor/Private/Dataflow/DataflowEditorUtil.h
@@ -54,6 +54,8 @@ namespace UE::Dataflow
 {
 	TSharedPtr<UE::Dataflow::FEngineContext> GetContext(TObjectPtr<UDataflowBaseContent> Content);
 
+	bool IsViewModeSupported(const UDataflowEdNode& EdNode, const UDataflowBaseContent& EditorContent, const IDataflowConstructionViewMode& ViewMode, bool bCheckDebugDraw = true);
+
 	bool CanRenderNodeOutput(const UDataflowEdNode& EdNode, const UDataflowBaseContent& EditorContent, const UE::Dataflow::IDataflowConstructionViewMode& ViewMode, bool bCheckDebugDraw = true);
 
 	/** This function fills the managed array collection linked to the rendering facade from the callbacks.
diff --git a/Engine/Plugins/Experimental/Dataflow/Source/DataflowEditor/Private/DataflowRendering/DataflowRenderableTypeInstance.cpp b/Engine/Plugins/Experimental/Dataflow/Source/DataflowEditor/Private/DataflowRendering/DataflowRenderableTypeInstance.cpp
--- a/Engine/Plugins/Experimental/Dataflow/Source/DataflowEditor/Private/DataflowRendering/DataflowRenderableTypeInstance.cpp
+++ b/Engine/Plugins/Experimental/Dataflow/Source/DataflowEditor/Private/DataflowRendering/DataflowRenderableTypeInstance.cpp
@@ -31,8 +31,7 @@ namespace UE::Dataflow
 					}
 				}
 			}
-			// we don't want to render all render groups, so we select the priority ones
-			return (RenderGroup == TEXT("Surface") || RenderGroup == TEXT("UVs"));
+			return true;
 		}
 	}
 
