Add vertex color to attribute transfer node

#Synced-CL 49234567

diff --git a/Engine/Plugins/Experimental/Dataflow/Source/DataflowNodes/Private/Dataflow/DataflowNodesPlugin.cpp b/Engine/Plugins/Experimental/Dataflow/Source/DataflowNodes/Private/Dataflow/DataflowNodesPlugin.cpp
--- a/Engine/Plugins/Experimental/Dataflow/Source/DataflowNodes/Private/Dataflow/DataflowNodesPlugin.cpp
+++ b/Engine/Plugins/Experimental/Dataflow/Source/DataflowNodes/Private/Dataflow/DataflowNodesPlugin.cpp
@@ -15,6 +15,7 @@
 #include "Dataflow/DataflowCollectionSetSkinningSkeletalMesh.h"
 #include "Dataflow/DataflowSkeletonAssetTerminalNode.h"
 #include "Dataflow/DataflowTextureToAttributeNode.h"
+#include "Dataflow/DataflowVertexColorToAttributeNode.h"
 
 #define LOCTEXT_NAMESPACE "DataflowNodes"
 
@@ -51,6 +52,7 @@ void IDataflowNodesPlugin::StartupModule()
 	UE::Dataflow::RegisterSelectionNodes();
 	UE::Dataflow::RegisterContextOverridesNodes();
 	UE::Dataflow::RegisterTextureToAttributeNodes();
+	UE::Dataflow::RegisterVertexColorToAttributeNodes();
 	UE::Dataflow::DataflowCollectionAttributeKeyNodes();
 	DATAFLOW_NODE_REGISTER_CREATION_FACTORY(FDataflowCollectionAddScalarVertexPropertyNode);
 	DATAFLOW_NODE_REGISTER_CREATION_FACTORY(FDataflowCollectionEditSkinWeightsNode);
diff --git a/Engine/Plugins/Experimental/Dataflow/Source/DataflowNodes/Private/Dataflow/DataflowVertexColorToAttributeNode.cpp b/Engine/Plugins/Experimental/Dataflow/Source/DataflowNodes/Private/Dataflow/DataflowVertexColorToAttributeNode.cpp
new file mode 100644
--- /dev/null
+++ b/Engine/Plugins/Experimental/Dataflow/Source/DataflowNodes/Private/Dataflow/DataflowVertexColorToAttributeNode.cpp
@@ -0,0 +1,102 @@
+// Copyright Epic Games, Inc. All Rights Reserved.
+
+#include "Dataflow/DataflowVertexColorToAttributeNode.h"
+
+#include "Dataflow/DataflowNodeFactory.h"
+#include "GeometryCollection/GeometryCollection.h"
+#include "GeometryCollection/Facades/CollectionMeshFacade.h"
+
+#include UE_INLINE_GENERATED_CPP_BY_NAME(DataflowVertexColorToAttributeNode)
+
+#define LOCTEXT_NAMESPACE "DataflowVertexColorToAttributeNode"
+
+namespace UE::Dataflow
+{
+	void RegisterVertexColorToAttributeNodes()
+	{
+		DATAFLOW_NODE_REGISTER_CREATION_FACTORY(FDataflowVertexColorToAttributeNode);
+	}
+};
+
+FDataflowVertexColorToAttributeNode::FDataflowVertexColorToAttributeNode(const UE::Dataflow::FNodeParameters& InParam, FGuid InGuid)
+	: FDataflowNode(InParam, InGuid)
+{
+	RegisterInputConnection(&Collection);
+	RegisterInputConnection(&ColorChannel);
+	RegisterInputConnection(&AttributeName);
+	RegisterInputConnection(&ScalingFactor);
+	RegisterOutputConnection(&Collection, &Collection);
+}
+
+void FDataflowVertexColorToAttributeNode::Evaluate(UE::Dataflow::FContext& Context, const FDataflowOutput* Out) const
+{
+	if (Out->IsA(&Collection))
+	{
+		const FName InAttributeName = FName(GetValue(Context, &AttributeName, AttributeName));
+		if (InAttributeName.IsNone())
+		{
+			Context.Error(LOCTEXT("NoAttributeName_Msg", "Attribute name is empty, data transfer from vertex color will be skipped"));
+			SafeForwardInput(Context, &Collection, &Collection);
+			return;
+		}
+
+		FManagedArrayCollection InCollection = GetValue(Context, &Collection, Collection);
+
+		const FName InVertexGroupName = VertexGroup.Name;
+		if (!InCollection.HasGroup(InVertexGroupName))
+		{
+			Context.Error(LOCTEXT("InvalidVertexGroup_Msg", "Vertex group does not exists, data transfer from vertex color will be skipped"));
+			SafeForwardInput(Context, &Collection, &Collection);
+			return;
+		}
+
+		// Create or get the new attribute 
+		TManagedArray<float>& TargetAttribute = InCollection.AddAttribute<float>(InAttributeName, InVertexGroupName);
+		TargetAttribute.Fill(0);
+
+		// Get the vertex color attribute (we have geometry collection or cloth collection that store it differently)
+		const TManagedArray<FLinearColor>* VertexColorAttribute = VertexColorAttribute = InCollection.FindAttribute<FLinearColor>(FGeometryCollection::ColorAttribute, InVertexGroupName);
+		if (!VertexColorAttribute)
+		{
+			// cannot really include cloth here so we need to directly serach for the vertex color attribute in the render vertices ( eventually we may converge to use only geometry collections )
+			VertexColorAttribute = InCollection.FindAttribute<FLinearColor>("RenderColor", InVertexGroupName);
+		}
+		if (!VertexColorAttribute)
+		{
+			Context.Warning(LOCTEXT("NoVertexColor_Msg", "No vertex color could be found in this collection, no data will be transfered and the attribute will be inititialized to zeros"));
+		}
+		else
+		{
+			const EDataflowImageChannel InColorChannel = GetValue(Context, &ColorChannel, ColorChannel);
+			const float InScalingFactor = FMath::Clamp(GetValue(Context, &ScalingFactor, ScalingFactor), 0.f, 1.f);
+
+			const int32 NumVertices = FMath::Min(TargetAttribute.Num(), VertexColorAttribute->Num());
+			for (int32 Index = 0; Index < NumVertices; ++Index)
+			{
+				const FLinearColor Color = (*VertexColorAttribute)[Index];
+				float Value = 0;
+				switch (InColorChannel)
+				{
+				case EDataflowImageChannel::Red: 
+					Value = Color.R;
+					break;
+				case EDataflowImageChannel::Green:
+					Value = Color.G;
+					break;
+				case EDataflowImageChannel::Blue:
+					Value = Color.B;
+					break;
+				case EDataflowImageChannel::Alpha:
+					Value = Color.A;
+					break;
+				}
+				TargetAttribute[Index] = Value * InScalingFactor;
+			}
+		}
+
+		SetValue(Context, MoveTemp(InCollection), &Collection);
+	}
+}
+
+#undef LOCTEXT_NAMESPACE 
+
diff --git a/Engine/Plugins/Experimental/Dataflow/Source/DataflowNodes/Private/Dataflow/DataflowVertexColorToAttributeNode.h b/Engine/Plugins/Experimental/Dataflow/Source/DataflowNodes/Private/Dataflow/DataflowVertexColorToAttributeNode.h
new file mode 100644
--- /dev/null
+++ b/Engine/Plugins/Experimental/Dataflow/Source/DataflowNodes/Private/Dataflow/DataflowVertexColorToAttributeNode.h
@@ -0,0 +1,53 @@
+// Copyright Epic Games, Inc. All Rights Reserved.
+
+#pragma once 
+
+#include "Dataflow/DataflowNode.h"
+#include "GeometryCollection/ManagedArrayCollection.h"
+#include "Dataflow/DataflowCollectionAddScalarVertexPropertyNode.h"
+
+#include "Dataflow/DataflowImage.h"
+
+#include "DataflowVertexColorToAttributeNode.generated.h"
+
+/**
+* Transfer data from a vertex color to a an attribute in a collection
+*/
+USTRUCT()
+struct FDataflowVertexColorToAttributeNode : public FDataflowNode
+{
+	GENERATED_USTRUCT_BODY()
+	DATAFLOW_NODE_DEFINE_INTERNAL(FDataflowVertexColorToAttributeNode, "VertexColorToAttribute", "Dataflow", "Transfer")
+
+public:
+	FDataflowVertexColorToAttributeNode(const UE::Dataflow::FNodeParameters& InParam, FGuid InGuid = FGuid::NewGuid());
+
+private:
+	/** Collection to transfer the attribute data to */
+	UPROPERTY(meta = (DataflowInput, DataflowOutput, DataflowPassthrough = Collection))
+	FManagedArrayCollection Collection;
+
+	/** Color channel to use as attribute value */
+	UPROPERTY(EditAnywhere, Category = Source, meta = (DataflowInput))
+	EDataflowImageChannel ColorChannel = EDataflowImageChannel::Red;
+
+	/** Name of the attribute to to transfer the vertex color to */
+	UPROPERTY(EditAnywhere, Category = Target, meta = (DataflowInput))
+	FString AttributeName;
+
+	/** Scaling factor applied to vertex colours (in range 0-1)  */
+	UPROPERTY(EditAnywhere, Category = Target, meta = (DataflowInput, UIMin = 0.0, UIMax = 1.0, ClampMin = 0.0, ClampMax = 1.0))
+	float ScalingFactor = 1.0f;
+
+	/** Vertex group to use */
+	UPROPERTY(EditAnywhere, Category = Target)
+	FScalarVertexPropertyGroup VertexGroup;
+
+	virtual void Evaluate(UE::Dataflow::FContext& Context, const FDataflowOutput* Out) const override;
+};
+
+namespace UE::Dataflow
+{
+	void RegisterVertexColorToAttributeNodes();
+}
+
