Dataflow : fix static analysis wraning

#Synced-CL 49237557

diff --git a/Engine/Plugins/Experimental/Dataflow/Source/DataflowNodes/Private/Dataflow/DataflowVertexColorToAttributeNode.cpp b/Engine/Plugins/Experimental/Dataflow/Source/DataflowNodes/Private/Dataflow/DataflowVertexColorToAttributeNode.cpp
--- a/Engine/Plugins/Experimental/Dataflow/Source/DataflowNodes/Private/Dataflow/DataflowVertexColorToAttributeNode.cpp
+++ b/Engine/Plugins/Experimental/Dataflow/Source/DataflowNodes/Private/Dataflow/DataflowVertexColorToAttributeNode.cpp
@@ -55,7 +55,7 @@ void FDataflowVertexColorToAttributeNode::Evaluate(UE::Dataflow::FContext& Conte
 		TargetAttribute.Fill(0);
 
 		// Get the vertex color attribute (we have geometry collection or cloth collection that store it differently)
-		const TManagedArray<FLinearColor>* VertexColorAttribute = VertexColorAttribute = InCollection.FindAttribute<FLinearColor>(FGeometryCollection::ColorAttribute, InVertexGroupName);
+		const TManagedArray<FLinearColor>* VertexColorAttribute = InCollection.FindAttribute<FLinearColor>(FGeometryCollection::ColorAttribute, InVertexGroupName);
 		if (!VertexColorAttribute)
 		{
 			// cannot really include cloth here so we need to directly serach for the vertex color attribute in the render vertices ( eventually we may converge to use only geometry collections )
