Dataflow :add support to visualize specific attribute when rendering a geometry collection

#Synced-CL 49237557

diff --git a/Engine/Plugins/Experimental/Dataflow/Source/DataflowEditor/Private/DataflowRendering/DataflowGeometryCollectionRenderableType.cpp b/Engine/Plugins/Experimental/Dataflow/Source/DataflowEditor/Private/DataflowRendering/DataflowGeometryCollectionRenderableType.cpp
--- a/Engine/Plugins/Experimental/Dataflow/Source/DataflowEditor/Private/DataflowRendering/DataflowGeometryCollectionRenderableType.cpp
+++ b/Engine/Plugins/Experimental/Dataflow/Source/DataflowEditor/Private/DataflowRendering/DataflowGeometryCollectionRenderableType.cpp
@@ -79,7 +79,10 @@ namespace UE::Dataflow::Private
 
 						TArrayView<const FTransform3f> BoneTransforms = MakeArrayView(TransformAttribute->GetConstArray());
 
-						// convert to dynamic meshes in parallel
+						// Override vertex color if we have an attribute to visualize
+						TransferAttributeToVisualizeToVertexColor(*GeomCollection, Instance.GetVertexAttributeToVisualize());
+
+						// Convert to dynamic meshes in parallel
 						ParallelFor(DynamicMeshes.Num(),
 							[&DynamicMeshes, &BoneTransforms, &GeomCollection](int32 TransformIndex)
 							{
@@ -112,6 +115,90 @@ namespace UE::Dataflow::Private
 			}
 		}
 
+		template<typename T>
+		void TransferTypedAttributeVertexColor(FGeometryCollection& GeomCollection, FName ValueAttributeName) const
+		{
+			TManagedArray<FLinearColor>* ColorAttribute = GeomCollection.FindAttributeTyped<FLinearColor>(FGeometryCollection::ColorAttribute, FGeometryCollection::VerticesGroup);
+			const TManagedArray<T>* ValueAttribute = GeomCollection.FindAttributeTyped<T>(ValueAttributeName, FGeometryCollection::VerticesGroup);
+			if (ColorAttribute && ValueAttribute)
+			{
+				check(ColorAttribute->Num() == ValueAttribute->Num());
+				const int32 NumVertices = ColorAttribute->Num();
+
+				T MaxValue{ 0 };
+				T MinValue{ 0 };
+				T Range{ 0 };
+				if constexpr (std::is_same_v<T, int32>)
+				{
+					if (const int32* MaxValuePtr = Algo::MaxElement(*ValueAttribute))
+					{
+						MaxValue = *MaxValuePtr;
+					}
+					if (const int32* MinValuePtr = Algo::MinElement(*ValueAttribute))
+					{
+						MinValue = *MinValuePtr;
+					}
+					Range = (MaxValue - MinValue);
+				}
+
+				for (int32 Index = 0; Index < NumVertices; ++Index)
+				{
+					const T& Value = (*ValueAttribute)[Index];
+					if constexpr (std::is_same_v<T, float> || std::is_same_v<T, double>)
+					{
+						(*ColorAttribute)[Index] = FLinearColor(Value, Value, Value);
+					}
+					else if constexpr (std::is_same_v<T, bool>)
+					{
+						const float FloatValue = Value ? 1.0f : 0.0f;
+						(*ColorAttribute)[Index] = FLinearColor(FloatValue, FloatValue, FloatValue);
+					}
+					else if constexpr (std::is_same_v<T, int32>)
+					{
+						const float Multiplier = (Range == 0) ? 1.f : (1.f / Range);
+						const float NormalizedValue{ (Value - MinValue) * Multiplier };
+						(*ColorAttribute)[Index] = FLinearColor(NormalizedValue, NormalizedValue, NormalizedValue);
+					}
+					else if constexpr (std::is_same_v<T, FVector3f> || std::is_same_v<T, FVector3d>)
+					{
+						FVector3d NormalizedValue{Value};
+						NormalizedValue.Normalize();
+						(*ColorAttribute)[Index] = FLinearColor(NormalizedValue);
+					}
+				}
+			}
+		}
+
+		void TransferAttributeToVisualizeToVertexColor(FGeometryCollection& GeomCollection, FName VertexAttributeName) const
+		{
+			if (VertexAttributeName.IsNone())
+			{
+				return;
+			}
+			const FManagedArrayCollection::EArrayType AttributeType = GeomCollection.GetAttributeType(VertexAttributeName, FGeometryCollection::VerticesGroup);
+			switch (AttributeType)
+			{
+			case FManagedArrayCollection::EArrayType::FFloatType:
+				TransferTypedAttributeVertexColor<float>(GeomCollection, VertexAttributeName);
+				break;
+			case FManagedArrayCollection::EArrayType::FDoubleType:
+				TransferTypedAttributeVertexColor<double>(GeomCollection, VertexAttributeName);
+				break;
+			case FManagedArrayCollection::EArrayType::FVectorType:
+				TransferTypedAttributeVertexColor<FVector3f>(GeomCollection, VertexAttributeName);
+				break;
+			case FManagedArrayCollection::EArrayType::FVector3dType:
+				TransferTypedAttributeVertexColor<FVector3d>(GeomCollection, VertexAttributeName);
+				break;
+			case FManagedArrayCollection::EArrayType::FInt32Type:
+				TransferTypedAttributeVertexColor<int32>(GeomCollection, VertexAttributeName);
+				break;
+			case FManagedArrayCollection::EArrayType::FBoolType:
+				TransferTypedAttributeVertexColor<bool>(GeomCollection, VertexAttributeName);
+				break;
+			}
+		}
+
 		void CullCurvesIfNeeded(FGeometryCollection& GeomCollection) const
 		{
 			// Performance control for curve based geometry that can ghave a lot of geometry
diff --git a/Engine/Plugins/Experimental/Dataflow/Source/DataflowEditor/Private/DataflowRendering/DataflowRenderableTypeInstance.cpp b/Engine/Plugins/Experimental/Dataflow/Source/DataflowEditor/Private/DataflowRendering/DataflowRenderableTypeInstance.cpp
--- a/Engine/Plugins/Experimental/Dataflow/Source/DataflowEditor/Private/DataflowRendering/DataflowRenderableTypeInstance.cpp
+++ b/Engine/Plugins/Experimental/Dataflow/Source/DataflowEditor/Private/DataflowRendering/DataflowRenderableTypeInstance.cpp
@@ -108,4 +108,13 @@ namespace UE::Dataflow
 		const uint32 RenderGroupHash = GetTypeHash(RenderableType? RenderableType->GetRenderGroup(): NAME_None);
 		return ::HashCombine(OutputCacheKey, RenderGroupHash);
 	}
+
+	const FName FRenderableTypeInstance::GetVertexAttributeToVisualize() const
+	{
+		if (RenderableType && Node)
+		{
+			return Node->GetVertexAttributeToVisualize(EvaluationContext, OutputName, RenderableType->GetRenderGroup());
+		}
+		return NAME_None;
+	}
 }
diff --git a/Engine/Plugins/Experimental/Dataflow/Source/DataflowEditor/Public/DataflowRendering/DataflowRenderableTypeInstance.h b/Engine/Plugins/Experimental/Dataflow/Source/DataflowEditor/Public/DataflowRendering/DataflowRenderableTypeInstance.h
--- a/Engine/Plugins/Experimental/Dataflow/Source/DataflowEditor/Public/DataflowRendering/DataflowRenderableTypeInstance.h
+++ b/Engine/Plugins/Experimental/Dataflow/Source/DataflowEditor/Public/DataflowRendering/DataflowRenderableTypeInstance.h
@@ -76,6 +76,8 @@ namespace UE::Dataflow
 			}
 		}
 
+		const FName GetVertexAttributeToVisualize() const;
+
 	private:
 		template <typename T>
 		const T& GetOutputValue(FName Name, const T& Default = T()) const
diff --git a/Engine/Plugins/Experimental/Dataflow/Source/DataflowNodes/Private/Dataflow/DataflowCollectionAddScalarVertexPropertyNode.cpp b/Engine/Plugins/Experimental/Dataflow/Source/DataflowNodes/Private/Dataflow/DataflowCollectionAddScalarVertexPropertyNode.cpp
--- a/Engine/Plugins/Experimental/Dataflow/Source/DataflowNodes/Private/Dataflow/DataflowCollectionAddScalarVertexPropertyNode.cpp
+++ b/Engine/Plugins/Experimental/Dataflow/Source/DataflowNodes/Private/Dataflow/DataflowCollectionAddScalarVertexPropertyNode.cpp
@@ -417,4 +417,9 @@ void FDataflowCollectionAddScalarVertexPropertyNode::ExtractVertexWeights(const
 	}
 }
 
+FName FDataflowCollectionAddScalarVertexPropertyNode::GetVertexAttributeToVisualize(UE::Dataflow::FContext& Context, const FName OutputName, const FName RenderGroup) const
+{
+	return FName(GetWeightAttributeKey(Context).Attribute);
+}
+
 #undef LOCTEXT_NAMESPACE
diff --git a/Engine/Plugins/Experimental/Dataflow/Source/DataflowNodes/Private/Dataflow/DataflowTextureToAttributeNode.cpp b/Engine/Plugins/Experimental/Dataflow/Source/DataflowNodes/Private/Dataflow/DataflowTextureToAttributeNode.cpp
--- a/Engine/Plugins/Experimental/Dataflow/Source/DataflowNodes/Private/Dataflow/DataflowTextureToAttributeNode.cpp
+++ b/Engine/Plugins/Experimental/Dataflow/Source/DataflowNodes/Private/Dataflow/DataflowTextureToAttributeNode.cpp
@@ -1,6 +1,8 @@
 // Copyright Epic Games, Inc. All Rights Reserved.
 
 #include "Dataflow/DataflowTextureToAttributeNode.h"
+
+#include "Components/DynamicMeshComponent.h"
 #include "Dataflow/DataflowNodeFactory.h"
 #include "Engine/Texture2D.h"
 #include "GeometryCollection/GeometryCollection.h"
@@ -169,5 +171,10 @@ void FDataflowTextureToAttributeNode::Evaluate(UE::Dataflow::FContext& Context,
 	}
 }
 
+FName FDataflowTextureToAttributeNode::GetVertexAttributeToVisualize(UE::Dataflow::FContext& Context, const FName OutputName, const FName RenderGroup) const
+{
+	return FName(GetValue(Context, &AttributeName, AttributeName));
+}
+
 #undef LOCTEXT_NAMESPACE 
 
diff --git a/Engine/Plugins/Experimental/Dataflow/Source/DataflowNodes/Private/Dataflow/DataflowTextureToAttributeNode.h b/Engine/Plugins/Experimental/Dataflow/Source/DataflowNodes/Private/Dataflow/DataflowTextureToAttributeNode.h
--- a/Engine/Plugins/Experimental/Dataflow/Source/DataflowNodes/Private/Dataflow/DataflowTextureToAttributeNode.h
+++ b/Engine/Plugins/Experimental/Dataflow/Source/DataflowNodes/Private/Dataflow/DataflowTextureToAttributeNode.h
@@ -46,6 +46,7 @@ private:
 	FScalarVertexPropertyGroup VertexGroup;
 
 	virtual void Evaluate(UE::Dataflow::FContext& Context, const FDataflowOutput* Out) const override;
+	virtual FName GetVertexAttributeToVisualize(UE::Dataflow::FContext& Context, const FName OutputName, const FName RenderGroup) const override;
 };
 
 namespace UE::Dataflow
diff --git a/Engine/Plugins/Experimental/Dataflow/Source/DataflowNodes/Public/Dataflow/DataflowCollectionAddScalarVertexPropertyNode.h b/Engine/Plugins/Experimental/Dataflow/Source/DataflowNodes/Public/Dataflow/DataflowCollectionAddScalarVertexPropertyNode.h
--- a/Engine/Plugins/Experimental/Dataflow/Source/DataflowNodes/Public/Dataflow/DataflowCollectionAddScalarVertexPropertyNode.h
+++ b/Engine/Plugins/Experimental/Dataflow/Source/DataflowNodes/Public/Dataflow/DataflowCollectionAddScalarVertexPropertyNode.h
@@ -162,6 +162,7 @@ private:
 	static constexpr float ReplaceChangedPassthroughValue = UE_BIG_NUMBER;
 	
 	UE_API virtual void Evaluate(UE::Dataflow::FContext& Context, const FDataflowOutput* Out) const override;
+	virtual FName GetVertexAttributeToVisualize(UE::Dataflow::FContext& Context, const FName OutputName, const FName RenderGroup) const override;
 
 	//~ Begin - FDataflowVertexAttributeEditableNode API
 	virtual void GetSupportedViewModes(UE::Dataflow::FContext& Context, TArray<FName>& OutViewModeNames) const override;
diff --git a/Engine/Source/Runtime/Experimental/Dataflow/Core/Public/Dataflow/DataflowNode.h b/Engine/Source/Runtime/Experimental/Dataflow/Core/Public/Dataflow/DataflowNode.h
--- a/Engine/Source/Runtime/Experimental/Dataflow/Core/Public/Dataflow/DataflowNode.h
+++ b/Engine/Source/Runtime/Experimental/Dataflow/Core/Public/Dataflow/DataflowNode.h
@@ -719,6 +719,13 @@ PRAGMA_ENABLE_DEPRECATION_WARNINGS
 	// returns true to disable node properties in its detail view when the corresponding input is connected ( default is true )
 	virtual bool MakeConnectedPropertiesReadOnly() const { return true; }
 
+	/**
+	* return a vertex attribute to visualize
+	* This attribute will be used as vertex color and displayed in the viewport 
+	* this is only used for collection based visualizations
+	*/
+	virtual FName GetVertexAttributeToVisualize(UE::Dataflow::FContext& Context, const FName OutputName, const FName RenderGroup) const { return NAME_None; }
+
 	// this method is called when the construction viewport render a specific output in the viewport 
 	// this allow the node to be able to set extra properties on the component or set it up differently 
 	virtual void OnRenderOutput(UE::Dataflow::FContext& Context, const FName OutputName, const FName RenderGroup, const TArray<UPrimitiveComponent*>& RenderComponents) const {};
