Dataflow paint tool : fix mirror direction and adjust min max value for various brush mode

#Synced-CL 49237557

diff --git a/Engine/Plugins/Experimental/Dataflow/Source/DataflowEditor/Private/DataflowEditorTools/DataflowEditorVertexAttributePaintTool.cpp b/Engine/Plugins/Experimental/Dataflow/Source/DataflowEditor/Private/DataflowEditorTools/DataflowEditorVertexAttributePaintTool.cpp
--- a/Engine/Plugins/Experimental/Dataflow/Source/DataflowEditor/Private/DataflowEditorTools/DataflowEditorVertexAttributePaintTool.cpp
+++ b/Engine/Plugins/Experimental/Dataflow/Source/DataflowEditor/Private/DataflowEditorTools/DataflowEditorVertexAttributePaintTool.cpp
@@ -2114,7 +2114,7 @@ void UDataflowEditorVertexAttributePaintTool::DrawHUD(FCanvas* Canvas, IToolsCon
 			RenderAPI->GetSceneView()->ScreenToPixel(ScreenBrushPos, ScreenBrushPos2D);
 
 			// Render value under at brush even in mesh mode
-			const FText BaseText = LOCTEXT("DataflowWeightMapPaintTool_ValueAtBrush_TextFormat", "Current Value: {0}");
+			const FText BaseText = LOCTEXT("DataflowWeightMapPaintTool_ValueAtBrush_TextFormat", "Value at brush: {0}");
 			const FText ValueAtBrushText = FText::Format(BaseText, FText::AsNumber(ToolProperties->BrushProperties.ValueAtBrush));
 			FCanvasTextItem TextItem(ScreenBrushPos2D, ValueAtBrushText, GEngine->GetLargeFont(), FLinearColor::White);
 			TextItem.EnableShadow(FLinearColor::Black);
@@ -2732,7 +2732,7 @@ void UDataflowEditorVertexAttributePaintTool::FMirrorData::FindMirroredIndices(c
 		if (Mesh.IsVertex(TargetVertexID))
 		{
 			const FVector TargetPosition = Mesh.GetVertex(TargetVertexID);
-			if (IsPointOnTargetMirrorSide(TargetPosition))
+			if (!IsPointOnTargetMirrorSide(TargetPosition))
 			{
 				OutVerticesToUpdate[Index] = INDEX_NONE;
 			}
diff --git a/Engine/Plugins/Experimental/Dataflow/Source/DataflowEditor/Private/DataflowEditorTools/DataflowEditorVertexAttributePaintToolDetailCustomization.cpp b/Engine/Plugins/Experimental/Dataflow/Source/DataflowEditor/Private/DataflowEditorTools/DataflowEditorVertexAttributePaintToolDetailCustomization.cpp
--- a/Engine/Plugins/Experimental/Dataflow/Source/DataflowEditor/Private/DataflowEditorTools/DataflowEditorVertexAttributePaintToolDetailCustomization.cpp
+++ b/Engine/Plugins/Experimental/Dataflow/Source/DataflowEditor/Private/DataflowEditorTools/DataflowEditorVertexAttributePaintToolDetailCustomization.cpp
@@ -633,6 +633,39 @@ namespace UE::Dataflow::Editor
 
 	void FVertexAttributePaintToolDetailCustomization::AddBrushValueRow(IDetailCategoryBuilder& BrushCategory) const
 	{
+		auto GetMinValue = [this]() -> float
+			{
+				if (ToolProperties.IsValid())
+				{
+					switch (ToolProperties->BrushProperties.BrushMode)
+					{
+					case EDataflowEditorToolEditOperation::Add: return -1.0f;
+					case EDataflowEditorToolEditOperation::Multiply: return 0.0f;
+					case EDataflowEditorToolEditOperation::Replace: return 0.0f;
+					case EDataflowEditorToolEditOperation::Invert: return 0.0f;
+					case EDataflowEditorToolEditOperation::Relax: return 0.0f;
+					}
+				}
+				return 0.0f;
+			};
+
+		auto GetMaxValue = [this]() -> float
+			{
+				if (ToolProperties.IsValid())
+				{
+					switch (ToolProperties->BrushProperties.BrushMode)
+					{
+					case EDataflowEditorToolEditOperation::Add: return 1.0f;
+					case EDataflowEditorToolEditOperation::Multiply: return 10.f;
+					case EDataflowEditorToolEditOperation::Replace: return 1.0f;
+					case EDataflowEditorToolEditOperation::Invert: return 1.0f;
+					case EDataflowEditorToolEditOperation::Relax: return 1.0f;
+					}
+				}
+				return 1.0f;
+			};
+
+
 		BrushCategory.AddCustomRow(LOCTEXT("ValueCategory", "Value"), false)
 		.NameContent()
 		[
@@ -644,20 +677,19 @@ namespace UE::Dataflow::Editor
 		.ValueContent()
 		[
 			SNew(SSpinBox<float>)
-			.MinValue(0.f)
-			.MaxValue(1.0f)
-			.MaxSliderValue(1.f)
+			.MinValue_Lambda(GetMinValue)
+			.MaxValue_Lambda(GetMaxValue)
 			.Value(1.0f)
 			.SupportDynamicSliderMaxValue(false)
 			.Value_Lambda([this]()
 				{
 					return ToolProperties.IsValid() ? ToolProperties->BrushProperties.AttributeValue : 1.f;
 				})
-			.OnValueChanged_Lambda([this](float NewValue)
+			.OnValueChanged_Lambda([this, GetMinValue, GetMaxValue](float NewValue)
 				{
 					if (ToolProperties.IsValid())
 					{
-						ToolProperties->BrushProperties.AttributeValue = FMath::Clamp(NewValue, 0.f, 1.f);
+						ToolProperties->BrushProperties.AttributeValue = FMath::Clamp(NewValue, GetMinValue(), GetMaxValue());
 						FPropertyChangedEvent PropertyChangedEvent(UBrushBaseProperties::StaticClass()->FindPropertyByName(GET_MEMBER_NAME_CHECKED(FDataflowEditorVertexAttributePaintToolBrushProperties, AttributeValue)));
 						ToolProperties->PostEditChangeProperty(PropertyChangedEvent);
 					}
@@ -1307,7 +1339,7 @@ namespace UE::Dataflow::Editor
 	{
 		TSharedRef<SSpinBox<float>> AddValueWidget =
 			SNew(SSpinBox<float>)
-			.MinValue(0)
+			.MinValue(-1)
 			.MaxValue(1)
 			.Value(1)
 			.ToolTipText(LOCTEXT("AddValueSliderToolTip", "Adjust the value to Add to the selected vertices."));
