Add value sampler tyeps and node + sampler to attribute node

#Synced-CL 49268397

diff --git a/Engine/Plugins/Experimental/Dataflow/Source/DataflowNodes/Private/Dataflow/DataflowNodesPlugin.cpp b/Engine/Plugins/Experimental/Dataflow/Source/DataflowNodes/Private/Dataflow/DataflowNodesPlugin.cpp
--- a/Engine/Plugins/Experimental/Dataflow/Source/DataflowNodes/Private/Dataflow/DataflowNodesPlugin.cpp
+++ b/Engine/Plugins/Experimental/Dataflow/Source/DataflowNodes/Private/Dataflow/DataflowNodesPlugin.cpp
@@ -16,6 +16,9 @@
 #include "Dataflow/DataflowSkeletonAssetTerminalNode.h"
 #include "Dataflow/DataflowTextureToAttributeNode.h"
 #include "Dataflow/DataflowVertexColorToAttributeNode.h"
+#include "Dataflow/DataflowValueSamplerToAttributeNode.h"
+#include "Dataflow/SamplerNodes/DataflowSamplerTypes.h"
+#include "Dataflow/SamplerNodes/DataflowGradientSamplerNode.h"
 
 #define LOCTEXT_NAMESPACE "DataflowNodes"
 
@@ -47,12 +50,17 @@ const FName FGeometryCollectionAddScalarVertexPropertyCallbacks::Name = FName("F
 
 void IDataflowNodesPlugin::StartupModule()
 {
+	UE::Dataflow::RegisterValueSamplerTypes();
+
 	UE::Dataflow::RegisterSkeletalMeshNodes();
 	UE::Dataflow::RegisterStaticMeshNodes();
 	UE::Dataflow::RegisterSelectionNodes();
 	UE::Dataflow::RegisterContextOverridesNodes();
 	UE::Dataflow::RegisterTextureToAttributeNodes();
 	UE::Dataflow::RegisterVertexColorToAttributeNodes();
+	UE::Dataflow::RegisterValueSamplerToAttributeNodes();
+	UE::Dataflow::RegisterGradientValueSamplerNodes();
+
 	UE::Dataflow::DataflowCollectionAttributeKeyNodes();
 	DATAFLOW_NODE_REGISTER_CREATION_FACTORY(FDataflowCollectionAddScalarVertexPropertyNode);
 	DATAFLOW_NODE_REGISTER_CREATION_FACTORY(FDataflowCollectionEditSkinWeightsNode);
diff --git a/Engine/Plugins/Experimental/Dataflow/Source/DataflowNodes/Private/Dataflow/DataflowValueSamplerToAttributeNode.cpp b/Engine/Plugins/Experimental/Dataflow/Source/DataflowNodes/Private/Dataflow/DataflowValueSamplerToAttributeNode.cpp
new file mode 100644
--- /dev/null
+++ b/Engine/Plugins/Experimental/Dataflow/Source/DataflowNodes/Private/Dataflow/DataflowValueSamplerToAttributeNode.cpp
@@ -0,0 +1,90 @@
+// Copyright Epic Games, Inc. All Rights Reserved.
+
+#include "Dataflow/DataflowValueSamplerToAttributeNode.h"
+
+#include "Dataflow/DataflowNodeFactory.h"
+#include "GeometryCollection/GeometryCollection.h"
+#include "GeometryCollection/Facades/CollectionMeshFacade.h"
+
+#include UE_INLINE_GENERATED_CPP_BY_NAME(DataflowValueSamplerToAttributeNode)
+
+#define LOCTEXT_NAMESPACE "DataflowValueSamplerToAttributeNode"
+
+namespace UE::Dataflow
+{
+	void RegisterValueSamplerToAttributeNodes()
+	{
+		DATAFLOW_NODE_REGISTER_CREATION_FACTORY(FDataflowValueSamplerToAttributeNode);
+	}
+};
+
+FDataflowValueSamplerToAttributeNode::FDataflowValueSamplerToAttributeNode(const UE::Dataflow::FNodeParameters& InParam, FGuid InGuid)
+	: FDataflowNode(InParam, InGuid)
+{
+	RegisterInputConnection(&Collection);
+	RegisterInputConnection(&Sampler);
+	RegisterInputConnection(&AttributeName);
+	RegisterOutputConnection(&Collection, &Collection);
+}
+
+void FDataflowValueSamplerToAttributeNode::Evaluate(UE::Dataflow::FContext& Context, const FDataflowOutput* Out) const
+{
+	if (Out->IsA(&Collection))
+	{
+		const FName InAttributeName = FName(GetValue(Context, &AttributeName, AttributeName));
+		if (InAttributeName.IsNone())
+		{
+			Context.Error(LOCTEXT("NoAttributeName_Msg", "Attribute name is empty, data will not be set"));
+			SafeForwardInput(Context, &Collection, &Collection);
+			return;
+		}
+
+		FManagedArrayCollection InCollection = GetValue(Context, &Collection, Collection);
+
+		const FName InVertexGroupName = VertexGroup.Name;
+		if (!InCollection.HasGroup(InVertexGroupName))
+		{
+			Context.Error(LOCTEXT("InvalidVertexGroup_Msg", "Vertex group does not exists, data will not be set"));
+			SafeForwardInput(Context, &Collection, &Collection);
+			return;
+		}
+
+		// Create or get the new attribute 
+		TManagedArray<float>& TargetAttribute = InCollection.AddAttribute<float>(InAttributeName, InVertexGroupName);
+		TargetAttribute.Fill(0);
+
+		const TManagedArray<FVector3f>* VertexAttribute = InCollection.FindAttribute<FVector3f>(FGeometryCollection::VertexPositionAttribute, InVertexGroupName);
+		if (!VertexAttribute)
+		{
+			// cannot really include cloth here so we need to directly serach for the vertex color attribute in the render vertices ( eventually we may converge to use only geometry collections )
+			VertexAttribute = InCollection.FindAttribute<FVector3f>("RenderPosition", InVertexGroupName);
+		}
+		if (!VertexAttribute)
+		{
+			// cannot really include cloth here so we need to directly serach for the vertex color attribute in the render vertices ( eventually we may converge to use only geometry collections )
+			VertexAttribute = InCollection.FindAttribute<FVector3f>("SimPosition3D", InVertexGroupName);
+		}
+		if (!VertexAttribute)
+		{
+			Context.Warning(LOCTEXT("NoVertexPosition_Msg", "No vertex position could be found in this collection, no data will be set and the attribute will be inititialized to zeros"));
+		}
+		else
+		{
+			if (const FDataflowValueSampler* InSampler = GetValue(Context, &Sampler).GetConstStructView().GetPtr<FDataflowValueSampler>())
+			{
+				TArrayView<float> Values = MakeArrayView<float>(TargetAttribute.GetData(), TargetAttribute.Num());
+				InSampler->Sample(VertexAttribute->GetConstArray(), Values);
+			}
+		}
+		
+		SetValue(Context, MoveTemp(InCollection), &Collection);
+	}
+}
+
+FName FDataflowValueSamplerToAttributeNode::GetVertexAttributeToVisualize(UE::Dataflow::FContext& Context, const FName OutputName, const FName RenderGroup) const
+{
+	return FName(GetValue(Context, &AttributeName, AttributeName));
+}
+
+#undef LOCTEXT_NAMESPACE 
+
diff --git a/Engine/Plugins/Experimental/Dataflow/Source/DataflowNodes/Private/Dataflow/DataflowValueSamplerToAttributeNode.h b/Engine/Plugins/Experimental/Dataflow/Source/DataflowNodes/Private/Dataflow/DataflowValueSamplerToAttributeNode.h
new file mode 100644
--- /dev/null
+++ b/Engine/Plugins/Experimental/Dataflow/Source/DataflowNodes/Private/Dataflow/DataflowValueSamplerToAttributeNode.h
@@ -0,0 +1,49 @@
+// Copyright Epic Games, Inc. All Rights Reserved.
+
+#pragma once 
+
+#include "Dataflow/DataflowNode.h"
+#include "GeometryCollection/ManagedArrayCollection.h"
+#include "Dataflow/DataflowCollectionAddScalarVertexPropertyNode.h"
+#include "Dataflow/SamplerNodes/DataflowSamplerTypes.h"
+
+#include "DataflowValueSamplerToAttributeNode.generated.h"
+
+/**
+* Use a value sampler to generate the value of an attribute
+*/
+USTRUCT()
+struct FDataflowValueSamplerToAttributeNode : public FDataflowNode
+{
+	GENERATED_USTRUCT_BODY()
+	DATAFLOW_NODE_DEFINE_INTERNAL(FDataflowValueSamplerToAttributeNode, "ValueSamplerToAttribute", "Dataflow", "Transfer Generate")
+
+public:
+	FDataflowValueSamplerToAttributeNode(const UE::Dataflow::FNodeParameters& InParam, FGuid InGuid = FGuid::NewGuid());
+
+private:
+	/** Collection to set the attribute data to */
+	UPROPERTY(meta = (DataflowInput, DataflowOutput, DataflowPassthrough = Collection))
+	FManagedArrayCollection Collection;
+
+	/** Sampler to use to set the value of the attribute */
+	UPROPERTY(EditAnywhere, Category = Source, meta = (DataflowInput))
+	FDataflowValueSamplerTypes Sampler;
+
+	/** Name of the attribute to set the value on */
+	UPROPERTY(EditAnywhere, Category = Target, meta = (DataflowInput))
+	FString AttributeName;
+
+	/** Vertex group to use */
+	UPROPERTY(EditAnywhere, Category = Target)
+	FScalarVertexPropertyGroup VertexGroup;
+
+	virtual void Evaluate(UE::Dataflow::FContext& Context, const FDataflowOutput* Out) const override;
+	virtual FName GetVertexAttributeToVisualize(UE::Dataflow::FContext& Context, const FName OutputName, const FName RenderGroup) const override;
+};
+
+namespace UE::Dataflow
+{
+	void RegisterValueSamplerToAttributeNodes();
+}
+
diff --git a/Engine/Plugins/Experimental/Dataflow/Source/DataflowNodes/Private/Dataflow/SamplerNodes/DataflowGradientSamplerNode.cpp b/Engine/Plugins/Experimental/Dataflow/Source/DataflowNodes/Private/Dataflow/SamplerNodes/DataflowGradientSamplerNode.cpp
new file mode 100644
--- /dev/null
+++ b/Engine/Plugins/Experimental/Dataflow/Source/DataflowNodes/Private/Dataflow/SamplerNodes/DataflowGradientSamplerNode.cpp
@@ -0,0 +1,125 @@
+// Copyright Epic Games, Inc. All Rights Reserved.
+
+#include "Dataflow/SamplerNodes/DataflowGradientSamplerNode.h"
+#include "Dataflow/DataflowNodeFactory.h"
+
+namespace UE::Dataflow
+{
+	void RegisterGradientValueSamplerNodes()
+	{
+		DATAFLOW_NODE_REGISTER_CREATION_FACTORY(FDataflowLinearGradientValueSamplerNode);
+		DATAFLOW_NODE_REGISTER_CREATION_FACTORY(FDataflowRadialGradientValueSamplerNode);
+	}
+}
+
+////////////////////////////////////////////////////////////////////////////////////////////////////////////////////////
+
+void FDataflowLinearGradientValueSampler::Sample(TArrayView<const FVector3f> Positions, TArrayView<float> OutValues) const
+{
+	const FVector3f GradientDirection = EndPoint - StartPoint;
+	const double SquareLength = GradientDirection.SquaredLength();
+
+	if (SquareLength == UE_SMALL_NUMBER)
+	{
+		for (float& OutValue : OutValues)
+		{
+			OutValue = StartValue;
+		}
+		return;
+	}
+
+	const float InvSquareLength = 1.f / SquareLength;
+
+	const int32 NumSamples = FMath::Min(Positions.Num(), OutValues.Num());
+	for (int32 Index = 0; Index < NumSamples; ++Index)
+	{
+		const FVector3f LocalPosition = Positions[Index] - StartPoint;
+		const float Alpha = LocalPosition.Dot(GradientDirection) * InvSquareLength;
+		const float Value = FMath::Lerp(StartValue, EndValue, Alpha);
+		OutValues[Index] = (bClamp)
+			? FMath::Clamp(Value, FMath::Min(StartValue, EndValue), FMath::Max(StartValue, EndValue))
+			: Value;
+	}
+}
+
+FDataflowLinearGradientValueSamplerNode::FDataflowLinearGradientValueSamplerNode(const UE::Dataflow::FNodeParameters& InParam, FGuid InGuid)
+	: FDataflowNode(InParam, InGuid)
+{
+	RegisterInputConnection(&Gradient.StartPoint, GET_MEMBER_NAME_CHECKED(FDataflowLinearGradientValueSampler, StartPoint));
+	RegisterInputConnection(&Gradient.StartValue, GET_MEMBER_NAME_CHECKED(FDataflowLinearGradientValueSampler, StartValue));
+	RegisterInputConnection(&Gradient.EndPoint, GET_MEMBER_NAME_CHECKED(FDataflowLinearGradientValueSampler, EndPoint));
+	RegisterInputConnection(&Gradient.EndValue, GET_MEMBER_NAME_CHECKED(FDataflowLinearGradientValueSampler, EndValue));
+	RegisterInputConnection(&Gradient.bClamp, GET_MEMBER_NAME_CHECKED(FDataflowLinearGradientValueSampler, bClamp));
+	RegisterOutputConnection(&Gradient);
+}
+
+void FDataflowLinearGradientValueSamplerNode::Evaluate(UE::Dataflow::FContext& Context, const FDataflowOutput* Out) const
+{
+	if (Out->IsA(&Gradient))
+	{
+		FDataflowLinearGradientValueSampler OutGradient;
+		OutGradient.StartPoint = GetValue(Context, &Gradient.StartPoint, Gradient.StartPoint);
+		OutGradient.StartValue = GetValue(Context, &Gradient.StartValue, Gradient.StartValue);
+		OutGradient.EndPoint = GetValue(Context, &Gradient.EndPoint, Gradient.EndPoint);
+		OutGradient.EndValue = GetValue(Context, &Gradient.EndValue, Gradient.EndValue);
+		OutGradient.bClamp = GetValue(Context, &Gradient.bClamp, Gradient.bClamp);
+		SetValue(Context, OutGradient, &Gradient);
+	}
+}
+
+////////////////////////////////////////////////////////////////////////////////////////////////////////////////////////
+
+void FDataflowRadialGradientValueSampler::Sample(TArrayView<const FVector3f> Positions, TArrayView<float> OutValues) const
+{
+	if (Radius == UE_SMALL_NUMBER)
+	{
+		for (float& OutValue : OutValues)
+		{
+			OutValue = EdgeValue;
+		}
+		return;
+	}
+
+	const float InvRadius = 1.f / Radius;
+
+	const int32 NumSamples = FMath::Min(Positions.Num(), OutValues.Num());
+	for (int32 Index = 0; Index < NumSamples; ++Index)
+	{
+		const FVector3f GradientVector = Positions[Index] - Center;
+		const float Alpha = GradientVector.Size() * InvRadius;
+
+		const float Value = FMath::Lerp(CenterValue, EdgeValue, Alpha);
+		OutValues[Index] = (bClamp)
+			? FMath::Clamp(Value, FMath::Min(CenterValue, EdgeValue), FMath::Max(CenterValue, EdgeValue))
+			: Value;
+	}
+}
+
+FDataflowRadialGradientValueSamplerNode::FDataflowRadialGradientValueSamplerNode(const UE::Dataflow::FNodeParameters& InParam, FGuid InGuid)
+	: FDataflowNode(InParam, InGuid)
+{
+	RegisterInputConnection(&Gradient.Center, GET_MEMBER_NAME_CHECKED(FDataflowRadialGradientValueSampler, Center));
+	RegisterInputConnection(&Gradient.Radius, GET_MEMBER_NAME_CHECKED(FDataflowRadialGradientValueSampler, Radius));
+	RegisterInputConnection(&Gradient.CenterValue, GET_MEMBER_NAME_CHECKED(FDataflowRadialGradientValueSampler, CenterValue));
+	RegisterInputConnection(&Gradient.EdgeValue, GET_MEMBER_NAME_CHECKED(FDataflowRadialGradientValueSampler, EdgeValue));
+	RegisterInputConnection(&Gradient.bClamp, GET_MEMBER_NAME_CHECKED(FDataflowRadialGradientValueSampler, bClamp));
+	RegisterOutputConnection(&Gradient);
+}
+
+void FDataflowRadialGradientValueSamplerNode::Evaluate(UE::Dataflow::FContext& Context, const FDataflowOutput* Out) const
+{
+	if (Out->IsA(&Gradient))
+	{
+		FDataflowRadialGradientValueSampler OutGradient;
+		OutGradient.Center = GetValue(Context, &Gradient.Center, Gradient.Center);
+		OutGradient.Radius = GetValue(Context, &Gradient.Radius, Gradient.Radius);
+		OutGradient.CenterValue = GetValue(Context, &Gradient.CenterValue, Gradient.CenterValue);
+		OutGradient.EdgeValue = GetValue(Context, &Gradient.EdgeValue, Gradient.EdgeValue);
+		OutGradient.bClamp = GetValue(Context, &Gradient.bClamp, Gradient.bClamp);
+		SetValue(Context, OutGradient, &Gradient);
+	}
+}
+
+////////////////////////////////////////////////////////////////////////////////////////////////////////////////////////
+
+
diff --git a/Engine/Plugins/Experimental/Dataflow/Source/DataflowNodes/Private/Dataflow/SamplerNodes/DataflowGradientSamplerNode.h b/Engine/Plugins/Experimental/Dataflow/Source/DataflowNodes/Private/Dataflow/SamplerNodes/DataflowGradientSamplerNode.h
new file mode 100644
--- /dev/null
+++ b/Engine/Plugins/Experimental/Dataflow/Source/DataflowNodes/Private/Dataflow/SamplerNodes/DataflowGradientSamplerNode.h
@@ -0,0 +1,117 @@
+// Copyright Epic Games, Inc. All Rights Reserved.
+
+#pragma once 
+
+#include "Dataflow/DataflowNode.h"
+#include "Dataflow/SamplerNodes/DataflowSamplerTypes.h"
+
+#include "DataflowGradientSamplerNode.generated.h"
+
+////////////////////////////////////////////////////////////////////////////////////////////////////////////////////////
+
+USTRUCT()
+struct FDataflowLinearGradientValueSampler : public FDataflowValueSampler
+{
+	GENERATED_BODY();
+
+	virtual ~FDataflowLinearGradientValueSampler() override = default;
+	virtual void Sample(TArrayView<const FVector3f> Positions, TArrayView<float> OutValues) const override;
+
+public:
+	// Point in space where the start value is defined
+	UPROPERTY(EditAnywhere, Category = Gradient, meta = (DataflowInput))
+	FVector3f StartPoint = { 0, 0, 100 };
+
+	// Value defined at start point
+	UPROPERTY(EditAnywhere, Category = Gradient, meta = (DataflowInput))
+	float StartValue = 1.0;
+
+	// Point in space where end value is defined
+	UPROPERTY(EditAnywhere, Category = Gradient, meta = (DataflowInput))
+	FVector3f EndPoint = { 0, 0, 0 };
+
+	// Value defined at end point
+	UPROPERTY(EditAnywhere, Category = Gradient, meta = (DataflowInput))
+	float EndValue = 0.0;
+
+	// If true generated value will be clamped beyond end and start point, otherwise, they will be extrapolated
+	UPROPERTY(EditAnywhere, Category = Gradient, meta = (DataflowInput))
+	bool bClamp = true;
+};
+
+/**
+ * Output a linear gradient sampler 
+ */
+USTRUCT()
+struct FDataflowLinearGradientValueSamplerNode : public FDataflowNode
+{
+	GENERATED_BODY()
+	DATAFLOW_NODE_DEFINE_INTERNAL(FDataflowLinearGradientValueSamplerNode, "Linear Gradient Sampler", "Samplers", "")
+
+public:
+	FDataflowLinearGradientValueSamplerNode(const UE::Dataflow::FNodeParameters& InParam, FGuid InGuid = FGuid::NewGuid());
+
+private:
+	UPROPERTY(EditAnywhere, Category = Sampler, meta = (DataflowOutput, ShowOnlyInnerProperties))
+	FDataflowLinearGradientValueSampler Gradient;
+
+	void Evaluate(UE::Dataflow::FContext& Context, const FDataflowOutput* Out) const override;
+};
+
+////////////////////////////////////////////////////////////////////////////////////////////////////////////////////////
+
+USTRUCT()
+struct FDataflowRadialGradientValueSampler : public FDataflowValueSampler
+{
+	GENERATED_BODY();
+
+	virtual ~FDataflowRadialGradientValueSampler() override = default;
+	virtual void Sample(TArrayView<const FVector3f> Positions, TArrayView<float> OutValues) const override;
+
+public:
+	// center defining the radial gradient
+	UPROPERTY(EditAnywhere, Category = Gradient, meta = (DataflowInput))
+	FVector3f Center = { 0, 0, 0 };
+
+	// Radius of the gradient
+	UPROPERTY(EditAnywhere, Category = Gradient, meta = (DataflowInput))
+	float Radius = 100.f;
+
+	// Value defined at center
+	UPROPERTY(EditAnywhere, Category = Gradient, meta = (DataflowInput))
+	float CenterValue = 1.0;
+
+	// Value defined at the radius distance from the center 
+	UPROPERTY(EditAnywhere, Category = Gradient, meta = (DataflowInput))
+	float EdgeValue = 0.0;
+
+	// If true generated values will be clamped beyond the radius otherwise they will they will be extrapolated
+	UPROPERTY(EditAnywhere, Category = Gradient, meta = (DataflowInput))
+	bool bClamp = true;
+};
+
+/**
+ * Output a radial gradient sampler
+ */
+USTRUCT()
+struct FDataflowRadialGradientValueSamplerNode : public FDataflowNode
+{
+	GENERATED_BODY()
+	DATAFLOW_NODE_DEFINE_INTERNAL(FDataflowRadialGradientValueSamplerNode, "Radial Gradient Sampler", "Samplers", "")
+
+public:
+	FDataflowRadialGradientValueSamplerNode(const UE::Dataflow::FNodeParameters& InParam, FGuid InGuid = FGuid::NewGuid());
+
+private:
+	UPROPERTY(EditAnywhere, Category = Sampler, meta = (DataflowOutput, ShowOnlyInnerProperties))
+	FDataflowRadialGradientValueSampler Gradient;
+
+	void Evaluate(UE::Dataflow::FContext& Context, const FDataflowOutput* Out) const override;
+};
+
+////////////////////////////////////////////////////////////////////////////////////////////////////////////////////////
+
+namespace UE::Dataflow
+{
+	void RegisterGradientValueSamplerNodes();
+}
\ No newline at end of file
diff --git a/Engine/Plugins/Experimental/Dataflow/Source/DataflowNodes/Private/Dataflow/SamplerNodes/DataflowSamplerTypes.cpp b/Engine/Plugins/Experimental/Dataflow/Source/DataflowNodes/Private/Dataflow/SamplerNodes/DataflowSamplerTypes.cpp
new file mode 100644
--- /dev/null
+++ b/Engine/Plugins/Experimental/Dataflow/Source/DataflowNodes/Private/Dataflow/SamplerNodes/DataflowSamplerTypes.cpp
@@ -0,0 +1,12 @@
+// Copyright Epic Games, Inc. All Rights Reserved.
+
+#include "Dataflow/SamplerNodes/DataflowSamplerTypes.h"
+#include "Dataflow/DataflowAnyTypeRegistry.h"
+
+namespace UE::Dataflow
+{
+	void RegisterValueSamplerTypes()
+	{
+		UE_DATAFLOW_REGISTER_ANYTYPE(FDataflowValueSamplerTypes);
+	}
+}
diff --git a/Engine/Plugins/Experimental/Dataflow/Source/DataflowNodes/Private/Dataflow/SamplerNodes/DataflowSamplerTypes.h b/Engine/Plugins/Experimental/Dataflow/Source/DataflowNodes/Private/Dataflow/SamplerNodes/DataflowSamplerTypes.h
new file mode 100644
--- /dev/null
+++ b/Engine/Plugins/Experimental/Dataflow/Source/DataflowNodes/Private/Dataflow/SamplerNodes/DataflowSamplerTypes.h
@@ -0,0 +1,33 @@
+// Copyright Epic Games, Inc. All Rights Reserved.
+
+#pragma once 
+
+#include "Dataflow/DataflowAnyType.h"
+
+#include "DataflowSamplerTypes.generated.h"
+
+/**
+* A Value sampler, sample at a specific 3D point and return a single value
+*/
+USTRUCT()
+struct FDataflowValueSampler
+{
+	GENERATED_BODY();
+public:
+	virtual ~FDataflowValueSampler() = default;
+	virtual void Sample(TArrayView<const FVector3f> Positions, TArrayView<float> OutValues) const { check(false); }
+};
+
+/** Value sampler anytype */
+USTRUCT()
+struct FDataflowValueSamplerTypes : public FDataflowAnyType
+{
+	using FPolicyType = FDataflowUDerivedStructTypePolicy<FDataflowValueSampler>;
+	using FStorageType = void;
+	GENERATED_BODY()
+};
+
+namespace UE::Dataflow
+{
+	void RegisterValueSamplerTypes();
+}
