Add support for transform selection rendering

#Synced-CL 49268397

diff --git a/Engine/Plugins/Experimental/Dataflow/Source/DataflowEditor/Private/DataflowRendering/DataflowGeometryCollectionRenderableType.cpp b/Engine/Plugins/Experimental/Dataflow/Source/DataflowEditor/Private/DataflowRendering/DataflowGeometryCollectionRenderableType.cpp
--- a/Engine/Plugins/Experimental/Dataflow/Source/DataflowEditor/Private/DataflowRendering/DataflowGeometryCollectionRenderableType.cpp
+++ b/Engine/Plugins/Experimental/Dataflow/Source/DataflowEditor/Private/DataflowRendering/DataflowGeometryCollectionRenderableType.cpp
@@ -18,6 +18,7 @@
 #include "GeometryCollection/ManagedArrayCollection.h"
 
 #include "GeometryCollection/Facades/CollectionCurveFacade.h"
+#include "GeometryCollection/Facades/CollectionTransformSelectionFacade.h"
 
 #include "UObject/ObjectPtr.h"
 
@@ -96,6 +97,13 @@ namespace UE::Dataflow::Private
 					}
 					Instance.CacheValue(DynamicMeshes);
 				}
+
+				FDataflowSelection Selection;
+				GetSelection(Instance, Collection, Selection);
+				const bool bIsValidTransformSelection = (Selection.GetGroupName() == FTransformCollection::TransformGroup && Selection.IsValidForCollection(Collection));
+
+				UMaterialInterface* SelectionMaterial = LoadObject<UMaterialInterface>(nullptr, TEXT("/Engine/EditorMaterials/Dataflow/DataflowSelection"));
+
 				// only create components for the one that have vertices
 				for (int32 TransformIndex = 0; TransformIndex < TransformAttribute->Num(); ++TransformIndex)
 				{
@@ -108,13 +116,38 @@ namespace UE::Dataflow::Private
 						{
 							Component->SetCastShadow(false);
 							Component->SetMesh(MoveTemp(DynamicMeshes[TransformIndex]));
-							Component->SetOverrideRenderMaterial(FDataflowEditorStyle::Get().VertexMaterial);
+
+							if (bIsValidTransformSelection && Selection.IsSelected(TransformIndex))
+							{
+								Component->SetOverrideRenderMaterial(SelectionMaterial);
+							}
+							else
+							{
+								Component->SetOverrideRenderMaterial(FDataflowEditorStyle::Get().VertexMaterial);
+							}
 						}
 					}
 				}
 			}
 		}
 
+		void GetSelection(const FRenderableTypeInstance& Instance, const FManagedArrayCollection& Collection, FDataflowSelection& OutSelection) const
+		{
+			Instance.GetSelectionToVisualize(OutSelection);
+			if (!OutSelection.IsValidForCollection(Collection))
+			{
+				FDataflowSelection DefaultValue;
+				OutSelection = Instance.GetOutputValueByType<FDataflowSelectionTypes>(DefaultValue);
+			}
+			if (OutSelection.GetGroupName() == FTransformCollection::TransformGroup && OutSelection.AnySelected())
+			{
+				GeometryCollection::Facades::FCollectionTransformSelectionFacade SelectionFacade(Collection);
+				TArray<int32> SelectionToConvert = OutSelection.AsArray();
+				SelectionFacade.ConvertSelectionToRigidNodes(SelectionToConvert);
+				OutSelection.SetFromArray(SelectionToConvert);
+			}
+		}
+
 		template<typename T>
 		void TransferTypedAttributeVertexColor(FGeometryCollection& GeomCollection, FName ValueAttributeName) const
 		{
diff --git a/Engine/Plugins/Experimental/Dataflow/Source/DataflowEditor/Private/DataflowRendering/DataflowRenderableTypeInstance.cpp b/Engine/Plugins/Experimental/Dataflow/Source/DataflowEditor/Private/DataflowRendering/DataflowRenderableTypeInstance.cpp
--- a/Engine/Plugins/Experimental/Dataflow/Source/DataflowEditor/Private/DataflowRendering/DataflowRenderableTypeInstance.cpp
+++ b/Engine/Plugins/Experimental/Dataflow/Source/DataflowEditor/Private/DataflowRendering/DataflowRenderableTypeInstance.cpp
@@ -117,4 +117,12 @@ namespace UE::Dataflow
 		}
 		return NAME_None;
 	}
+
+	const void FRenderableTypeInstance::GetSelectionToVisualize(FDataflowSelection& OutSelection) const
+	{
+		if (RenderableType && Node)
+		{
+			return Node->GetSelectionToVisualize(EvaluationContext, OutputName, RenderableType->GetRenderGroup(), OutSelection);
+		}
+	}
 }
diff --git a/Engine/Plugins/Experimental/Dataflow/Source/DataflowEditor/Public/DataflowRendering/DataflowRenderableTypeInstance.h b/Engine/Plugins/Experimental/Dataflow/Source/DataflowEditor/Public/DataflowRendering/DataflowRenderableTypeInstance.h
--- a/Engine/Plugins/Experimental/Dataflow/Source/DataflowEditor/Public/DataflowRendering/DataflowRenderableTypeInstance.h
+++ b/Engine/Plugins/Experimental/Dataflow/Source/DataflowEditor/Public/DataflowRendering/DataflowRenderableTypeInstance.h
@@ -10,6 +10,7 @@
 
 class UPrimitiveComponent;
 class AActor;
+struct FDataflowSelection;
 
 #define UE_API DATAFLOWEDITOR_API
 
@@ -78,7 +79,8 @@ namespace UE::Dataflow
 
 		const FName GetVertexAttributeToVisualize() const;
 
-	private:
+		const void GetSelectionToVisualize(FDataflowSelection& OutSelection) const;
+
 		template <typename T>
 		const T& GetOutputValue(FName Name, const T& Default = T()) const
 		{
@@ -95,7 +97,41 @@ namespace UE::Dataflow
 			}
 			return Default;
 		}
+		/** Find the first output of a specific match type and return its value - non-anytype version */
+		template <typename T UE_REQUIRES(!std::is_base_of_v<FDataflowAnyType, T>)>
+		const T& GetOutputValueByType(const T& Default) const
+		{
+			if (Node)
+			{
+				for (const FDataflowOutput* Output : Node->GetOutputs())
+				{
+					if (Output && Output->GetType() == TDataflowPolicyTypeName<T>::GetName())
+					{
+						return GetOutputValue<T>(Output->GetName(), Default);
+					}
+				}
+			}
+			return Default;
+		}
 
+		/** Find the first output of a specific match type and return its value - anytype version */
+		template <typename T UE_REQUIRES(std::is_base_of_v<FDataflowAnyType, T>)>
+		const T::FStorageType& GetOutputValueByType(const T::FStorageType& Default) const
+		{
+			if (Node)
+			{
+				for (const FDataflowOutput* Output : Node->GetOutputs())
+				{
+					if (Output && T::FPolicyType::SupportsTypeStatic(Output->GetType()))
+					{
+						return GetOutputValue<typename T::FStorageType>(Output->GetName(), Default);
+					}
+				}
+			}
+			return Default;
+		}
+
+	private:
 		UE_API FContextCacheKey ComputeCacheKey() const;
 
 		FRenderableTypeInstance(FContext& InContext, const TSharedPtr<const FDataflowNode>& InNode, const FName InOutputName, const IRenderableType* InRenderableType)
diff --git a/Engine/Source/Runtime/Experimental/Dataflow/Core/Private/Dataflow/DataflowSelection.cpp b/Engine/Source/Runtime/Experimental/Dataflow/Core/Private/Dataflow/DataflowSelection.cpp
--- a/Engine/Source/Runtime/Experimental/Dataflow/Core/Private/Dataflow/DataflowSelection.cpp
+++ b/Engine/Source/Runtime/Experimental/Dataflow/Core/Private/Dataflow/DataflowSelection.cpp
@@ -149,6 +149,11 @@ void FDataflowSelection::Subtract(const FDataflowSelection& Other, FDataflowSele
 	Result.SelectionArray = TBitArray<>::BitwiseAND(SelectionArray, InvSelectionArray, EBitwiseOperatorFlags::MaxSize);
 }
 
+bool FDataflowSelection::IsValid() const
+{
+	return (!GroupName.IsNone() && SelectionArray.Num() > 0);
+}
+
 int32 FDataflowSelection::NumSelected() const
 {
 	return SelectionArray.CountSetBits(0, SelectionArray.Num());
@@ -197,7 +202,7 @@ void FDataflowSelection::SetSelected(TArray<int32> Indices)
 
 bool FDataflowSelection::IsValidForCollection(const FManagedArrayCollection& InCollection) const
 {
-	return Num() == InCollection.NumElements(GroupName);
+	return !GroupName.IsNone() && Num() == InCollection.NumElements(GroupName);
 }
 
 bool FDataflowSelection::InitFromArray(const FManagedArrayCollection& InCollection, const TArray<int32>& InSelectionArr)
diff --git a/Engine/Source/Runtime/Experimental/Dataflow/Core/Public/Dataflow/DataflowNode.h b/Engine/Source/Runtime/Experimental/Dataflow/Core/Public/Dataflow/DataflowNode.h
--- a/Engine/Source/Runtime/Experimental/Dataflow/Core/Public/Dataflow/DataflowNode.h
+++ b/Engine/Source/Runtime/Experimental/Dataflow/Core/Public/Dataflow/DataflowNode.h
@@ -726,6 +726,11 @@ PRAGMA_ENABLE_DEPRECATION_WARNINGS
 	*/
 	virtual FName GetVertexAttributeToVisualize(UE::Dataflow::FContext& Context, const FName OutputName, const FName RenderGroup) const { return NAME_None; }
 
+	/**
+	* Fills a dataflow selection object used to visualize the selection a node may output for a sepcific output and render group
+	*/
+	virtual void GetSelectionToVisualize(UE::Dataflow::FContext& Context, const FName OutputName, const FName RenderGroup, FDataflowSelection& OutSelection) const { }
+
 	// this method is called when the construction viewport render a specific output in the viewport 
 	// this allow the node to be able to set extra properties on the component or set it up differently 
 	virtual void OnRenderOutput(UE::Dataflow::FContext& Context, const FName OutputName, const FName RenderGroup, const TArray<UPrimitiveComponent*>& RenderComponents) const {};
diff --git a/Engine/Source/Runtime/Experimental/Dataflow/Core/Public/Dataflow/DataflowSelection.h b/Engine/Source/Runtime/Experimental/Dataflow/Core/Public/Dataflow/DataflowSelection.h
--- a/Engine/Source/Runtime/Experimental/Dataflow/Core/Public/Dataflow/DataflowSelection.h
+++ b/Engine/Source/Runtime/Experimental/Dataflow/Core/Public/Dataflow/DataflowSelection.h
@@ -27,6 +27,9 @@ struct FDataflowSelection
 	int32 Num() const { return SelectionArray.Num(); }
 	DATAFLOWCORE_API int32 NumSelected() const;
 	DATAFLOWCORE_API bool AnySelected() const;
+
+	bool IsValid() const;
+
 	inline bool IsValidIndex(int32 Idx) const
 	{
 		return SelectionArray.IsValidIndex(Idx);
@@ -75,6 +78,8 @@ struct FDataflowSelection
 	*/
 	DATAFLOWCORE_API FString ToString();
 
+	FName GetGroupName() const { return GroupName; }
+
 private:
 	FName GroupName;
 	TBitArray<> SelectionArray;
