Dataflow : Add attribute visualization node

#Synced-CL 49268397

diff --git a/Engine/Plugins/Experimental/Dataflow/Source/DataflowEditor/Private/DataflowRendering/DataflowGeometryCollectionRenderableType.cpp b/Engine/Plugins/Experimental/Dataflow/Source/DataflowEditor/Private/DataflowRendering/DataflowGeometryCollectionRenderableType.cpp
--- a/Engine/Plugins/Experimental/Dataflow/Source/DataflowEditor/Private/DataflowRendering/DataflowGeometryCollectionRenderableType.cpp
+++ b/Engine/Plugins/Experimental/Dataflow/Source/DataflowEditor/Private/DataflowRendering/DataflowGeometryCollectionRenderableType.cpp
@@ -149,13 +149,12 @@ namespace UE::Dataflow::Private
 		}
 
 		template<typename T>
-		void TransferTypedAttributeVertexColor(FGeometryCollection& GeomCollection, FName ValueAttributeName) const
+		void TransferTypedAttributeVertexColor(FGeometryCollection& GeomCollection, const FDataflowNode::FAttributeKey& ValueAttributeKey) const
 		{
 			TManagedArray<FLinearColor>* ColorAttribute = GeomCollection.FindAttributeTyped<FLinearColor>(FGeometryCollection::ColorAttribute, FGeometryCollection::VerticesGroup);
-			const TManagedArray<T>* ValueAttribute = GeomCollection.FindAttributeTyped<T>(ValueAttributeName, FGeometryCollection::VerticesGroup);
-			if (ColorAttribute && ValueAttribute)
+			const TManagedArray<T>* ValueAttribute = GeomCollection.FindAttributeTyped<T>(ValueAttributeKey.AttributeName, ValueAttributeKey.GroupName);
+			if (ColorAttribute && ValueAttribute && ValueAttribute->Num() == ColorAttribute->Num())
 			{
-				check(ColorAttribute->Num() == ValueAttribute->Num());
 				const int32 NumVertices = ColorAttribute->Num();
 
 				T MaxValue{ 0 };
@@ -202,32 +201,32 @@ namespace UE::Dataflow::Private
 			}
 		}
 
-		void TransferAttributeToVisualizeToVertexColor(FGeometryCollection& GeomCollection, FName VertexAttributeName) const
+		void TransferAttributeToVisualizeToVertexColor(FGeometryCollection& GeomCollection, const FDataflowNode::FAttributeKey& ValueAttributeKey) const
 		{
-			if (VertexAttributeName.IsNone())
+			if (ValueAttributeKey.AttributeName.IsNone() || ValueAttributeKey.GroupName.IsNone())
 			{
 				return;
 			}
-			const FManagedArrayCollection::EArrayType AttributeType = GeomCollection.GetAttributeType(VertexAttributeName, FGeometryCollection::VerticesGroup);
+			const FManagedArrayCollection::EArrayType AttributeType = GeomCollection.GetAttributeType(ValueAttributeKey.AttributeName, ValueAttributeKey.GroupName);
 			switch (AttributeType)
 			{
 			case FManagedArrayCollection::EArrayType::FFloatType:
-				TransferTypedAttributeVertexColor<float>(GeomCollection, VertexAttributeName);
+				TransferTypedAttributeVertexColor<float>(GeomCollection, ValueAttributeKey);
 				break;
 			case FManagedArrayCollection::EArrayType::FDoubleType:
-				TransferTypedAttributeVertexColor<double>(GeomCollection, VertexAttributeName);
+				TransferTypedAttributeVertexColor<double>(GeomCollection, ValueAttributeKey);
 				break;
 			case FManagedArrayCollection::EArrayType::FVectorType:
-				TransferTypedAttributeVertexColor<FVector3f>(GeomCollection, VertexAttributeName);
+				TransferTypedAttributeVertexColor<FVector3f>(GeomCollection, ValueAttributeKey);
 				break;
 			case FManagedArrayCollection::EArrayType::FVector3dType:
-				TransferTypedAttributeVertexColor<FVector3d>(GeomCollection, VertexAttributeName);
+				TransferTypedAttributeVertexColor<FVector3d>(GeomCollection, ValueAttributeKey);
 				break;
 			case FManagedArrayCollection::EArrayType::FInt32Type:
-				TransferTypedAttributeVertexColor<int32>(GeomCollection, VertexAttributeName);
+				TransferTypedAttributeVertexColor<int32>(GeomCollection, ValueAttributeKey);
 				break;
 			case FManagedArrayCollection::EArrayType::FBoolType:
-				TransferTypedAttributeVertexColor<bool>(GeomCollection, VertexAttributeName);
+				TransferTypedAttributeVertexColor<bool>(GeomCollection, ValueAttributeKey);
 				break;
 			}
 		}
diff --git a/Engine/Plugins/Experimental/Dataflow/Source/DataflowEditor/Private/DataflowRendering/DataflowRenderableTypeInstance.cpp b/Engine/Plugins/Experimental/Dataflow/Source/DataflowEditor/Private/DataflowRendering/DataflowRenderableTypeInstance.cpp
--- a/Engine/Plugins/Experimental/Dataflow/Source/DataflowEditor/Private/DataflowRendering/DataflowRenderableTypeInstance.cpp
+++ b/Engine/Plugins/Experimental/Dataflow/Source/DataflowEditor/Private/DataflowRendering/DataflowRenderableTypeInstance.cpp
@@ -109,13 +109,13 @@ namespace UE::Dataflow
 		return ::HashCombine(OutputCacheKey, RenderGroupHash);
 	}
 
-	const FName FRenderableTypeInstance::GetVertexAttributeToVisualize() const
+	const FDataflowNode::FAttributeKey FRenderableTypeInstance::GetVertexAttributeToVisualize() const
 	{
 		if (RenderableType && Node)
 		{
 			return Node->GetVertexAttributeToVisualize(EvaluationContext, OutputName, RenderableType->GetRenderGroup());
 		}
-		return NAME_None;
+		return {};
 	}
 
 	const void FRenderableTypeInstance::GetSelectionToVisualize(FDataflowSelection& OutSelection) const
diff --git a/Engine/Plugins/Experimental/Dataflow/Source/DataflowEditor/Public/DataflowRendering/DataflowRenderableTypeInstance.h b/Engine/Plugins/Experimental/Dataflow/Source/DataflowEditor/Public/DataflowRendering/DataflowRenderableTypeInstance.h
--- a/Engine/Plugins/Experimental/Dataflow/Source/DataflowEditor/Public/DataflowRendering/DataflowRenderableTypeInstance.h
+++ b/Engine/Plugins/Experimental/Dataflow/Source/DataflowEditor/Public/DataflowRendering/DataflowRenderableTypeInstance.h
@@ -77,7 +77,7 @@ namespace UE::Dataflow
 			}
 		}
 
-		const FName GetVertexAttributeToVisualize() const;
+		const FDataflowNode::FAttributeKey GetVertexAttributeToVisualize() const;
 
 		const void GetSelectionToVisualize(FDataflowSelection& OutSelection) const;
 
diff --git a/Engine/Plugins/Experimental/Dataflow/Source/DataflowNodes/Private/Dataflow/DataflowCollectionAddScalarVertexPropertyNode.cpp b/Engine/Plugins/Experimental/Dataflow/Source/DataflowNodes/Private/Dataflow/DataflowCollectionAddScalarVertexPropertyNode.cpp
--- a/Engine/Plugins/Experimental/Dataflow/Source/DataflowNodes/Private/Dataflow/DataflowCollectionAddScalarVertexPropertyNode.cpp
+++ b/Engine/Plugins/Experimental/Dataflow/Source/DataflowNodes/Private/Dataflow/DataflowCollectionAddScalarVertexPropertyNode.cpp
@@ -417,9 +417,15 @@ void FDataflowCollectionAddScalarVertexPropertyNode::ExtractVertexWeights(const
 	}
 }
 
-FName FDataflowCollectionAddScalarVertexPropertyNode::GetVertexAttributeToVisualize(UE::Dataflow::FContext& Context, const FName OutputName, const FName RenderGroup) const
+FDataflowNode::FAttributeKey FDataflowCollectionAddScalarVertexPropertyNode::GetVertexAttributeToVisualize(UE::Dataflow::FContext& Context, const FName OutputName, const FName RenderGroup) const
 {
-	return FName(GetWeightAttributeKey(Context).Attribute);
+	const FCollectionAttributeKey CollectionAttributeKey = GetWeightAttributeKey(Context);
+
+	return FAttributeKey
+	{
+		.AttributeName = FName(CollectionAttributeKey.Attribute),
+		.GroupName = FName(CollectionAttributeKey.Group),
+	};
 }
 
 #undef LOCTEXT_NAMESPACE
diff --git a/Engine/Plugins/Experimental/Dataflow/Source/DataflowNodes/Private/Dataflow/DataflowNodesPlugin.cpp b/Engine/Plugins/Experimental/Dataflow/Source/DataflowNodes/Private/Dataflow/DataflowNodesPlugin.cpp
--- a/Engine/Plugins/Experimental/Dataflow/Source/DataflowNodes/Private/Dataflow/DataflowNodesPlugin.cpp
+++ b/Engine/Plugins/Experimental/Dataflow/Source/DataflowNodes/Private/Dataflow/DataflowNodesPlugin.cpp
@@ -16,6 +16,7 @@
 #include "Dataflow/DataflowSkeletonAssetTerminalNode.h"
 #include "Dataflow/DataflowTextureToAttributeNode.h"
 #include "Dataflow/DataflowVertexColorToAttributeNode.h"
+#include "Dataflow/DataflowVisualizeAttributeNode.h"
 #include "Dataflow/DataflowValueSamplerToAttributeNode.h"
 #include "Dataflow/SamplerNodes/DataflowSamplerTypes.h"
 #include "Dataflow/SamplerNodes/DataflowGradientSamplerNode.h"
@@ -56,8 +57,11 @@ void IDataflowNodesPlugin::StartupModule()
 	UE::Dataflow::RegisterStaticMeshNodes();
 	UE::Dataflow::RegisterSelectionNodes();
 	UE::Dataflow::RegisterContextOverridesNodes();
+
 	UE::Dataflow::RegisterTextureToAttributeNodes();
 	UE::Dataflow::RegisterVertexColorToAttributeNodes();
+	UE::Dataflow::RegisterVisualizeAttributeNodes();
+
 	UE::Dataflow::RegisterValueSamplerToAttributeNodes();
 	UE::Dataflow::RegisterGradientValueSamplerNodes();
 
diff --git a/Engine/Plugins/Experimental/Dataflow/Source/DataflowNodes/Private/Dataflow/DataflowTextureToAttributeNode.cpp b/Engine/Plugins/Experimental/Dataflow/Source/DataflowNodes/Private/Dataflow/DataflowTextureToAttributeNode.cpp
--- a/Engine/Plugins/Experimental/Dataflow/Source/DataflowNodes/Private/Dataflow/DataflowTextureToAttributeNode.cpp
+++ b/Engine/Plugins/Experimental/Dataflow/Source/DataflowNodes/Private/Dataflow/DataflowTextureToAttributeNode.cpp
@@ -171,9 +171,13 @@ void FDataflowTextureToAttributeNode::Evaluate(UE::Dataflow::FContext& Context,
 	}
 }
 
-FName FDataflowTextureToAttributeNode::GetVertexAttributeToVisualize(UE::Dataflow::FContext& Context, const FName OutputName, const FName RenderGroup) const
+FDataflowNode::FAttributeKey FDataflowTextureToAttributeNode::GetVertexAttributeToVisualize(UE::Dataflow::FContext& Context, const FName OutputName, const FName RenderGroup) const
 {
-	return FName(GetValue(Context, &AttributeName, AttributeName));
+	return FAttributeKey
+	{
+		.AttributeName = FName(GetValue(Context, &AttributeName, AttributeName)),
+		.GroupName = VertexGroup.Name,
+	};
 }
 
 #undef LOCTEXT_NAMESPACE 
diff --git a/Engine/Plugins/Experimental/Dataflow/Source/DataflowNodes/Private/Dataflow/DataflowTextureToAttributeNode.h b/Engine/Plugins/Experimental/Dataflow/Source/DataflowNodes/Private/Dataflow/DataflowTextureToAttributeNode.h
--- a/Engine/Plugins/Experimental/Dataflow/Source/DataflowNodes/Private/Dataflow/DataflowTextureToAttributeNode.h
+++ b/Engine/Plugins/Experimental/Dataflow/Source/DataflowNodes/Private/Dataflow/DataflowTextureToAttributeNode.h
@@ -46,7 +46,7 @@ private:
 	FScalarVertexPropertyGroup VertexGroup;
 
 	virtual void Evaluate(UE::Dataflow::FContext& Context, const FDataflowOutput* Out) const override;
-	virtual FName GetVertexAttributeToVisualize(UE::Dataflow::FContext& Context, const FName OutputName, const FName RenderGroup) const override;
+	virtual FAttributeKey GetVertexAttributeToVisualize(UE::Dataflow::FContext& Context, const FName OutputName, const FName RenderGroup) const override;
 };
 
 namespace UE::Dataflow
diff --git a/Engine/Plugins/Experimental/Dataflow/Source/DataflowNodes/Private/Dataflow/DataflowValueSamplerToAttributeNode.cpp b/Engine/Plugins/Experimental/Dataflow/Source/DataflowNodes/Private/Dataflow/DataflowValueSamplerToAttributeNode.cpp
--- a/Engine/Plugins/Experimental/Dataflow/Source/DataflowNodes/Private/Dataflow/DataflowValueSamplerToAttributeNode.cpp
+++ b/Engine/Plugins/Experimental/Dataflow/Source/DataflowNodes/Private/Dataflow/DataflowValueSamplerToAttributeNode.cpp
@@ -81,9 +81,13 @@ void FDataflowValueSamplerToAttributeNode::Evaluate(UE::Dataflow::FContext& Cont
 	}
 }
 
-FName FDataflowValueSamplerToAttributeNode::GetVertexAttributeToVisualize(UE::Dataflow::FContext& Context, const FName OutputName, const FName RenderGroup) const
+FDataflowNode::FAttributeKey FDataflowValueSamplerToAttributeNode::GetVertexAttributeToVisualize(UE::Dataflow::FContext& Context, const FName OutputName, const FName RenderGroup) const
 {
-	return FName(GetValue(Context, &AttributeName, AttributeName));
+	return FAttributeKey
+	{
+		.AttributeName = FName(GetValue(Context, &AttributeName, AttributeName)),
+		.GroupName = VertexGroup.Name,
+	};
 }
 
 #undef LOCTEXT_NAMESPACE 
diff --git a/Engine/Plugins/Experimental/Dataflow/Source/DataflowNodes/Private/Dataflow/DataflowValueSamplerToAttributeNode.h b/Engine/Plugins/Experimental/Dataflow/Source/DataflowNodes/Private/Dataflow/DataflowValueSamplerToAttributeNode.h
--- a/Engine/Plugins/Experimental/Dataflow/Source/DataflowNodes/Private/Dataflow/DataflowValueSamplerToAttributeNode.h
+++ b/Engine/Plugins/Experimental/Dataflow/Source/DataflowNodes/Private/Dataflow/DataflowValueSamplerToAttributeNode.h
@@ -39,7 +39,7 @@ private:
 	FScalarVertexPropertyGroup VertexGroup;
 
 	virtual void Evaluate(UE::Dataflow::FContext& Context, const FDataflowOutput* Out) const override;
-	virtual FName GetVertexAttributeToVisualize(UE::Dataflow::FContext& Context, const FName OutputName, const FName RenderGroup) const override;
+	virtual FAttributeKey GetVertexAttributeToVisualize(UE::Dataflow::FContext& Context, const FName OutputName, const FName RenderGroup) const override;
 };
 
 namespace UE::Dataflow
diff --git a/Engine/Plugins/Experimental/Dataflow/Source/DataflowNodes/Private/Dataflow/DataflowVisualizeAttributeNode.cpp b/Engine/Plugins/Experimental/Dataflow/Source/DataflowNodes/Private/Dataflow/DataflowVisualizeAttributeNode.cpp
new file mode 100644
--- /dev/null
+++ b/Engine/Plugins/Experimental/Dataflow/Source/DataflowNodes/Private/Dataflow/DataflowVisualizeAttributeNode.cpp
@@ -0,0 +1,73 @@
+// Copyright Epic Games, Inc. All Rights Reserved.
+
+#include "Dataflow/DataflowVisualizeAttributeNode.h"
+
+#include "Dataflow/DataflowNodeFactory.h"
+#include "GeometryCollection/GeometryCollection.h"
+#include "GeometryCollection/Facades/CollectionMeshFacade.h"
+
+#include UE_INLINE_GENERATED_CPP_BY_NAME(DataflowVisualizeAttributeNode)
+
+#define LOCTEXT_NAMESPACE "DataflowVisualizeAttributeNode"
+
+namespace UE::Dataflow
+{
+	void RegisterVisualizeAttributeNodes()
+	{
+		DATAFLOW_NODE_REGISTER_CREATION_FACTORY(FDataflowVisualizeAttributeNode);
+	}
+};
+
+FDataflowVisualizeAttributeNode::FDataflowVisualizeAttributeNode(const UE::Dataflow::FNodeParameters& InParam, FGuid InGuid)
+	: FDataflowNode(InParam, InGuid)
+{
+	RegisterInputConnection(&Collection);
+	RegisterInputConnection(&AttributeName);
+	RegisterOutputConnection(&Collection, &Collection);
+}
+
+void FDataflowVisualizeAttributeNode::Evaluate(UE::Dataflow::FContext& Context, const FDataflowOutput* Out) const
+{
+	if (Out->IsA(&Collection))
+	{
+		// All the checks below are just for validation and informing the user about inputs
+		const FName InAttributeName = FName(GetValue(Context, &AttributeName, AttributeName));
+		if (InAttributeName.IsNone())
+		{
+			Context.Error(LOCTEXT("NoAttributeName_Msg", "Attribute name is empty, no data will be displayed"));
+			SafeForwardInput(Context, &Collection, &Collection);
+			return;
+		}
+
+		FManagedArrayCollection InCollection = GetValue(Context, &Collection, Collection);
+
+		const FName InVertexGroupName = VertexGroup.Name;
+		if (!InCollection.HasGroup(InVertexGroupName))
+		{
+			Context.Error(LOCTEXT("InvalidVertexGroup_Msg", "Vertex group does not exists, no data will be displayed"));
+			SafeForwardInput(Context, &Collection, &Collection);
+			return;
+		}
+
+		if (!InCollection.HasAttribute(InAttributeName, InVertexGroupName))
+		{
+			Context.Error(LOCTEXT("NoAttribute_Msg", "Attribute could not be found, no data will be displayed"));
+			SafeForwardInput(Context, &Collection, &Collection);
+			return;
+		}
+
+		SafeForwardInput(Context, &Collection, &Collection);
+	}
+}
+
+FDataflowNode::FAttributeKey FDataflowVisualizeAttributeNode::GetVertexAttributeToVisualize(UE::Dataflow::FContext& Context, const FName OutputName, const FName RenderGroup) const
+{
+	return FAttributeKey
+	{
+		.AttributeName = FName(GetValue(Context, &AttributeName, AttributeName)),
+		.GroupName = VertexGroup.Name,
+	};
+}
+
+#undef LOCTEXT_NAMESPACE 
+
diff --git a/Engine/Plugins/Experimental/Dataflow/Source/DataflowNodes/Private/Dataflow/DataflowVisualizeAttributeNode.h b/Engine/Plugins/Experimental/Dataflow/Source/DataflowNodes/Private/Dataflow/DataflowVisualizeAttributeNode.h
new file mode 100644
--- /dev/null
+++ b/Engine/Plugins/Experimental/Dataflow/Source/DataflowNodes/Private/Dataflow/DataflowVisualizeAttributeNode.h
@@ -0,0 +1,45 @@
+// Copyright Epic Games, Inc. All Rights Reserved.
+
+#pragma once 
+
+#include "Dataflow/DataflowNode.h"
+#include "GeometryCollection/ManagedArrayCollection.h"
+#include "Dataflow/DataflowCollectionAddScalarVertexPropertyNode.h"
+
+#include "DataflowVisualizeAttributeNode.generated.h"
+
+/**
+* Node to visualize a specific attribute in the construction viewport as vertex color
+* it supports a number of attribute type ( decimal, vector or boolean )
+*/
+USTRUCT()
+struct FDataflowVisualizeAttributeNode : public FDataflowNode
+{
+	GENERATED_USTRUCT_BODY()
+	DATAFLOW_NODE_DEFINE_INTERNAL(FDataflowVisualizeAttributeNode, "VisualizeAttribute", "Dataflow", "Debug Render Vertex Color View")
+
+public:
+	FDataflowVisualizeAttributeNode(const UE::Dataflow::FNodeParameters& InParam, FGuid InGuid = FGuid::NewGuid());
+
+private:
+	/** Collection to read the data from  */
+	UPROPERTY(meta = (DataflowInput, DataflowOutput, DataflowPassthrough = Collection))
+	FManagedArrayCollection Collection;
+
+	/** Name of the attribute to visualize */
+	UPROPERTY(EditAnywhere, Category = Target, meta = (DataflowInput))
+	FString AttributeName;
+
+	/** Vertex group matching the attribute */
+	UPROPERTY(EditAnywhere, Category = Target)
+	FScalarVertexPropertyGroup VertexGroup;
+
+	virtual void Evaluate(UE::Dataflow::FContext& Context, const FDataflowOutput* Out) const override;
+	virtual FAttributeKey GetVertexAttributeToVisualize(UE::Dataflow::FContext& Context, const FName OutputName, const FName RenderGroup) const override;
+};
+
+namespace UE::Dataflow
+{
+	void RegisterVisualizeAttributeNodes();
+}
+
diff --git a/Engine/Plugins/Experimental/Dataflow/Source/DataflowNodes/Public/Dataflow/DataflowCollectionAddScalarVertexPropertyNode.h b/Engine/Plugins/Experimental/Dataflow/Source/DataflowNodes/Public/Dataflow/DataflowCollectionAddScalarVertexPropertyNode.h
--- a/Engine/Plugins/Experimental/Dataflow/Source/DataflowNodes/Public/Dataflow/DataflowCollectionAddScalarVertexPropertyNode.h
+++ b/Engine/Plugins/Experimental/Dataflow/Source/DataflowNodes/Public/Dataflow/DataflowCollectionAddScalarVertexPropertyNode.h
@@ -162,7 +162,7 @@ private:
 	static constexpr float ReplaceChangedPassthroughValue = UE_BIG_NUMBER;
 	
 	UE_API virtual void Evaluate(UE::Dataflow::FContext& Context, const FDataflowOutput* Out) const override;
-	virtual FName GetVertexAttributeToVisualize(UE::Dataflow::FContext& Context, const FName OutputName, const FName RenderGroup) const override;
+	virtual FAttributeKey GetVertexAttributeToVisualize(UE::Dataflow::FContext& Context, const FName OutputName, const FName RenderGroup) const override;
 
 	//~ Begin - FDataflowVertexAttributeEditableNode API
 	virtual void GetSupportedViewModes(UE::Dataflow::FContext& Context, TArray<FName>& OutViewModeNames) const override;
diff --git a/Engine/Source/Runtime/Experimental/Dataflow/Core/Public/Dataflow/DataflowNode.h b/Engine/Source/Runtime/Experimental/Dataflow/Core/Public/Dataflow/DataflowNode.h
--- a/Engine/Source/Runtime/Experimental/Dataflow/Core/Public/Dataflow/DataflowNode.h
+++ b/Engine/Source/Runtime/Experimental/Dataflow/Core/Public/Dataflow/DataflowNode.h
@@ -719,12 +719,18 @@ PRAGMA_ENABLE_DEPRECATION_WARNINGS
 	// returns true to disable node properties in its detail view when the corresponding input is connected ( default is true )
 	virtual bool MakeConnectedPropertiesReadOnly() const { return true; }
 
+	struct FAttributeKey
+	{
+		FName AttributeName;
+		FName GroupName;
+	};
+
 	/**
 	* return a vertex attribute to visualize
 	* This attribute will be used as vertex color and displayed in the viewport 
 	* this is only used for collection based visualizations
 	*/
-	virtual FName GetVertexAttributeToVisualize(UE::Dataflow::FContext& Context, const FName OutputName, const FName RenderGroup) const { return NAME_None; }
+	virtual FAttributeKey GetVertexAttributeToVisualize(UE::Dataflow::FContext& Context, const FName OutputName, const FName RenderGroup) const { return {}; }
 
 	/**
 	* Fills a dataflow selection object used to visualize the selection a node may output for a sepcific output and render group
