Dataflow editor : Add rendering config and implement it for proximity renderable type
- Render settings for renderable types need to inherit from UDataflowRenderableTypeSettings
- Settings are displayed in the tool property section on the left of the construction view

#Synced-CL 49325867

diff --git a/Engine/Plugins/Experimental/Dataflow/Source/DataflowEditor/Private/Dataflow/DataflowConstructionScene.cpp b/Engine/Plugins/Experimental/Dataflow/Source/DataflowEditor/Private/Dataflow/DataflowConstructionScene.cpp
--- a/Engine/Plugins/Experimental/Dataflow/Source/DataflowEditor/Private/Dataflow/DataflowConstructionScene.cpp
+++ b/Engine/Plugins/Experimental/Dataflow/Source/DataflowEditor/Private/Dataflow/DataflowConstructionScene.cpp
@@ -45,6 +45,19 @@ TArray<TObjectPtr<UDynamicMeshComponent>> FDataflowConstructionScene::GetDynamic
 	return MoveTemp(OutValues);
 }
 
+TArray<UDataflowRenderableTypeSettings*> FDataflowConstructionScene::GetRenderableTypeSettings() const
+{
+	TArray<UDataflowRenderableTypeSettings*> RenderSettingsArray;
+	for (const UE::Dataflow::FRenderableTypeInstance& RenderableInstance : RenderableTypeInstances)
+	{
+		if (UDataflowRenderableTypeSettings* RenderSettings = RenderableInstance.GetRenderSettings())
+		{
+			RenderSettingsArray.Add(RenderSettings);
+		}
+	}
+	return RenderSettingsArray;
+}
+
 /** Hide all or a single component */
 void FDataflowConstructionScene::SetVisibility(bool bVisibility, UActorComponent* InComponent)
 {
@@ -472,6 +485,11 @@ void FDataflowConstructionScene::UpdateConstructionScene()
 	{
 		TerminalContent->SetConstructionDirty(false);
 	}
+
+	if (OnConstructionSceneUpdated.IsBound())
+	{
+		OnConstructionSceneUpdated.Broadcast();
+	}
 }
 
 void FDataflowConstructionScene::SelectNodeComponents(const UDataflowEdNode* EdNode)
@@ -690,7 +708,7 @@ void FDataflowConstructionScene::UpdateComponentsFromNode(const UDataflowEdNode*
 						UE::Dataflow::FRenderableTypeInstance::GetRenderableInstancesForNode(*DataflowContext, *ConstructionViewMode, *DataflowNode, NodeRenderableTypeInstances);
 					}
 				}
-				NodeRenderableTypeInstances.Append(RenderableTypeInstances);
+				RenderableTypeInstances.Append(NodeRenderableTypeInstances);
 
 				// extract components from instances
 				const float RenderableTypeProgress = (NodeRenderableTypeInstances.Num() > 0) ? (1.f / (float)NodeRenderableTypeInstances.Num()) : 0;
diff --git a/Engine/Plugins/Experimental/Dataflow/Source/DataflowEditor/Private/Dataflow/DataflowEditorMode.cpp b/Engine/Plugins/Experimental/Dataflow/Source/DataflowEditor/Private/Dataflow/DataflowEditorMode.cpp
--- a/Engine/Plugins/Experimental/Dataflow/Source/DataflowEditor/Private/Dataflow/DataflowEditorMode.cpp
+++ b/Engine/Plugins/Experimental/Dataflow/Source/DataflowEditor/Private/Dataflow/DataflowEditorMode.cpp
@@ -506,6 +506,12 @@ void UDataflowEditorMode::Exit()
 
 void UDataflowEditorMode::SetDataflowConstructionScene(TWeakPtr<FDataflowConstructionScene> InConstructionScene)
 {
+	// Remove previous delegate
+	if (const TSharedPtr<FDataflowConstructionScene> ConstructionScenePtr = ConstructionScene.Pin())
+	{
+		ConstructionScenePtr->OnConstructionSceneUpdated.RemoveAll(this);
+	}
+
 	ConstructionScene = InConstructionScene;
 
 	if (const TSharedPtr<FDataflowConstructionScene> ConstructionScenePtr = ConstructionScene.Pin())
@@ -526,6 +532,20 @@ void UDataflowEditorMode::SetDataflowConstructionScene(TWeakPtr<FDataflowConstru
 
 		PreviewToolManager->OnToolStarted.AddSP(DataflowModeToolkit, &FDataflowEditorModeToolkit::OnToolStarted);
 		PreviewToolManager->OnToolEnded.AddSP(DataflowModeToolkit, &FDataflowEditorModeToolkit::OnToolEnded);
+
+		ConstructionScenePtr->OnConstructionSceneUpdated.AddUObject(this, &UDataflowEditorMode::OnConstructionSceneUpdated);
+	}
+}
+
+void UDataflowEditorMode::OnConstructionSceneUpdated()
+{
+	if (FDataflowEditorModeToolkit* const DataflowModeToolkit = static_cast<FDataflowEditorModeToolkit*>(Toolkit.Get()))
+	{
+		if (const TSharedPtr<FDataflowConstructionScene> ConstructionScenePtr = ConstructionScene.Pin())
+		{
+			TArray<UDataflowRenderableTypeSettings*> RenderSettings = ConstructionScenePtr->GetRenderableTypeSettings();
+			DataflowModeToolkit->OnRenderSettingsUpdated(RenderSettings);
+		}
 	}
 }
 
diff --git a/Engine/Plugins/Experimental/Dataflow/Source/DataflowEditor/Private/Dataflow/DataflowEditorModeToolkit.cpp b/Engine/Plugins/Experimental/Dataflow/Source/DataflowEditor/Private/Dataflow/DataflowEditorModeToolkit.cpp
--- a/Engine/Plugins/Experimental/Dataflow/Source/DataflowEditor/Private/Dataflow/DataflowEditorModeToolkit.cpp
+++ b/Engine/Plugins/Experimental/Dataflow/Source/DataflowEditor/Private/Dataflow/DataflowEditorModeToolkit.cpp
@@ -9,6 +9,7 @@
 #include "Dataflow/DataflowConstructionViewport.h"
 #include "Dataflow/DataflowEditorPreviewSceneBase.h"
 #include "Dataflow/DataflowConstructionScene.h"
+#include "DataflowRendering/DataflowRenderableTypeSettings.h"
 #include "Tools/EdModeInteractiveToolsContext.h"
 #include "Framework/Application/SlateApplication.h"
 #include "InteractiveToolManager.h"
@@ -107,10 +108,20 @@ void FDataflowEditorModeToolkit::Init(const TSharedPtr<IToolkitHost>& InitToolki
 		Context->OnToolWarningMessage.AddSP(this, &FBaseCharacterFXEditorModeToolkit::PostWarning);
 	}
 
+	ModeDetailsView->OnFinishedChangingProperties().AddSP(this, &FDataflowEditorModeToolkit::OnSettingsPropertiesChanged);
+
 	// add viewport overlay widget to accept / cancel tool
 	MakeToolAcceptCancelWidget();
 }
 
+void FDataflowEditorModeToolkit::OnSettingsPropertiesChanged(const FPropertyChangedEvent& PropertyChangedEvent)
+{
+	if (ActiveToolName.IsEmpty())
+	{
+		RefreshConstructionScene();
+	}
+}
+
 void FDataflowEditorModeToolkit::MakeToolAcceptCancelWidget()
 {
 	SAssignNew(ViewportOverlayWidget, SHorizontalBox)
@@ -365,22 +376,8 @@ void FDataflowEditorModeToolkit::OnToolStarted(UInteractiveToolManager* Manager,
 
 void FDataflowEditorModeToolkit::OnToolEnded(UInteractiveToolManager* Manager, UInteractiveTool* Tool)
 {
-	auto GetConstructionScene = [this]() -> TWeakPtr<FDataflowConstructionScene>
-		{
-			constexpr bool bEvenIfPendingKill = false;
-			if (OwningEditorMode.IsValid(bEvenIfPendingKill))
-			{
-				UEdMode* const Mode = OwningEditorMode.Get();
-				if (UDataflowEditorMode* const DataflowEdMode = Cast<UDataflowEditorMode>(Mode))
-				{
-					return DataflowEdMode->GetDataflowConstructionScene();
-				}
-			}
-			return nullptr;
-		};
-
 	ModeToolPallete.ModeHeaderArea->SetVisibility(EVisibility::Visible);
-	ModeDetailsView->SetObject(nullptr);
+	UpdateModeDetailsView();
 	ActiveToolName = FText::GetEmpty();
 	ClearNotification();
 	ClearWarning();
@@ -396,6 +393,25 @@ void FDataflowEditorModeToolkit::OnToolEnded(UInteractiveToolManager* Manager, U
 		CurTool->OnPropertyModifiedDirectlyByTool.RemoveAll(this);
 	}
 
+	RefreshConstructionScene();
+}
+
+void FDataflowEditorModeToolkit::RefreshConstructionScene()
+{
+	auto GetConstructionScene = [this]() -> TWeakPtr<FDataflowConstructionScene>
+		{
+			constexpr bool bEvenIfPendingKill = false;
+			if (OwningEditorMode.IsValid(bEvenIfPendingKill))
+			{
+				UEdMode* const Mode = OwningEditorMode.Get();
+				if (UDataflowEditorMode* const DataflowEdMode = Cast<UDataflowEditorMode>(Mode))
+				{
+					return DataflowEdMode->GetDataflowConstructionScene();
+				}
+			}
+			return nullptr;
+		};
+
 	if (const TSharedPtr<FDataflowConstructionScene> ConstructionScene = GetConstructionScene().Pin())
 	{
 		if (const TObjectPtr<UDataflowBaseContent>& EditorContent = ConstructionScene->GetEditorContent())
@@ -415,15 +431,34 @@ void FDataflowEditorModeToolkit::SetSimulationViewportWidget(TWeakPtr<SDataflowS
 	SimulationViewportWidget = InSimulationViewportWidget;
 }
 
-void FDataflowEditorModeToolkit::UpdateActiveToolProperties(UInteractiveTool* Tool)
+void FDataflowEditorModeToolkit::OnRenderSettingsUpdated(const TArray<UDataflowRenderableTypeSettings*>& InRenderSettings)
 {
-	UInteractiveTool* CurTool = GetScriptableEditorMode()->GetToolManager(EToolsContextScope::EdMode)->GetActiveTool(EToolSide::Left);
-	if (CurTool == nullptr)
+	RenderSettings = InRenderSettings; 
+	UpdateModeDetailsView();
+}
+
+void FDataflowEditorModeToolkit::UpdateModeDetailsView()
+{
+	if (ModeDetailsView)
 	{
-		return;
+		TArray<UObject*> ObjectsToSet;
+		if (UInteractiveTool* CurTool = GetScriptableEditorMode()->GetToolManager(EToolsContextScope::EdMode)->GetActiveTool(EToolSide::Left))
+		{
+			ObjectsToSet.Append(CurTool->GetToolProperties(true));
+		}
+		else
+		{
+			// if a tool is up , it has its own settings, so we only display the selection render settings when the tool is not up 
+			ObjectsToSet.Append(RenderSettings);
+		}
+
+		ModeDetailsView->SetObjects(ObjectsToSet);
 	}
-		
-	ModeDetailsView->SetObjects(CurTool->GetToolProperties(true));
+}
+
+void FDataflowEditorModeToolkit::UpdateActiveToolProperties(UInteractiveTool* Tool)
+{
+	UpdateModeDetailsView();
 }
 
 void FDataflowEditorModeToolkit::InvalidateCachedDetailPanelState(UObject* ChangedObject)
diff --git a/Engine/Plugins/Experimental/Dataflow/Source/DataflowEditor/Private/DataflowRendering/DataflowGeometryCollectionProximityRenderableType.cpp b/Engine/Plugins/Experimental/Dataflow/Source/DataflowEditor/Private/DataflowRendering/DataflowGeometryCollectionProximityRenderableType.cpp
--- a/Engine/Plugins/Experimental/Dataflow/Source/DataflowEditor/Private/DataflowRendering/DataflowGeometryCollectionProximityRenderableType.cpp
+++ b/Engine/Plugins/Experimental/Dataflow/Source/DataflowEditor/Private/DataflowRendering/DataflowGeometryCollectionProximityRenderableType.cpp
@@ -31,6 +31,7 @@ namespace UE::Dataflow::Private
 		UE_DATAFLOW_IRENDERABLE_PRIMARY_TYPE(FManagedArrayCollection, Collection);
 		UE_DATAFLOW_IRENDERABLE_RENDER_GROUP(Proximity);
 		UE_DATAFLOW_IRENDERABLE_VIEW_MODE(FDataflowConstruction3DViewMode);
+		UE_DATAFLOW_IRENDERABLE_SETTINGS(UDataflowGeometryCollectionProximityRenderSettings);
 
 	public:
 		FGeometryCollectionProximityRenderableType()
@@ -51,56 +52,72 @@ namespace UE::Dataflow::Private
 
 		virtual void GetPrimitiveComponents(const FRenderableTypeInstance& Instance, FRenderableComponents& OutComponents) const override
 		{
+			const UDataflowGeometryCollectionProximityRenderSettings* Settings = Instance.GetTypedRenderSettings<UDataflowGeometryCollectionProximityRenderSettings>();
+			
+			const bool bShowCenters = Settings ? Settings->bShowCenters : true;
+			const bool bShowConnections = Settings ? Settings->bShowConnections : true;
+
 			const FManagedArrayCollection& Collection = GetCollection(Instance);
 
 			if (Collection.HasAttribute(ProximityAttributeName, FGeometryCollection::GeometryGroup) &&
 				Collection.HasAttribute(BoundingBoxAttributeName, FGeometryCollection::GeometryGroup))
 			{
-				static const FName PointComponentName = TEXT("Connectivity_Points");
-				UPointSetComponent* PointComponent = OutComponents.AddNewComponent<UPointSetComponent>(PointComponentName);
-				static const FName LineComponentName = TEXT("Connectivity_Lines");
-				ULineSetComponent* LineComponent = OutComponents.AddNewComponent<ULineSetComponent>(LineComponentName);
-
-				if (PointComponent && LineComponent)
+				const TManagedArray<TSet<int32>>& Proximity = Collection.GetAttribute<TSet<int32>>(ProximityAttributeName, FGeometryCollection::GeometryGroup);
+				const TManagedArray<FBox>& BoundingBox = Collection.GetAttribute<FBox>(BoundingBoxAttributeName, FGeometryCollection::GeometryGroup);
+				const TManagedArray<FTransform3f>* Transform = Collection.FindAttributeTyped<FTransform3f>(FGeometryCollection::TransformAttribute, FGeometryCollection::TransformGroup);
+				const TManagedArray<int32>* Parent = Collection.FindAttributeTyped<int32>(FGeometryCollection::ParentAttribute, FGeometryCollection::TransformGroup);
+				const TManagedArray<int32>* GeometryToTransformIndex = Collection.FindAttributeTyped<int32>(FGeometryCollection::TransformIndexAttribute, FGeometryCollection::GeometryGroup);
+
+				// compute transforms in collection space 
+				TArray<FTransform> GlobalTransformArray;
+				if (Transform && Parent)
 				{
-					const TManagedArray<TSet<int32>>& Proximity = Collection.GetAttribute<TSet<int32>>(ProximityAttributeName, FGeometryCollection::GeometryGroup);
-					const TManagedArray<FBox>& BoundingBox = Collection.GetAttribute<FBox>(BoundingBoxAttributeName, FGeometryCollection::GeometryGroup);
-					const TManagedArray<FTransform3f>* Transform = Collection.FindAttributeTyped<FTransform3f>(FGeometryCollection::TransformAttribute, FGeometryCollection::TransformGroup);
-					const TManagedArray<int32>* Parent = Collection.FindAttributeTyped<int32>(FGeometryCollection::ParentAttribute, FGeometryCollection::TransformGroup);
-					const TManagedArray<int32>* GeometryToTransformIndex = Collection.FindAttributeTyped<int32>(FGeometryCollection::TransformIndexAttribute, FGeometryCollection::GeometryGroup);
-
-					// compute transforms in collection space 
-					TArray<FTransform> GlobalTransformArray;
-					if (Transform && Parent)
-					{
-						GeometryCollectionAlgo::GlobalMatrices(*Transform, *Parent, GlobalTransformArray);
-					}
+					GeometryCollectionAlgo::GlobalMatrices(*Transform, *Parent, GlobalTransformArray);
+				}
 
-					auto GetGeometryTransform =
-						[&GeometryToTransformIndex, &GlobalTransformArray](int32 GeoIndex) -> FTransform
+				auto GetGeometryTransform =
+					[&GeometryToTransformIndex, &GlobalTransformArray](int32 GeoIndex) -> FTransform
+					{
+						if (GeometryToTransformIndex && GeometryToTransformIndex->IsValidIndex(GeoIndex))
 						{
-							if (GeometryToTransformIndex && GeometryToTransformIndex->IsValidIndex(GeoIndex))
+							const int32 TransformIndex = (*GeometryToTransformIndex)[GeoIndex];
+							if (GlobalTransformArray.IsValidIndex(TransformIndex))
 							{
-								const int32 TransformIndex = (*GeometryToTransformIndex)[GeoIndex];
-								if (GlobalTransformArray.IsValidIndex(TransformIndex))
-								{
-									return GlobalTransformArray[TransformIndex];
-								}
+								return GlobalTransformArray[TransformIndex];
 							}
-							return FTransform::Identity;
-						};
+						}
+						return FTransform::Identity;
+					};
+
+				const int32 NumGeometry = Collection.NumElements(FGeometryCollection::GeometryGroup);
+				TArray<FVector> Centers;
+				if (bShowCenters)
+				{
+					Centers.Reserve(NumGeometry);
+				}
+				TArray<FVector> LineStarts;
+				TArray<FVector> LineEnds;
+				if (bShowConnections)
+				{
+					// rough estimate of number of connection
+					LineStarts.Reserve(NumGeometry * 2);
+					LineEnds.Reserve(NumGeometry * 2);
+				}
 
-					const int32 NumGeometry = Collection.NumElements(FGeometryCollection::GeometryGroup);
-					TArray<FVector> Centers;
-					for (int32 GeoIndex = 0; GeoIndex < NumGeometry; ++GeoIndex)
+				for (int32 GeoIndex = 0; GeoIndex < NumGeometry; ++GeoIndex)
+				{
+					if (Proximity[GeoIndex].Num() > 0)
 					{
-						if (Proximity[GeoIndex].Num() > 0)
-						{
-							const FTransform GeometryTransform = GetGeometryTransform(GeoIndex);
-							const FVector Center = GeometryTransform.TransformPosition(BoundingBox[GeoIndex].GetCenter());
+						const FTransform GeometryTransform = GetGeometryTransform(GeoIndex);
+						const FVector Center = GeometryTransform.TransformPosition(BoundingBox[GeoIndex].GetCenter());
 
-							PointComponent->AddPoint(FRenderablePoint(Center, FColor::Blue, 15.0f));
+						if (bShowCenters)
+						{
+							Centers.Add(Center);
+						}
 
+						if (bShowConnections)
+						{
 							for (const int32 OtherGeoIndex : Proximity[GeoIndex])
 							{
 								if (OtherGeoIndex > GeoIndex)
@@ -108,14 +125,39 @@ namespace UE::Dataflow::Private
 									const FTransform OtherGeometryTransform = GetGeometryTransform(OtherGeoIndex);
 									const FVector OtherCenter = OtherGeometryTransform.TransformPosition(BoundingBox[OtherGeoIndex].GetCenter());
 
-									LineComponent->AddLine(FRenderableLine(Center, OtherCenter, FColor::Yellow, 2.0f));
+									LineStarts.Add(Center);
+									LineEnds.Add(OtherCenter);
+
 								}
 							}
 						}
 					}
+				}
+
+				if (bShowCenters && Centers.Num())
+				{
+					static const FName PointComponentName = TEXT("Connectivity_Points");
+					if (UPointSetComponent* PointComponent = OutComponents.AddNewComponent<UPointSetComponent>(PointComponentName))
+					{
+						const float PointSize = Settings ? Settings->PointSize : 15.0f;
+						const FColor PointColor = Settings ? Settings->PointColor : FColor::Blue;
 
-					PointComponent->SetPointMaterial(Material);
-					LineComponent->SetLineMaterial(Material);
+						PointComponent->AddPoints(Centers, PointColor, PointSize);
+						PointComponent->SetPointMaterial(Material);
+					}
+				}
+
+				if (bShowConnections && LineStarts.Num() && LineStarts.Num() == LineEnds.Num())
+				{
+					static const FName LineComponentName = TEXT("Connectivity_Lines");
+					if (ULineSetComponent* LineComponent = OutComponents.AddNewComponent<ULineSetComponent>(LineComponentName))
+					{
+						const float LineThickness = Settings? Settings->LineThickness: 4.0f;
+						const FColor LineColor = Settings ? Settings->LineColor: FColor::Yellow;
+
+						LineComponent->AddLines(LineStarts, LineEnds, LineColor, LineThickness);
+						LineComponent->SetLineMaterial(Material);
+					}
 				}
 			}
 		}
diff --git a/Engine/Plugins/Experimental/Dataflow/Source/DataflowEditor/Private/DataflowRendering/DataflowGeometryCollectionProximityRenderableType.h b/Engine/Plugins/Experimental/Dataflow/Source/DataflowEditor/Private/DataflowRendering/DataflowGeometryCollectionProximityRenderableType.h
--- a/Engine/Plugins/Experimental/Dataflow/Source/DataflowEditor/Private/DataflowRendering/DataflowGeometryCollectionProximityRenderableType.h
+++ b/Engine/Plugins/Experimental/Dataflow/Source/DataflowEditor/Private/DataflowRendering/DataflowGeometryCollectionProximityRenderableType.h
@@ -2,6 +2,35 @@
 
 #pragma once
 
+#include "DataflowRendering/DataflowRenderableTypeSettings.h"
+#include "Math/Color.h"
+
+#include "DataflowGeometryCollectionProximityRenderableType.generated.h"
+
+UCLASS(MinimalAPI, AutoExpandCategories="Proximity|Centers, Proximity|Connections")
+class UDataflowGeometryCollectionProximityRenderSettings : public UDataflowRenderableTypeSettings
+{
+	GENERATED_BODY()
+public:
+	UPROPERTY(EditAnywhere, Category = "Proximity|Centers")
+	bool bShowCenters = true;
+
+	UPROPERTY(EditAnywhere, Category = "Proximity|Centers")
+	float PointSize = 15.0f;
+
+	UPROPERTY(EditAnywhere, Category = "Proximity|Centers")
+	FColor PointColor = FColor::Blue;
+
+	UPROPERTY(EditAnywhere, Category = "Proximity|Connections")
+	bool bShowConnections = true;
+
+	UPROPERTY(EditAnywhere, Category = "Proximity|Connections")
+	float LineThickness = 4.0f;
+
+	UPROPERTY(EditAnywhere, Category = "Proximity|Connections")
+	FColor LineColor = FColor::Yellow;
+};
+
 namespace UE::Dataflow::Private
 {
 	void RegisterGeometryCollectionProximityRenderableTypes();
diff --git a/Engine/Plugins/Experimental/Dataflow/Source/DataflowEditor/Private/DataflowRendering/DataflowRenderableTypeInstance.cpp b/Engine/Plugins/Experimental/Dataflow/Source/DataflowEditor/Private/DataflowRendering/DataflowRenderableTypeInstance.cpp
--- a/Engine/Plugins/Experimental/Dataflow/Source/DataflowEditor/Private/DataflowRendering/DataflowRenderableTypeInstance.cpp
+++ b/Engine/Plugins/Experimental/Dataflow/Source/DataflowEditor/Private/DataflowRendering/DataflowRenderableTypeInstance.cpp
@@ -6,6 +6,8 @@
 #include "DataflowRendering/DataflowRenderableComponents.h"
 #include "DataflowRendering/DataflowRenderableType.h"
 #include "DataflowRendering/DataflowRenderableTypeRegistry.h"
+#include "DataflowRendering/DataflowRenderableTypeSettings.h"
+#include "UObject/UObjectGlobals.h"
 
 namespace UE::Dataflow
 {
@@ -109,6 +111,24 @@ namespace UE::Dataflow
 		return ::HashCombine(OutputCacheKey, RenderGroupHash);
 	}
 
+	UDataflowRenderableTypeSettings* FRenderableTypeInstance::GetRenderSettings() const
+	{
+		if (RenderableType && Node)
+		{
+			if (const UClass* SettingsClass = RenderableType->GetSettingsClass())
+			{
+				const FContextCacheKey RenderSettingsKey = ::GetTypeHash(SettingsClass);
+				if (!EvaluationContext.HasData(RenderSettingsKey))
+				{
+					TObjectPtr<UDataflowRenderableTypeSettings> NewRenderSettings = NewObject<UDataflowRenderableTypeSettings>(GetTransientPackageAsObject(), SettingsClass);
+					EvaluationContext.SetData(RenderSettingsKey, nullptr, NewRenderSettings, Node->GetGuid(), /*Node->GetValueHash()*/0, Node->GetTimestamp());
+				}
+				return EvaluationContext.GetData<TObjectPtr<UDataflowRenderableTypeSettings>>(RenderSettingsKey, nullptr, {}).Get();
+			}
+		}
+		return nullptr;
+	}
+
 	const FDataflowNode::FAttributeKey FRenderableTypeInstance::GetVertexAttributeToVisualize() const
 	{
 		if (RenderableType && Node)
diff --git a/Engine/Plugins/Experimental/Dataflow/Source/DataflowEditor/Private/DataflowRendering/DataflowRenderableTypeSettings.h b/Engine/Plugins/Experimental/Dataflow/Source/DataflowEditor/Private/DataflowRendering/DataflowRenderableTypeSettings.h
new file mode 100644
--- /dev/null
+++ b/Engine/Plugins/Experimental/Dataflow/Source/DataflowEditor/Private/DataflowRendering/DataflowRenderableTypeSettings.h
@@ -0,0 +1,14 @@
+// Copyright Epic Games, Inc. All Rights Reserved.
+
+#pragma once
+
+#include "UObject/Object.h"
+
+#include "DataflowRenderableTypeSettings.generated.h"
+
+UCLASS(MinimalAPI)
+class UDataflowRenderableTypeSettings : public UObject
+{
+	GENERATED_BODY()
+};
+
diff --git a/Engine/Plugins/Experimental/Dataflow/Source/DataflowEditor/Public/Dataflow/DataflowConstructionScene.h b/Engine/Plugins/Experimental/Dataflow/Source/DataflowEditor/Public/Dataflow/DataflowConstructionScene.h
--- a/Engine/Plugins/Experimental/Dataflow/Source/DataflowEditor/Public/Dataflow/DataflowConstructionScene.h
+++ b/Engine/Plugins/Experimental/Dataflow/Source/DataflowEditor/Public/Dataflow/DataflowConstructionScene.h
@@ -15,6 +15,7 @@ class UDynamicMeshComponent;
 class UMeshElementsVisualizer;
 class UDataflowEdNode;
 class UInteractiveToolPropertySet;
+class UDataflowRenderableTypeSettings;
 
 namespace UE
 {
@@ -77,11 +78,20 @@ public:
 	/** Array of all DynamicMeshComponents in the scene */
 	UE_API TArray<TObjectPtr<UDynamicMeshComponent>> GetDynamicMeshComponents() const;
 
+	/** Get the render settings for the selected nodes */
+	UE_API TArray<UObject*> GetRenderSettings() const;
+
 	/** select only the components that are relates to a specific node */
 	UE_API void SelectNodeComponents(const UDataflowEdNode* EdNode);
 
 	FDebugMesh DebugMesh;
 
+	/** Get the render settings for matching the selected nodes */
+	UE_API TArray<UDataflowRenderableTypeSettings*> GetRenderableTypeSettings() const;
+
+	DECLARE_MULTICAST_DELEGATE(FOnConstructionSceneUpdated)
+	FOnConstructionSceneUpdated OnConstructionSceneUpdated;
+
 private:
 
 	/** Reset the scene mesh visualizer */
diff --git a/Engine/Plugins/Experimental/Dataflow/Source/DataflowEditor/Public/Dataflow/DataflowEditorMode.h b/Engine/Plugins/Experimental/Dataflow/Source/DataflowEditor/Public/Dataflow/DataflowEditorMode.h
--- a/Engine/Plugins/Experimental/Dataflow/Source/DataflowEditor/Public/Dataflow/DataflowEditorMode.h
+++ b/Engine/Plugins/Experimental/Dataflow/Source/DataflowEditor/Public/Dataflow/DataflowEditorMode.h
@@ -223,6 +223,8 @@ private:
 	TWeakPtr<FDataflowConstructionViewportClient, ESPMode::ThreadSafe> ConstructionViewportClient;
 	TWeakPtr<FDataflowSimulationViewportClient, ESPMode::ThreadSafe> SimulationViewportClient;
 
+	void OnConstructionSceneUpdated();
+
 	// The first time we get a valid mesh, refocus the camera on it
 	bool bFirstValid2DMesh = true;
 	bool bFirstValid3DMesh = true;
diff --git a/Engine/Plugins/Experimental/Dataflow/Source/DataflowEditor/Public/Dataflow/DataflowEditorModeToolkit.h b/Engine/Plugins/Experimental/Dataflow/Source/DataflowEditor/Public/Dataflow/DataflowEditorModeToolkit.h
--- a/Engine/Plugins/Experimental/Dataflow/Source/DataflowEditor/Public/Dataflow/DataflowEditorModeToolkit.h
+++ b/Engine/Plugins/Experimental/Dataflow/Source/DataflowEditor/Public/Dataflow/DataflowEditorModeToolkit.h
@@ -10,6 +10,7 @@ class UEditorInteractiveToolsContext;
 class SDataflowConstructionViewport;
 class SDataflowSimulationViewport;
 class SBaseCharacterFXEditorViewport;
+class UDataflowRenderableTypeSettings;
 
 /** Tool pallete data for the dataflow mode toolkit */
 struct FDataflowEditorModeToolPalette
@@ -56,6 +57,8 @@ public:
 	UE_API void SetConstructionViewportWidget(TWeakPtr<SDataflowConstructionViewport>);
 	UE_API void SetSimulationViewportWidget(TWeakPtr<SDataflowSimulationViewport>);
 
+	void OnRenderSettingsUpdated(const TArray<UDataflowRenderableTypeSettings*>& RenderSettings);
+
 private:
 
 	/** Register the toolkit builder palettes */
@@ -67,6 +70,8 @@ private:
 	/** Update the active tool properties */
 	void UpdateActiveToolProperties(UInteractiveTool* Tool);
 
+	void UpdateModeDetailsView();
+
 	/** Invalidate the cache state */
 	void InvalidateCachedDetailPanelState(UObject* ChangedObject);
 
@@ -80,6 +85,10 @@ private:
 	/** Current tool context used to create nodes */
 	UE_API UEditorInteractiveToolsContext* GetCurrentToolsContext();
 
+	void OnSettingsPropertiesChanged(const FPropertyChangedEvent& PropertyChangedEvent);
+
+	void RefreshConstructionScene();
+
 	/** Construction viewport */
 	TWeakPtr<SDataflowConstructionViewport> ConstructionViewportWidget;
 
@@ -88,6 +97,8 @@ private:
 
 	/** Dataflow toolkit mode tool palette data */
 	FDataflowEditorModeToolPalette ModeToolPallete;
+
+	TArray<UDataflowRenderableTypeSettings*> RenderSettings;
 };
 
 #undef UE_API
diff --git a/Engine/Plugins/Experimental/Dataflow/Source/DataflowEditor/Public/DataflowRendering/DataflowRenderableType.h b/Engine/Plugins/Experimental/Dataflow/Source/DataflowEditor/Public/DataflowRendering/DataflowRenderableType.h
--- a/Engine/Plugins/Experimental/Dataflow/Source/DataflowEditor/Public/DataflowRendering/DataflowRenderableType.h
+++ b/Engine/Plugins/Experimental/Dataflow/Source/DataflowEditor/Public/DataflowRendering/DataflowRenderableType.h
@@ -23,6 +23,9 @@ namespace UE::Dataflow
 		/** Returns the supported render group, useful when rendering a specific type in a particular way (ie : "Surface" may render the geometry related groups ) */
 		virtual FName GetRenderGroup() const { return NAME_None; }
 
+		/** returns the class of the settings object if any */
+		virtual const UClass* GetSettingsClass() const { return nullptr; }
+
 		/** return true if the ViewMode is supported for rendering this */
 		virtual bool IsViewModeSupported(const IDataflowConstructionViewMode& ViewMode) const = 0;
 
@@ -52,4 +55,7 @@ namespace UE::Dataflow
 	virtual bool IsViewModeSupported(const UE::Dataflow::IDataflowConstructionViewMode& InViewMode) const override \
 		{ return (InViewMode.GetName() == ViewModeClass::Name); }
 
+#define UE_DATAFLOW_IRENDERABLE_SETTINGS(RenderSettingsClass) \
+	virtual const UClass* GetSettingsClass() const override \
+		{ return RenderSettingsClass::StaticClass(); }
 
diff --git a/Engine/Plugins/Experimental/Dataflow/Source/DataflowEditor/Public/DataflowRendering/DataflowRenderableTypeInstance.h b/Engine/Plugins/Experimental/Dataflow/Source/DataflowEditor/Public/DataflowRendering/DataflowRenderableTypeInstance.h
--- a/Engine/Plugins/Experimental/Dataflow/Source/DataflowEditor/Public/DataflowRendering/DataflowRenderableTypeInstance.h
+++ b/Engine/Plugins/Experimental/Dataflow/Source/DataflowEditor/Public/DataflowRendering/DataflowRenderableTypeInstance.h
@@ -11,6 +11,7 @@
 class UPrimitiveComponent;
 class AActor;
 struct FDataflowSelection;
+class UDataflowRenderableTypeSettings;
 
 #define UE_API DATAFLOWEDITOR_API
 
@@ -81,6 +82,14 @@ namespace UE::Dataflow
 
 		const void GetSelectionToVisualize(FDataflowSelection& OutSelection) const;
 
+		UDataflowRenderableTypeSettings* GetRenderSettings() const;
+
+		template <typename T UE_REQUIRES(std::is_base_of_v<UDataflowRenderableTypeSettings, T>)>
+		const T* GetTypedRenderSettings() const
+		{
+			return Cast<T>(GetRenderSettings());
+		}
+
 		template <typename T>
 		const T& GetOutputValue(FName Name, const T& Default = T()) const
 		{
