Dataflow : add render settings for geometry collection renderable types

#Synced-CL 49325867

diff --git a/Engine/Plugins/Experimental/Dataflow/Source/DataflowEditor/Private/DataflowRendering/DataflowGeometryCollectionRenderableType.cpp b/Engine/Plugins/Experimental/Dataflow/Source/DataflowEditor/Private/DataflowRendering/DataflowGeometryCollectionRenderableType.cpp
--- a/Engine/Plugins/Experimental/Dataflow/Source/DataflowEditor/Private/DataflowRendering/DataflowGeometryCollectionRenderableType.cpp
+++ b/Engine/Plugins/Experimental/Dataflow/Source/DataflowEditor/Private/DataflowRendering/DataflowGeometryCollectionRenderableType.cpp
@@ -25,6 +25,46 @@
 #include "PlanarCut.h" // for ConvertGeometryCollectionToDynamicMesh
 #include "Dataflow/DataflowEditorStyle.h"
 
+
+namespace GeometryCollectionSurfaceRenderableType::Private
+{
+	static const TCHAR* DefaultDataflowSelectionMaterialName = TEXT("/Engine/EditorMaterials/Dataflow/DataflowSelection");
+}
+
+UDataflowGeometryCollectionRenderSettings::UDataflowGeometryCollectionRenderSettings(const FObjectInitializer& ObjectInitializer)
+{
+	SelectionMaterial = LoadObject<UMaterialInterface>(nullptr, GeometryCollectionSurfaceRenderableType::Private::DefaultDataflowSelectionMaterialName);
+
+	constexpr bool bOnlyRGB = true;
+	ColorRamp.SetColorAtTime(0.0f, FLinearColor::Black, bOnlyRGB);
+	ColorRamp.SetColorAtTime(1.0f, FLinearColor::White, bOnlyRGB);
+
+#if WITH_EDITOR
+	ColorRamp.OnColorCurveChangedDelegate.AddWeakLambda(this, [this](const TArray<FRichCurve*>& CHangedCurves)
+		{
+			FPropertyChangedEvent PropertyChangedEvent(UDataflowGeometryCollectionRenderSettings::StaticClass()->FindPropertyByName(GET_MEMBER_NAME_CHECKED(UDataflowGeometryCollectionRenderSettings, ColorRamp)));
+			PostEditChangeProperty(PropertyChangedEvent);
+		});
+#endif
+}
+
+const int32 UDataflowGeometryCollectionRenderSettings::GetHashFromProperties() const
+{
+	int32 TypeHash = 0;
+	TypeHash = HashCombine(TypeHash, GetTypeHash(bShowSelection));
+	TypeHash = HashCombine(TypeHash, GetTypeHash(SelectionMaterial));
+	TypeHash = HashCombine(TypeHash, GetTypeHash(bShowVertexAttributeAsVertexColor));
+	TypeHash = HashCombine(TypeHash, GetTypeHash(VertexAttributeOverride));
+	TypeHash = HashCombine(TypeHash, GetTypeHash(bUseColorRamp));
+
+	const TArray<FRichCurveEditInfoConst> Curves = ColorRamp.GetCurves();
+	for (const FRichCurveEditInfoConst& Curve : Curves)
+	{
+		TypeHash = HashCombine(TypeHash, GetTypeHash(Curve));
+	}
+	return TypeHash;
+}
+
 namespace UE::Dataflow::Private
 {
 	int32 CurveSurfaceVertexThreshold = 50000;
@@ -37,6 +77,7 @@ namespace UE::Dataflow::Private
 		UE_DATAFLOW_IRENDERABLE_PRIMARY_TYPE(FManagedArrayCollection, Collection);
 		UE_DATAFLOW_IRENDERABLE_RENDER_GROUP(Surface);
 		UE_DATAFLOW_IRENDERABLE_VIEW_MODE(FDataflowConstruction3DViewMode);
+		UE_DATAFLOW_IRENDERABLE_SETTINGS(UDataflowGeometryCollectionRenderSettings);
 
 		virtual bool CanRender(const FRenderableTypeInstance& Instance) const
 		{
@@ -46,6 +87,10 @@ namespace UE::Dataflow::Private
 
 		virtual void GetPrimitiveComponents(const FRenderableTypeInstance& Instance, FRenderableComponents& OutComponents) const override
 		{
+			const UDataflowGeometryCollectionRenderSettings* Settings = Instance.GetTypedRenderSettings<UDataflowGeometryCollectionRenderSettings>();
+
+			const bool bShowSelection = Settings ? Settings->bShowSelection : true;
+
 			const FManagedArrayCollection& Collection = GetCollection(Instance);
 
 			// TODO(dataflow) : use a facade or a more generic way to convert a collection to a dynamic mesh
@@ -53,20 +98,30 @@ namespace UE::Dataflow::Private
 			const TManagedArray<int32>* TransformIndexAttribute = Collection.FindAttribute<int32>(FGeometryCollection::TransformIndexAttribute, FGeometryCollection::GeometryGroup);
 			const TManagedArray<FString>* BoneNameAttribute = Collection.FindAttribute<FString>(FTransformCollection::BoneNameAttribute, FTransformCollection::TransformGroup);
 
+			static const FName MeshArrayCacheName = TEXT("MeshArray");
+			static const FName SettingsHashCacheName = TEXT("SettingsHash");
 			if (TransformAttribute && TransformIndexAttribute)
 			{
 				TArray<FDynamicMesh3> DynamicMeshes;
-				if (Instance.HasUptoDateCachedValue())
+
+				// first check if the setting shave changes
+				const int32 CachedSettingsHash = Instance.GetCachedValue<int32>(SettingsHashCacheName);
+				const int32 CurrentSettingsHash = Settings ? Settings->GetHashFromProperties() : 0;
+				if (CachedSettingsHash == CurrentSettingsHash)
 				{
-					const TArray<FDynamicMesh3>& CachedDynamicMeshes = Instance.GetCachedValue<TArray<FDynamicMesh3>>();
-					DynamicMeshes.SetNum(CachedDynamicMeshes.Num());
+					// now check if we have precached meshes
+					if (Instance.HasUptoDateCachedValue(MeshArrayCacheName))
+					{
+						const TArray<FDynamicMesh3>& CachedDynamicMeshes = Instance.GetCachedValue<TArray<FDynamicMesh3>>(MeshArrayCacheName);
+						DynamicMeshes.SetNum(CachedDynamicMeshes.Num());
 
-					// copying dynamic meshes can be expensive let's parallelize it 
-					ParallelFor(DynamicMeshes.Num(),
-						[&CachedDynamicMeshes, &DynamicMeshes](int32 Index)
-						{
-							DynamicMeshes[Index] = CachedDynamicMeshes[Index];
-						});
+						// copying dynamic meshes can be expensive let's parallelize it 
+						ParallelFor(DynamicMeshes.Num(),
+							[&CachedDynamicMeshes, &DynamicMeshes](int32 Index)
+							{
+								DynamicMeshes[Index] = CachedDynamicMeshes[Index];
+							});
+					}
 				}
 				if (DynamicMeshes.Num() != TransformAttribute->Num())
 				{
@@ -81,7 +136,18 @@ namespace UE::Dataflow::Private
 						TArrayView<const FTransform3f> BoneTransforms = MakeArrayView(TransformAttribute->GetConstArray());
 
 						// Override vertex color if we have an attribute to visualize
-						TransferAttributeToVisualizeToVertexColor(*GeomCollection, Instance.GetVertexAttributeToVisualize());
+						const bool bShowAttributeIfAvailable = Settings ? Settings->bShowVertexAttributeAsVertexColor : true;
+						if (bShowAttributeIfAvailable)
+						{
+							FDataflowNode::FAttributeKey AttributeKey;
+							AttributeKey.AttributeName = Settings ? FName(Settings->VertexAttributeOverride): NAME_None;
+							AttributeKey.GroupName = FGeometryCollection::VerticesGroup;
+							if (AttributeKey.AttributeName.IsNone())
+							{
+								AttributeKey = Instance.GetVertexAttributeToVisualize();
+							}
+							TransferAttributeToVisualizeToVertexColor(*GeomCollection, AttributeKey, Settings);
+						}
 
 						// Convert to dynamic meshes in parallel
 						ParallelFor(DynamicMeshes.Num(),
@@ -95,14 +161,20 @@ namespace UE::Dataflow::Private
 								::ConvertGeometryCollectionToDynamicMesh(DynamicMesh, UnusedTransform, false, *GeomCollection, bWeldEdges, BoneTransforms, true, TransformIndicesToConvert);
 							});
 					}
-					Instance.CacheValue(DynamicMeshes);
+					Instance.CacheValue(DynamicMeshes, MeshArrayCacheName);
+					Instance.CacheValue(CurrentSettingsHash, SettingsHashCacheName);
 				}
 
 				FDataflowSelection Selection;
-				GetSelection(Instance, Collection, Selection);
+				if (bShowSelection)
+				{
+					GetSelection(Instance, Collection, Selection);
+				}
 				const bool bIsValidTransformSelection = (Selection.GetGroupName() == FTransformCollection::TransformGroup && Selection.IsValidForCollection(Collection));
 
-				UMaterialInterface* SelectionMaterial = LoadObject<UMaterialInterface>(nullptr, TEXT("/Engine/EditorMaterials/Dataflow/DataflowSelection"));
+				UMaterialInterface* SelectionMaterial = Settings
+					? Settings->SelectionMaterial.Get()
+					: SelectionMaterial = LoadObject<UMaterialInterface>(nullptr, GeometryCollectionSurfaceRenderableType::Private::DefaultDataflowSelectionMaterialName);
 
 				// only create components for the one that have vertices
 				for (int32 TransformIndex = 0; TransformIndex < TransformAttribute->Num(); ++TransformIndex)
@@ -149,12 +221,17 @@ namespace UE::Dataflow::Private
 		}
 
 		template<typename T>
-		void TransferTypedAttributeVertexColor(FGeometryCollection& GeomCollection, const FDataflowNode::FAttributeKey& ValueAttributeKey) const
+		void TransferTypedAttributeVertexColor(FGeometryCollection& GeomCollection, const FDataflowNode::FAttributeKey& ValueAttributeKey, const UDataflowGeometryCollectionRenderSettings* Settings) const
 		{
 			TManagedArray<FLinearColor>* ColorAttribute = GeomCollection.FindAttributeTyped<FLinearColor>(FGeometryCollection::ColorAttribute, FGeometryCollection::VerticesGroup);
 			const TManagedArray<T>* ValueAttribute = GeomCollection.FindAttributeTyped<T>(ValueAttributeKey.AttributeName, ValueAttributeKey.GroupName);
 			if (ColorAttribute && ValueAttribute && ValueAttribute->Num() == ColorAttribute->Num())
 			{
+				static const FLinearColorRamp EmptyColorRamp;
+
+				const FLinearColorRamp& ColorRamp = Settings ? Settings->ColorRamp : EmptyColorRamp;
+				const bool bRemapValue = (Settings && !ColorRamp.IsEmpty()) ? Settings->bUseColorRamp : false;
+
 				const int32 NumVertices = ColorAttribute->Num();
 
 				T MaxValue{ 0 };
@@ -178,18 +255,24 @@ namespace UE::Dataflow::Private
 					const T& Value = (*ValueAttribute)[Index];
 					if constexpr (std::is_same_v<T, float> || std::is_same_v<T, double>)
 					{
-						(*ColorAttribute)[Index] = FLinearColor(Value, Value, Value);
+						(*ColorAttribute)[Index] = bRemapValue
+							? ColorRamp.GetLinearColorValue(Value)
+							: FLinearColor(Value, Value, Value);
 					}
 					else if constexpr (std::is_same_v<T, bool>)
 					{
 						const float FloatValue = Value ? 1.0f : 0.0f;
-						(*ColorAttribute)[Index] = FLinearColor(FloatValue, FloatValue, FloatValue);
+						(*ColorAttribute)[Index] = bRemapValue
+							? ColorRamp.GetLinearColorValue(Value)
+							: FLinearColor(FloatValue, FloatValue, FloatValue);
 					}
 					else if constexpr (std::is_same_v<T, int32>)
 					{
 						const float Multiplier = (Range == 0) ? 1.f : (1.f / Range);
 						const float NormalizedValue{ (Value - MinValue) * Multiplier };
-						(*ColorAttribute)[Index] = FLinearColor(NormalizedValue, NormalizedValue, NormalizedValue);
+						(*ColorAttribute)[Index] = bRemapValue
+							? ColorRamp.GetLinearColorValue(NormalizedValue)
+							: FLinearColor(NormalizedValue, NormalizedValue, NormalizedValue);
 					}
 					else if constexpr (std::is_same_v<T, FVector3f> || std::is_same_v<T, FVector3d>)
 					{
@@ -201,7 +284,7 @@ namespace UE::Dataflow::Private
 			}
 		}
 
-		void TransferAttributeToVisualizeToVertexColor(FGeometryCollection& GeomCollection, const FDataflowNode::FAttributeKey& ValueAttributeKey) const
+		void TransferAttributeToVisualizeToVertexColor(FGeometryCollection& GeomCollection, const FDataflowNode::FAttributeKey& ValueAttributeKey, const UDataflowGeometryCollectionRenderSettings* Settings) const
 		{
 			if (ValueAttributeKey.AttributeName.IsNone() || ValueAttributeKey.GroupName.IsNone())
 			{
@@ -211,22 +294,22 @@ namespace UE::Dataflow::Private
 			switch (AttributeType)
 			{
 			case FManagedArrayCollection::EArrayType::FFloatType:
-				TransferTypedAttributeVertexColor<float>(GeomCollection, ValueAttributeKey);
+				TransferTypedAttributeVertexColor<float>(GeomCollection, ValueAttributeKey, Settings);
 				break;
 			case FManagedArrayCollection::EArrayType::FDoubleType:
-				TransferTypedAttributeVertexColor<double>(GeomCollection, ValueAttributeKey);
+				TransferTypedAttributeVertexColor<double>(GeomCollection, ValueAttributeKey, Settings);
 				break;
 			case FManagedArrayCollection::EArrayType::FVectorType:
-				TransferTypedAttributeVertexColor<FVector3f>(GeomCollection, ValueAttributeKey);
+				TransferTypedAttributeVertexColor<FVector3f>(GeomCollection, ValueAttributeKey, Settings);
 				break;
 			case FManagedArrayCollection::EArrayType::FVector3dType:
-				TransferTypedAttributeVertexColor<FVector3d>(GeomCollection, ValueAttributeKey);
+				TransferTypedAttributeVertexColor<FVector3d>(GeomCollection, ValueAttributeKey, Settings);
 				break;
 			case FManagedArrayCollection::EArrayType::FInt32Type:
-				TransferTypedAttributeVertexColor<int32>(GeomCollection, ValueAttributeKey);
+				TransferTypedAttributeVertexColor<int32>(GeomCollection, ValueAttributeKey, Settings);
 				break;
 			case FManagedArrayCollection::EArrayType::FBoolType:
-				TransferTypedAttributeVertexColor<bool>(GeomCollection, ValueAttributeKey);
+				TransferTypedAttributeVertexColor<bool>(GeomCollection, ValueAttributeKey, Settings);
 				break;
 			}
 		}
diff --git a/Engine/Plugins/Experimental/Dataflow/Source/DataflowEditor/Private/DataflowRendering/DataflowGeometryCollectionRenderableType.h b/Engine/Plugins/Experimental/Dataflow/Source/DataflowEditor/Private/DataflowRendering/DataflowGeometryCollectionRenderableType.h
--- a/Engine/Plugins/Experimental/Dataflow/Source/DataflowEditor/Private/DataflowRendering/DataflowGeometryCollectionRenderableType.h
+++ b/Engine/Plugins/Experimental/Dataflow/Source/DataflowEditor/Private/DataflowRendering/DataflowGeometryCollectionRenderableType.h
@@ -2,6 +2,49 @@
 
 #pragma once
 
+#include "DataflowRendering/DataflowRenderableTypeSettings.h"
+#include "Curves/LinearColorRamp.h"
+#include "UObject/ObjectPtr.h"
+
+#include "DataflowGeometryCollectionRenderableType.generated.h"
+
+class UMaterialInterface;
+
+UCLASS(MinimalAPI, AutoExpandCategories = "Surface|Selection, Surface|Attribute")
+class UDataflowGeometryCollectionRenderSettings : public UDataflowRenderableTypeSettings
+{
+	GENERATED_BODY()
+
+	UDataflowGeometryCollectionRenderSettings(const FObjectInitializer& ObjectInitializer = FObjectInitializer::Get());
+
+public:
+	/** Show selected transforms when relevant to the selected node */
+	UPROPERTY(EditAnywhere, Category = "Surface|Selection")
+	bool bShowSelection = true;
+
+	/** Material used to display the transform selection */
+	UPROPERTY(EditAnywhere, Category = "Surface|Selection", meta = (EditCondition = "bShowSelection"))
+	TObjectPtr<UMaterialInterface> SelectionMaterial = nullptr;
+
+	/** Show vertex attribute if available from the node or the override below */
+	UPROPERTY(EditAnywhere, Category = "Surface|Attribute")
+	bool bShowVertexAttributeAsVertexColor = true;
+
+	/** Vertex attribute to display - if not empty this will override any attribute displayed fro the node */
+	UPROPERTY(EditAnywhere, Category = "Surface|Attribute", meta = (EditCondition = "bShowVertexAttributeAsVertexColor"))
+	FString VertexAttributeOverride;
+
+	/** If true use the color ramp to color vertices from the attribute */
+	UPROPERTY(EditAnywhere, Category = "Surface|Attribute", meta = (EditCondition = "bShowVertexAttributeAsVertexColor"))
+	bool bUseColorRamp = false;
+
+	/** Vertex attribute to display - if not empty this will override any attribute displayed for the node */
+	UPROPERTY(EditAnywhere, Category = "Surface|Attribute", meta = (EditCondition = "bShowVertexAttributeAsVertexColor && bUseColorRamp"))
+	FLinearColorRamp ColorRamp;
+
+	const int32 GetHashFromProperties() const ;
+};
+
 namespace UE::Dataflow::Private
 {
 	void RegisterGeometryCollectionRenderableTypes();
diff --git a/Engine/Plugins/Experimental/Dataflow/Source/DataflowEditor/Public/DataflowRendering/DataflowRenderableTypeInstance.h b/Engine/Plugins/Experimental/Dataflow/Source/DataflowEditor/Public/DataflowRendering/DataflowRenderableTypeInstance.h
--- a/Engine/Plugins/Experimental/Dataflow/Source/DataflowEditor/Public/DataflowRendering/DataflowRenderableTypeInstance.h
+++ b/Engine/Plugins/Experimental/Dataflow/Source/DataflowEditor/Public/DataflowRendering/DataflowRenderableTypeInstance.h
@@ -50,30 +50,33 @@ namespace UE::Dataflow
 			return GetOutputValue(OutputName, Default);
 		}
 
-		bool HasUptoDateCachedValue() const
+		bool HasUptoDateCachedValue(FName CacheName = NAME_None) const
 		{
 			if (Node)
 			{
-				const FContextCacheKey Key = ComputeCacheKey();
+				FContextCacheKey Key = ComputeCacheKey();
+				Key = HashCombine(Key, GetTypeHash(CacheName));
 				return EvaluationContext.HasData(Key, Node->GetTimestamp());
 			}
 			return false;
 		}
 
 		template<typename T>
-		const T& GetCachedValue() const
+		const T& GetCachedValue(FName CacheName = NAME_None) const
 		{
-			static const T StaticDefaultValue;
-			const FContextCacheKey Key = ComputeCacheKey();
+			static const T StaticDefaultValue{};
+			FContextCacheKey Key = ComputeCacheKey();
+			Key = HashCombine(Key, GetTypeHash(CacheName));
 			return EvaluationContext.GetData<T>(Key, nullptr, StaticDefaultValue);
 		}
 
 		template<typename T>
-		void CacheValue(T&& InValue) const
+		void CacheValue(T&& InValue, FName CacheName = NAME_None) const
 		{
 			if (Node)
 			{
-				const FContextCacheKey Key = ComputeCacheKey();
+				FContextCacheKey Key = ComputeCacheKey();
+				Key = HashCombine(Key, GetTypeHash(CacheName));
 				EvaluationContext.SetData(Key, nullptr, Forward<T>(InValue), Node->GetGuid(), /*Node->GetValueHash()*/0, Node->GetTimestamp());
 			}
 		}
