Dataflow : few fixes for geometry collection render settings

#Synced-CL 49325867

diff --git a/Engine/Plugins/Experimental/Dataflow/Source/DataflowEditor/Private/DataflowRendering/DataflowGeometryCollectionRenderableType.cpp b/Engine/Plugins/Experimental/Dataflow/Source/DataflowEditor/Private/DataflowRendering/DataflowGeometryCollectionRenderableType.cpp
--- a/Engine/Plugins/Experimental/Dataflow/Source/DataflowEditor/Private/DataflowRendering/DataflowGeometryCollectionRenderableType.cpp
+++ b/Engine/Plugins/Experimental/Dataflow/Source/DataflowEditor/Private/DataflowRendering/DataflowGeometryCollectionRenderableType.cpp
@@ -54,6 +54,7 @@ const int32 UDataflowGeometryCollectionRenderSettings::GetHashFromProperties() c
 	TypeHash = HashCombine(TypeHash, GetTypeHash(bShowSelection));
 	TypeHash = HashCombine(TypeHash, GetTypeHash(SelectionMaterial));
 	TypeHash = HashCombine(TypeHash, GetTypeHash(bShowVertexAttributeAsVertexColor));
+	TypeHash = HashCombine(TypeHash, GetTypeHash(bUseVertexAttributeOverride));
 	TypeHash = HashCombine(TypeHash, GetTypeHash(VertexAttributeOverride));
 	TypeHash = HashCombine(TypeHash, GetTypeHash(bUseColorRamp));
 
@@ -140,8 +141,11 @@ namespace UE::Dataflow::Private
 						if (bShowAttributeIfAvailable)
 						{
 							FDataflowNode::FAttributeKey AttributeKey;
-							AttributeKey.AttributeName = Settings ? FName(Settings->VertexAttributeOverride): NAME_None;
-							AttributeKey.GroupName = FGeometryCollection::VerticesGroup;
+							if (Settings && Settings->bUseVertexAttributeOverride)
+							{
+								AttributeKey.AttributeName = Settings ? FName(Settings->VertexAttributeOverride) : NAME_None;
+								AttributeKey.GroupName = FGeometryCollection::VerticesGroup;
+							}
 							if (AttributeKey.AttributeName.IsNone())
 							{
 								AttributeKey = Instance.GetVertexAttributeToVisualize();
diff --git a/Engine/Plugins/Experimental/Dataflow/Source/DataflowEditor/Private/DataflowRendering/DataflowGeometryCollectionRenderableType.h b/Engine/Plugins/Experimental/Dataflow/Source/DataflowEditor/Private/DataflowRendering/DataflowGeometryCollectionRenderableType.h
--- a/Engine/Plugins/Experimental/Dataflow/Source/DataflowEditor/Private/DataflowRendering/DataflowGeometryCollectionRenderableType.h
+++ b/Engine/Plugins/Experimental/Dataflow/Source/DataflowEditor/Private/DataflowRendering/DataflowGeometryCollectionRenderableType.h
@@ -23,22 +23,26 @@ public:
 	bool bShowSelection = true;
 
 	/** Material used to display the transform selection */
-	UPROPERTY(EditAnywhere, Category = "Surface|Selection", meta = (EditCondition = "bShowSelection"))
+	UPROPERTY(EditAnywhere, Category = "Surface|Selection", meta = (EditCondition = bShowSelection))
 	TObjectPtr<UMaterialInterface> SelectionMaterial = nullptr;
 
 	/** Show vertex attribute if available from the node or the override below */
 	UPROPERTY(EditAnywhere, Category = "Surface|Attribute")
 	bool bShowVertexAttributeAsVertexColor = true;
 
-	/** Vertex attribute to display - if not empty this will override any attribute displayed fro the node */
-	UPROPERTY(EditAnywhere, Category = "Surface|Attribute", meta = (EditCondition = "bShowVertexAttributeAsVertexColor"))
+	/** control whether the attribute override shoudl be used or not */
+	UPROPERTY(EditAnywhere, Category = "Surface|Attribute", meta = (InlineEditConditionToggle, EditCondition = bShowVertexAttributeAsVertexColor))
+	bool bUseVertexAttributeOverride = false;
+
+	/** Vertex attribute to display - if not empty this will override any attribute displayed for the node */
+	UPROPERTY(EditAnywhere, Category = "Surface|Attribute", meta = (EditCondition = bUseVertexAttributeOverride))
 	FString VertexAttributeOverride;
 
 	/** If true use the color ramp to color vertices from the attribute */
-	UPROPERTY(EditAnywhere, Category = "Surface|Attribute", meta = (EditCondition = "bShowVertexAttributeAsVertexColor"))
+	UPROPERTY(EditAnywhere, Category = "Surface|Attribute", meta = (EditCondition = bShowVertexAttributeAsVertexColor))
 	bool bUseColorRamp = false;
 
-	/** Vertex attribute to display - if not empty this will override any attribute displayed for the node */
+	/** Color ramp to use to display an single value attribute - ignored for vector based attributes */
 	UPROPERTY(EditAnywhere, Category = "Surface|Attribute", meta = (EditCondition = "bShowVertexAttributeAsVertexColor && bUseColorRamp"))
 	FLinearColorRamp ColorRamp;
 
