Dataflow : Fix static analysis warning

#Synced-CL 49325867

diff --git a/Engine/Plugins/Experimental/Dataflow/Source/DataflowEditor/Private/DataflowRendering/DataflowGeometryCollectionRenderableType.cpp b/Engine/Plugins/Experimental/Dataflow/Source/DataflowEditor/Private/DataflowRendering/DataflowGeometryCollectionRenderableType.cpp
--- a/Engine/Plugins/Experimental/Dataflow/Source/DataflowEditor/Private/DataflowRendering/DataflowGeometryCollectionRenderableType.cpp
+++ b/Engine/Plugins/Experimental/Dataflow/Source/DataflowEditor/Private/DataflowRendering/DataflowGeometryCollectionRenderableType.cpp
@@ -178,7 +178,7 @@ namespace UE::Dataflow::Private
 
 				UMaterialInterface* SelectionMaterial = Settings
 					? Settings->SelectionMaterial.Get()
-					: SelectionMaterial = LoadObject<UMaterialInterface>(nullptr, GeometryCollectionSurfaceRenderableType::Private::DefaultDataflowSelectionMaterialName);
+					: LoadObject<UMaterialInterface>(nullptr, GeometryCollectionSurfaceRenderableType::Private::DefaultDataflowSelectionMaterialName);
 
 				// only create components for the one that have vertices
 				for (int32 TransformIndex = 0; TransformIndex < TransformAttribute->Num(); ++TransformIndex)
