Dataflow : View modes now ca return a viewport rotation
- to be used for an incoming change about better support for 2D viewport

#Synced-CL 49372252

diff --git a/Engine/Plugins/ChaosClothAssetEditor/Source/ChaosClothAssetDataflowNodes/Private/ChaosClothAsset/ClothDataflowViewModes.cpp b/Engine/Plugins/ChaosClothAssetEditor/Source/ChaosClothAssetDataflowNodes/Private/ChaosClothAsset/ClothDataflowViewModes.cpp
--- a/Engine/Plugins/ChaosClothAssetEditor/Source/ChaosClothAssetDataflowNodes/Private/ChaosClothAsset/ClothDataflowViewModes.cpp
+++ b/Engine/Plugins/ChaosClothAssetEditor/Source/ChaosClothAssetDataflowNodes/Private/ChaosClothAsset/ClothDataflowViewModes.cpp
@@ -23,6 +23,11 @@ namespace UE::Chaos::ClothAsset
 		return FText(LOCTEXT("Cloth2DSimViewTooltipText", "2D Simulation Mesh"));
 	}
 
+	FRotator FCloth2DSimViewMode::GetViewportOrientation() const
+	{
+		return FRotator(90, -90, 0);
+	}
+
 	const FName FCloth3DSimViewMode::Name = FName("Cloth3DSimView");
 
 	FName FCloth3DSimViewMode::GetName() const
diff --git a/Engine/Plugins/ChaosClothAssetEditor/Source/ChaosClothAssetDataflowNodes/Private/ChaosClothAsset/ClothDataflowViewModes.h b/Engine/Plugins/ChaosClothAssetEditor/Source/ChaosClothAssetDataflowNodes/Private/ChaosClothAsset/ClothDataflowViewModes.h
--- a/Engine/Plugins/ChaosClothAssetEditor/Source/ChaosClothAssetDataflowNodes/Private/ChaosClothAsset/ClothDataflowViewModes.h
+++ b/Engine/Plugins/ChaosClothAssetEditor/Source/ChaosClothAssetDataflowNodes/Private/ChaosClothAsset/ClothDataflowViewModes.h
@@ -15,6 +15,7 @@ namespace UE::Chaos::ClothAsset
 		virtual FName GetName() const override;
 		virtual FText GetButtonText() const override;
 		virtual FText GetTooltipText() const override;
+		virtual FRotator GetViewportOrientation() const override;
 	};
 
 	class FCloth3DSimViewMode : public UE::Dataflow::FDataflowConstruction3DViewModeBase
diff --git a/Engine/Plugins/Experimental/Dataflow/Source/DataflowEditor/Private/Dataflow/DataflowRenderingViewMode.cpp b/Engine/Plugins/Experimental/Dataflow/Source/DataflowEditor/Private/Dataflow/DataflowRenderingViewMode.cpp
--- a/Engine/Plugins/Experimental/Dataflow/Source/DataflowEditor/Private/Dataflow/DataflowRenderingViewMode.cpp
+++ b/Engine/Plugins/Experimental/Dataflow/Source/DataflowEditor/Private/Dataflow/DataflowRenderingViewMode.cpp
@@ -41,6 +41,11 @@ namespace UE::Dataflow
 	{
 		return FText(LOCTEXT("2DViewTooltipText", "Default 2D View"));
 	}
+
+	FRotator FDataflowConstruction2DViewMode::GetViewportOrientation() const
+	{
+		return FRotator(90, -90, 0);
+	}
 	
 	////////////////////////////////////////////////////////////////////////////////////////////////
 
@@ -61,6 +66,11 @@ namespace UE::Dataflow
 		return FText(LOCTEXT("3DViewTooltipText", "Default 3D View"));
 	}
 
+	FRotator FDataflowConstruction3DViewMode::GetViewportOrientation() const
+	{ 
+		return FRotator(30, -30, 0);
+	}
+
 	////////////////////////////////////////////////////////////////////////////////////////////////
 
 	FName FDataflowConstructionUVViewMode::Name = FName("UVView");
@@ -80,6 +90,11 @@ namespace UE::Dataflow
 		return FText(LOCTEXT("UVViewTooltipText", "UV/Texture Coordinate View"));
 	}
 
+	FRotator FDataflowConstructionUVViewMode::GetViewportOrientation() const
+	{
+		return FRotator(-90, 0, 0);
+	}
+
 	////////////////////////////////////////////////////////////////////////////////////////////////
 
 	FName FDataflowConstructionMaterialViewMode::Name = FName("MaterialView");
diff --git a/Engine/Plugins/Experimental/Dataflow/Source/DataflowEditor/Public/Dataflow/DataflowRenderingViewMode.h b/Engine/Plugins/Experimental/Dataflow/Source/DataflowEditor/Public/Dataflow/DataflowRenderingViewMode.h
--- a/Engine/Plugins/Experimental/Dataflow/Source/DataflowEditor/Public/Dataflow/DataflowRenderingViewMode.h
+++ b/Engine/Plugins/Experimental/Dataflow/Source/DataflowEditor/Public/Dataflow/DataflowRenderingViewMode.h
@@ -35,6 +35,7 @@ namespace UE::Dataflow
 		virtual FText GetButtonText() const = 0;
 		virtual FText GetTooltipText() const = 0;
 		virtual ELevelViewportType GetViewportType() const = 0;
+		virtual FRotator GetViewportOrientation() const { return FRotator::ZeroRotator; }
 		
 		DATAFLOWEDITOR_API bool IsPerspective() const;
 	};
@@ -64,6 +65,7 @@ namespace UE::Dataflow
 		DATAFLOWEDITOR_API virtual FName GetName() const override;
 		DATAFLOWEDITOR_API virtual FText GetButtonText() const override;
 		DATAFLOWEDITOR_API virtual FText GetTooltipText() const override;
+		DATAFLOWEDITOR_API virtual FRotator GetViewportOrientation() const override;
 	};
 
 	class FDataflowConstruction3DViewMode : public FDataflowConstruction3DViewModeBase
@@ -75,6 +77,7 @@ namespace UE::Dataflow
 		DATAFLOWEDITOR_API virtual FName GetName() const override;
 		DATAFLOWEDITOR_API virtual FText GetButtonText() const override;
 		DATAFLOWEDITOR_API virtual FText GetTooltipText() const override;
+		DATAFLOWEDITOR_API virtual FRotator GetViewportOrientation() const override;
 	};
 
 	// UV view mode (same as 2D but with a different name and button text)
@@ -87,6 +90,7 @@ namespace UE::Dataflow
 		DATAFLOWEDITOR_API virtual FName GetName() const override;
 		DATAFLOWEDITOR_API virtual FText GetButtonText() const override;
 		DATAFLOWEDITOR_API virtual FText GetTooltipText() const override;
+		DATAFLOWEDITOR_API virtual FRotator GetViewportOrientation() const override;
 	};
 
 	class FDataflowConstructionMaterialViewMode : public FDataflowConstruction3DViewModeBase
