Dataflow : Fix the 2D based view ( UV, 2D and Sim2D )
- Use perspective viewport type to emulate 2D views
- View mode can now provide the right viewport orientation
- Fix various issues with focus methods
- Better control for mouse behviour and zoom

#Synced-CL 49372252

diff --git a/Engine/Plugins/Experimental/Dataflow/Source/DataflowEditor/Private/Dataflow/DataflowConstructionViewport.cpp b/Engine/Plugins/Experimental/Dataflow/Source/DataflowEditor/Private/Dataflow/DataflowConstructionViewport.cpp
--- a/Engine/Plugins/Experimental/Dataflow/Source/DataflowEditor/Private/Dataflow/DataflowConstructionViewport.cpp
+++ b/Engine/Plugins/Experimental/Dataflow/Source/DataflowEditor/Private/Dataflow/DataflowConstructionViewport.cpp
@@ -30,13 +30,13 @@ void SDataflowConstructionViewport::Construct(const FArguments& InArgs, const FA
 {
 	SAssetEditorViewport::FArguments ParentArgs;
 	ParentArgs._EditorViewportClient = InArgs._ViewportClient;
+	DataflowClient = InArgs._ViewportClient;
 	SAssetEditorViewport::Construct(ParentArgs, InViewportConstructionArgs);
 	Client->VisibilityDelegate.BindSP(this, &SDataflowConstructionViewport::IsVisible);
 }
 
 TWeakPtr<FDataflowConstructionScene> SDataflowConstructionViewport::GetConstructionScene() const
 {
-	const TSharedPtr<FDataflowConstructionViewportClient> DataflowClient = StaticCastSharedPtr<FDataflowConstructionViewportClient>(Client);
 	if (DataflowClient)
 	{
 		if (const TSharedPtr<FDataflowEditorToolkit> Toolkit = DataflowClient->GetDataflowEditorToolkit().Pin())
@@ -59,11 +59,11 @@ void SDataflowConstructionViewport::BuildVisualizationMenu(FToolMenuSection& Men
 				{
 					if (const TSharedPtr<SDataflowConstructionViewport> Viewport = StaticCastSharedPtr<SDataflowConstructionViewport>(WeakViewport.Pin()))
 					{
-						if (const TSharedPtr<FEditorViewportClient> ViewportClient = Viewport->GetViewportClient())
+						if (const TSharedPtr<FDataflowConstructionViewportClient> DataflowClient = Viewport->GetDataflowClient())
 						{
 							for (const TPair<FName, TUniquePtr<UE::Dataflow::IDataflowConstructionVisualization>>& Visualization : UE::Dataflow::FDataflowConstructionVisualizationRegistry::GetInstance().GetVisualizations())
 							{
-								Visualization.Value->ExtendViewportShowMenu(StaticCastSharedPtr<FDataflowConstructionViewportClient>(ViewportClient), Menu);
+								Visualization.Value->ExtendViewportShowMenu(DataflowClient, Menu);
 							}
 						}
 					}
@@ -177,22 +177,27 @@ TSharedPtr<IPreviewProfileController> SDataflowConstructionViewport::CreatePrevi
 }
 
 void SDataflowConstructionViewport::OnFocusViewportToSelection()
+{
+	FocusViewportToSelection(/*bInstant*/false);
+}
+
+void SDataflowConstructionViewport::FocusViewportToSelection(bool bInstant)
 {
 	const UDataflowEditorMode* const DataflowEdMode = GetEdMode();
 
-	if (DataflowEdMode)
+	if (DataflowEdMode && DataflowClient)
 	{
 		FBox BoundingBox = DataflowEdMode->SelectionBoundingBox();
 		if (BoundingBox.IsValid && !(BoundingBox.Min == FVector::Zero() && BoundingBox.Max == FVector::Zero()))
 		{
-			if (Client->IsOrtho())
+			if (DataflowClient->IsOrtho())
 			{
 				const FVector Center = BoundingBox.GetCenter();
 				const FVector Extent = BoundingBox.GetExtent() * 0.5;
 				BoundingBox = FBox(Center - Extent, Center + Extent);
 			}
 
-			Client->FocusViewportOnBox(BoundingBox);
+			DataflowClient->FocusViewportOnBoundingBox(BoundingBox, bInstant);
 		}
 	}
 }
@@ -341,7 +346,6 @@ void SDataflowConstructionViewport::PopulateViewportOverlays(TSharedRef<SOverlay
 
 FText SDataflowConstructionViewport::GetOverlayText() const
 {
-	const TSharedPtr<FDataflowConstructionViewportClient> DataflowClient = StaticCastSharedPtr<FDataflowConstructionViewportClient>(Client);
 	if (DataflowClient)
 	{
 		return FText::FromString(DataflowClient->GetOverlayString());
diff --git a/Engine/Plugins/Experimental/Dataflow/Source/DataflowEditor/Private/Dataflow/DataflowConstructionViewportClient.cpp b/Engine/Plugins/Experimental/Dataflow/Source/DataflowEditor/Private/Dataflow/DataflowConstructionViewportClient.cpp
--- a/Engine/Plugins/Experimental/Dataflow/Source/DataflowEditor/Private/Dataflow/DataflowConstructionViewportClient.cpp
+++ b/Engine/Plugins/Experimental/Dataflow/Source/DataflowEditor/Private/Dataflow/DataflowConstructionViewportClient.cpp
@@ -22,6 +22,7 @@
 #include "SGraphPanel.h"
 #include "SNodePanel.h"
 #include "BaseBehaviors/ClickDragBehavior.h"
+#include "BaseBehaviors/MouseWheelBehavior.h"
 #include "Behaviors/2DViewportBehaviorTargets.h"
 
 #define LOCTEXT_NAMESPACE "DataflowConstructionViewportClient"
@@ -60,10 +61,20 @@ FDataflowConstructionViewportClient::FDataflowConstructionViewportClient(FEditor
 	}
 }
 
+bool FDataflowConstructionViewportClient::ShouldOrbitCamera() const
+{
+	if (ConstructionViewMode && ConstructionViewMode->IsPerspective())
+	{
+		return FEditorViewportClient::ShouldOrbitCamera();
+	}
+	return false;
+}
+
+
 void FDataflowConstructionViewportClient::UpdateFromSettings()
 {
 	BehaviorsFor2DMode.Reset();
-	OrthoScrollBehaviorTarget.Reset();
+	OrthoScrollBehaviorTarget = MakeUnique<FEditor2DScrollBehaviorTarget>(this);
 
 	bool bUseRightMouseButton = true;
 	bool bUseMiddleMouseButton = true;
@@ -73,36 +84,40 @@ void FDataflowConstructionViewportClient::UpdateFromSettings()
 		bUseMiddleMouseButton = (Options->ConstructionViewportMousePanButton == EDataflowConstructionViewportMousePanButton::Middle || Options->ConstructionViewportMousePanButton == EDataflowConstructionViewportMousePanButton::RightOrMiddle);
 	}
 
-	if (!bUseRightMouseButton)
-	{
-		// Intercept right mouse drag with a do-nothing ClickDragInputBehavior
-		const TObjectPtr<ULocalClickDragInputBehavior> RightMouseClickDragBehavior = NewObject<ULocalClickDragInputBehavior>();
-		RightMouseClickDragBehavior->Initialize();
+	// We'll have the priority of our viewport behaviors be lower (i.e. higher numerically) than both the gizmo default and the tool default
+	constexpr int ViewportBehaviorPriority = FMath::Max(FInputCapturePriority::DEFAULT_GIZMO_PRIORITY, FInputCapturePriority::DEFAULT_TOOL_PRIORITY) + 3;
 
-		// We'll have the priority of our viewport behaviors be lower (i.e. higher numerically) than both the gizmo default and the tool default
-		constexpr int ViewportBehaviorPriority = FMath::Max(FInputCapturePriority::DEFAULT_GIZMO_PRIORITY, FInputCapturePriority::DEFAULT_TOOL_PRIORITY) + 2;
-		RightMouseClickDragBehavior->SetDefaultPriority(ViewportBehaviorPriority);
-		RightMouseClickDragBehavior->SetUseRightMouseButton();
-		RightMouseClickDragBehavior->CanBeginClickDragFunc = [](const FInputDeviceRay& InputDeviceRay)
-			{
-				return FInputRayHit(TNumericLimits<float>::Max()); // bHit is true. Depth is max to lose the standard tiebreaker.
-			};
-		BehaviorsFor2DMode.Add(RightMouseClickDragBehavior);
+	if (bUseRightMouseButton)
+	{
+		UClickDragInputBehavior* const RightMouseClickDragInputBehavior = NewObject<UClickDragInputBehavior>();
+		RightMouseClickDragInputBehavior->Initialize(OrthoScrollBehaviorTarget.Get());
+		RightMouseClickDragInputBehavior->SetDefaultPriority(ViewportBehaviorPriority);
+		RightMouseClickDragInputBehavior->SetUseRightMouseButton();
+		BehaviorsFor2DMode.Add(RightMouseClickDragInputBehavior);
 	}
 
 	if (bUseMiddleMouseButton)
 	{
-		OrthoScrollBehaviorTarget = MakeUnique<FEditor2DScrollBehaviorTarget>(this);
 		UClickDragInputBehavior* const MiddleMouseClickDragInputBehavior = NewObject<UClickDragInputBehavior>();
 		MiddleMouseClickDragInputBehavior->Initialize(OrthoScrollBehaviorTarget.Get());
-
-		// We'll have the priority of our viewport behaviors be lower (i.e. higher numerically) than both the gizmo default and the tool default
-		constexpr int ViewportBehaviorPriority = FMath::Max(FInputCapturePriority::DEFAULT_GIZMO_PRIORITY, FInputCapturePriority::DEFAULT_TOOL_PRIORITY) + 3;
 		MiddleMouseClickDragInputBehavior->SetDefaultPriority(ViewportBehaviorPriority);
 		MiddleMouseClickDragInputBehavior->SetUseMiddleMouseButton();
 		BehaviorsFor2DMode.Add(MiddleMouseClickDragInputBehavior);
 	}
 
+	ZoomBehaviorTarget = MakeUnique<FEditor2DMouseWheelZoomBehaviorTarget>(this);
+	{
+		constexpr float CameraFarPlaneWorldZ = -10.0f;
+		constexpr float CameraNearPlaneProportionZ = 0.8f;
+		ZoomBehaviorTarget->SetCameraFarPlaneWorldZ(CameraFarPlaneWorldZ);
+		ZoomBehaviorTarget->SetCameraNearPlaneProportionZ(CameraNearPlaneProportionZ);
+		ZoomBehaviorTarget->SetZoomLimits(0.001, 100000);
+		UMouseWheelInputBehavior* const ZoomBehavior = NewObject<UMouseWheelInputBehavior>();
+		ZoomBehavior->Initialize(ZoomBehaviorTarget.Get());
+		ZoomBehavior->SetDefaultPriority(ViewportBehaviorPriority);
+		BehaviorsFor2DMode.Add(ZoomBehavior);
+	}
+
 	if (const UDataflowEditorOptions* const Options = UDataflowEditorOptions::StaticClass()->GetDefaultObject<UDataflowEditorOptions>())
 	{
 		FOVAngle = Options->ConstructionViewFOV;
@@ -173,6 +188,12 @@ void FDataflowConstructionViewportClient::Tick(float DeltaSeconds)
 	}
 }
 
+void FDataflowConstructionViewportClient::ShowMeshEdges(bool bShow)
+{
+	EngineShowFlags.SetMeshEdges(bShow);
+	Invalidate();
+}
+
 USelection* FDataflowConstructionViewportClient::GetSelectedComponents() const 
 {
 	if (const TSharedPtr<FDataflowConstructionScene> DataflowConstructionScenePtr = ConstructionScene.Pin())
@@ -197,8 +218,23 @@ bool FDataflowConstructionViewportClient::InputKey(const FInputKeyEventArgs& Eve
 			return true;
 		}
 	}
+	// Only route input keys to the super class if we are in perspective view mode
+	// This is to avoid being able to orbit the camera in 2D views
+	if (ConstructionViewMode && ConstructionViewMode->IsPerspective())
+	{
+		return FDataflowEditorViewportClientBase::InputKey(EventArgs);
+	}
 
-	return FEditorViewportClient::InputKey(EventArgs);
+	// We'll support disabling input like our base class, even if it does not end up being used.
+	if (bDisableInput)
+	{
+		return true;
+	}
+
+	// Our viewport manipulation is placed in the input router that ModeTools manages
+	return ModeTools
+		? ModeTools->InputKey(this, EventArgs.Viewport, EventArgs.Key, EventArgs.Event)
+		: false;
 }
 
 
@@ -319,6 +355,21 @@ void FDataflowConstructionViewportClient::OnViewportClicked(HHitProxy* HitProxy)
 	OnSelectionChangedMulticast.Broadcast(CurrentlySelectedComponents, DataflowElements);
 }
 
+void FDataflowConstructionViewportClient::FocusViewportOnBoundingBox(const FBox& BoundingBox, bool bInstant)
+{
+	FDataflowEditorViewportClientBase::FocusViewportOnBox(BoundingBox, bInstant);
+
+	// make sure we adjust the near and far planes for 2D since the focus function above may move the camera so that the zero plane is no longer visible
+	if (ConstructionViewMode && !ConstructionViewMode->IsPerspective())
+	{
+		const double AbsZ = FMath::Abs(ViewTransformPerspective.GetLocation().Z);
+		constexpr double CameraFarPlaneWorldZ = -10.0;
+		constexpr double CameraNearPlaneProportionZ = 0.8;
+		OverrideFarClipPlane(static_cast<float>(AbsZ - CameraFarPlaneWorldZ));
+		OverrideNearClipPlane(static_cast<float>(AbsZ * (1.0 - CameraNearPlaneProportionZ)));
+	}
+}
+
 void FDataflowConstructionViewportClient::SetConstructionViewMode(const UE::Dataflow::IDataflowConstructionViewMode* InViewMode)
 {
 	checkf(InViewMode, TEXT("SetConstructionViewMode received null IDataflowConstructionViewMode pointer"));
@@ -330,18 +381,30 @@ void FDataflowConstructionViewportClient::SetConstructionViewMode(const UE::Data
 
 	ConstructionViewMode = InViewMode;
 
-	//
-	// Update input behaviors
-	//
+	if (ConstructionViewMode && ConstructionViewMode->IsPerspective())
+	{
+		OverrideFarClipPlane(0);
+		OverrideNearClipPlane(UE_KINDA_SMALL_NUMBER);
+	}
+	else
+	{
+		const double AbsZ = FMath::Abs(ViewTransformPerspective.GetLocation().Z);
+		constexpr double CameraFarPlaneWorldZ = -10.0;
+		constexpr double CameraNearPlaneProportionZ = 0.8;
+		OverrideFarClipPlane(static_cast<float>(AbsZ - CameraFarPlaneWorldZ));
+		OverrideNearClipPlane(static_cast<float>(AbsZ * (1.0 - CameraNearPlaneProportionZ)));
+	}
 
 	BehaviorSet->RemoveAll();
 
-	for (UInputBehavior* const Behavior : BaseBehaviors)
+	if (InViewMode->IsPerspective())
 	{
-		BehaviorSet->Add(Behavior);
+		for (UInputBehavior* const Behavior : BaseBehaviors)
+		{
+			BehaviorSet->Add(Behavior);
+		}
 	}
-
-	if (!InViewMode->IsPerspective())
+	else
 	{
 		for (UInputBehavior* const Behavior : BehaviorsFor2DMode)
 		{
@@ -355,19 +418,9 @@ void FDataflowConstructionViewportClient::SetConstructionViewMode(const UE::Data
 		ModeTools->GetInteractiveToolsContext()->InputRouter->RegisterSource(this);
 	}
 
-
-	SetViewportType((ELevelViewportType)ConstructionViewMode->GetViewportType());
-
 	if (const FViewportCameraTransform* const FoundPreviousTransform = SavedInactiveViewTransforms.Find(InViewMode->GetName()))
 	{
-		if (ConstructionViewMode->IsPerspective())
-		{
-			ViewTransformPerspective = *FoundPreviousTransform;
-		}
-		else
-		{
-			ViewTransformOrthographic = *FoundPreviousTransform;
-		}
+		ViewTransformPerspective = *FoundPreviousTransform;
 	}
 	else
 	{
@@ -375,10 +428,11 @@ void FDataflowConstructionViewportClient::SetConstructionViewMode(const UE::Data
 	}
 
 	bDrawAxes = ConstructionViewMode->IsPerspective();
+	//EngineShowFlags.SetGrid(ConstructionViewMode->IsPerspective());
+
 	Invalidate();
 }
 
-
 void FDataflowConstructionViewportClient::Draw(const FSceneView* View, FPrimitiveDrawInterface* PDI)
 {
 	FDataflowEditorViewportClientBase::Draw(View, PDI);
diff --git a/Engine/Plugins/Experimental/Dataflow/Source/DataflowEditor/Private/Dataflow/DataflowEditorMode.cpp b/Engine/Plugins/Experimental/Dataflow/Source/DataflowEditor/Private/Dataflow/DataflowEditorMode.cpp
--- a/Engine/Plugins/Experimental/Dataflow/Source/DataflowEditor/Private/Dataflow/DataflowEditorMode.cpp
+++ b/Engine/Plugins/Experimental/Dataflow/Source/DataflowEditor/Private/Dataflow/DataflowEditorMode.cpp
@@ -107,8 +107,7 @@ void UDataflowEditorMode::Enter()
 	}
 
 	// Reset tracking of whether we've seen a valid mesh in 2D or 3D viewports
-	bFirstValid2DMesh = true;
-	bFirstValid3DMesh = true;
+	FirstTimeRefocusedViews.Reset();
 }
 
 void UDataflowEditorMode::SetDataflowEditor(UDataflowEditor* InDataflowEditor) 
@@ -590,6 +589,11 @@ void UDataflowEditorMode::RefocusConstructionViewportClient()
 		// This will happen in FocusViewportOnBox anyways; do it now to get a consistent end result
 		PinnedVC->ToggleOrbitCamera(false);
 
+		const bool bPerspectiveView = (ConstructionViewMode && ConstructionViewMode->IsPerspective());
+		const FVector  ViewportPosition = bPerspectiveView ? FVector(0, 150, 200) : FVector(0, 0, -100);
+		const FRotator ViewportOrientation = ConstructionViewMode ? ConstructionViewMode->GetViewportOrientation() : FRotator::ZeroRotator;
+		PinnedVC->SetInitialViewTransform(ELevelViewportType::LVT_Perspective, ViewportPosition, ViewportOrientation, DEFAULT_ORTHOZOOM);
+
 		const FBox SceneBounds = SceneBoundingBox();
 		constexpr bool bInstant = true;
 		PinnedVC->FocusViewportOnBox(SceneBounds, bInstant);
@@ -625,15 +629,10 @@ void UDataflowEditorMode::FirstTimeFocusConstructionViewport()
 	const TSharedPtr<FDataflowConstructionScene> ConstructionScenePtr = ConstructionScene.Pin();
 	if (ConstructionScenePtr && ConstructionScenePtr->HasRenderableGeometry())
 	{
-		const bool bIs2D = !ConstructionViewMode->IsPerspective();
-		if (bIs2D && bFirstValid2DMesh)
-		{
-			bFirstValid2DMesh = false;
-			RefocusConstructionViewportClient();
-		}
-		else if (!bIs2D && bFirstValid3DMesh)
+		const FName ViewModeName = ConstructionViewMode? ConstructionViewMode->GetName(): NAME_None;
+		if (!FirstTimeRefocusedViews.Contains(ViewModeName))
 		{
-			bFirstValid3DMesh = false;
+			FirstTimeRefocusedViews.Add(ViewModeName);
 			RefocusConstructionViewportClient();
 		}
 	}
diff --git a/Engine/Plugins/Experimental/Dataflow/Source/DataflowEditor/Public/Dataflow/DataflowConstructionViewport.h b/Engine/Plugins/Experimental/Dataflow/Source/DataflowEditor/Public/Dataflow/DataflowConstructionViewport.h
--- a/Engine/Plugins/Experimental/Dataflow/Source/DataflowEditor/Public/Dataflow/DataflowConstructionViewport.h
+++ b/Engine/Plugins/Experimental/Dataflow/Source/DataflowEditor/Public/Dataflow/DataflowConstructionViewport.h
@@ -6,6 +6,7 @@
 
 class UDataflowEditorMode;
 class FDataflowConstructionScene;
+class FDataflowConstructionViewportClient;
 class SRichTextBlock;
 
 struct FToolMenuSection;
@@ -16,7 +17,7 @@ class SDataflowConstructionViewport : public SBaseCharacterFXEditorViewport, pub
 {
 public:
 	SLATE_BEGIN_ARGS(SDataflowConstructionViewport) {}
-	SLATE_ARGUMENT(TSharedPtr<FEditorViewportClient>, ViewportClient)
+	SLATE_ARGUMENT(TSharedPtr<FDataflowConstructionViewportClient>, ViewportClient)
 	SLATE_END_ARGS()
 
 	SDataflowConstructionViewport();
@@ -37,6 +38,9 @@ public:
 	virtual void PopulateViewportOverlays(TSharedRef<SOverlay> Overlay) override;
 
 	FText GetOverlayText() const;
+	void FocusViewportToSelection(bool bInstant);
+
+	TSharedPtr<FDataflowConstructionViewportClient> GetDataflowClient() const { return DataflowClient; }
 
 private:
 	static void BuildVisualizationMenu(FToolMenuSection& MenuSection);
@@ -44,6 +48,8 @@ private:
 
 	TWeakPtr<FDataflowConstructionScene> GetConstructionScene() const;
 
+	TSharedPtr<FDataflowConstructionViewportClient> DataflowClient;
+
 	/** Pointer to the box into which the overlay text items are added */
 	TSharedPtr<SRichTextBlock> OverlayText;
 
diff --git a/Engine/Plugins/Experimental/Dataflow/Source/DataflowEditor/Public/Dataflow/DataflowConstructionViewportClient.h b/Engine/Plugins/Experimental/Dataflow/Source/DataflowEditor/Public/Dataflow/DataflowConstructionViewportClient.h
--- a/Engine/Plugins/Experimental/Dataflow/Source/DataflowEditor/Public/Dataflow/DataflowConstructionViewportClient.h
+++ b/Engine/Plugins/Experimental/Dataflow/Source/DataflowEditor/Public/Dataflow/DataflowConstructionViewportClient.h
@@ -12,6 +12,7 @@ class FDataflowEditorToolkit;
 class FDataflowConstructionScene;
 class USelection;
 class IClickDragBehaviorTarget;
+class FEditor2DMouseWheelZoomBehaviorTarget;
 namespace UE::Dataflow
 {
 	class IDataflowConstructionViewMode;
@@ -51,6 +52,10 @@ public:
 	/** Return the widget mode only used for tooling */
 	UE::Widget::EWidgetMode GetToolWidgetMode() const {return WidgetMode;}
 
+	UE_API void ShowMeshEdges(bool bShow);
+
+	UE_API void FocusViewportOnBoundingBox(const FBox& BoundingBox, bool bInstant);
+
 protected:
 
 	// FGCObject
@@ -71,10 +76,12 @@ private:
 	UE_API virtual bool CanSetWidgetMode(UE::Widget::EWidgetMode NewMode) const override;
 	UE_API virtual UE::Widget::EWidgetMode GetWidgetMode() const override;
 	UE_API virtual void SetWidgetMode(UE::Widget::EWidgetMode NewMode) override;
+	UE_API virtual bool ShouldOrbitCamera() const override;
 
 	UE_API void UpdateFromSettings();
 
 	TUniquePtr<IClickDragBehaviorTarget> OrthoScrollBehaviorTarget;
+	TUniquePtr<FEditor2DMouseWheelZoomBehaviorTarget> ZoomBehaviorTarget;
 	TArray<TObjectPtr<UInputBehavior>> BehaviorsFor2DMode;
 
 	/** Toolkit used to create the viewport client */
diff --git a/Engine/Plugins/Experimental/Dataflow/Source/DataflowEditor/Public/Dataflow/DataflowEditorMode.h b/Engine/Plugins/Experimental/Dataflow/Source/DataflowEditor/Public/Dataflow/DataflowEditorMode.h
--- a/Engine/Plugins/Experimental/Dataflow/Source/DataflowEditor/Public/Dataflow/DataflowEditorMode.h
+++ b/Engine/Plugins/Experimental/Dataflow/Source/DataflowEditor/Public/Dataflow/DataflowEditorMode.h
@@ -226,8 +226,7 @@ private:
 	void OnConstructionSceneUpdated();
 
 	// The first time we get a valid mesh, refocus the camera on it
-	bool bFirstValid2DMesh = true;
-	bool bFirstValid3DMesh = true;
+	TSet<FName> FirstTimeRefocusedViews;
 
 	// Whether the rest space viewport should focus on the rest space mesh on the next tick
 	bool bShouldFocusConstructionView = true;
