Dataflow : several editor fixes
- Make sure we display one set of redner settigs when two nodes are displayed in the construction view
- Add settings for geometry collection UVs render types
- make uv nodes to use the new rendering system
- Fix default camera orientation for 3D view

#Synced-CL 49372252

diff --git a/Engine/Plugins/Experimental/Dataflow/Source/DataflowEditor/Private/Dataflow/DataflowConstructionScene.cpp b/Engine/Plugins/Experimental/Dataflow/Source/DataflowEditor/Private/Dataflow/DataflowConstructionScene.cpp
--- a/Engine/Plugins/Experimental/Dataflow/Source/DataflowEditor/Private/Dataflow/DataflowConstructionScene.cpp
+++ b/Engine/Plugins/Experimental/Dataflow/Source/DataflowEditor/Private/Dataflow/DataflowConstructionScene.cpp
@@ -52,7 +52,7 @@ TArray<UDataflowRenderableTypeSettings*> FDataflowConstructionScene::GetRenderab
 	{
 		if (UDataflowRenderableTypeSettings* RenderSettings = RenderableInstance.GetRenderSettings())
 		{
-			RenderSettingsArray.Add(RenderSettings);
+			RenderSettingsArray.AddUnique(RenderSettings);
 		}
 	}
 	return RenderSettingsArray;
diff --git a/Engine/Plugins/Experimental/Dataflow/Source/DataflowEditor/Private/Dataflow/DataflowRenderingViewMode.cpp b/Engine/Plugins/Experimental/Dataflow/Source/DataflowEditor/Private/Dataflow/DataflowRenderingViewMode.cpp
--- a/Engine/Plugins/Experimental/Dataflow/Source/DataflowEditor/Private/Dataflow/DataflowRenderingViewMode.cpp
+++ b/Engine/Plugins/Experimental/Dataflow/Source/DataflowEditor/Private/Dataflow/DataflowRenderingViewMode.cpp
@@ -68,7 +68,7 @@ namespace UE::Dataflow
 
 	FRotator FDataflowConstruction3DViewMode::GetViewportOrientation() const
 	{ 
-		return FRotator(30, -30, 0);
+		return FRotator(-30, -30, 0);
 	}
 
 	////////////////////////////////////////////////////////////////////////////////////////////////
diff --git a/Engine/Plugins/Experimental/Dataflow/Source/DataflowEditor/Private/DataflowRendering/DataflowGeometryCollectionRenderableType.cpp b/Engine/Plugins/Experimental/Dataflow/Source/DataflowEditor/Private/DataflowRendering/DataflowGeometryCollectionRenderableType.cpp
--- a/Engine/Plugins/Experimental/Dataflow/Source/DataflowEditor/Private/DataflowRendering/DataflowGeometryCollectionRenderableType.cpp
+++ b/Engine/Plugins/Experimental/Dataflow/Source/DataflowEditor/Private/DataflowRendering/DataflowGeometryCollectionRenderableType.cpp
@@ -3,9 +3,7 @@
 #include "DataflowRendering/DataflowGeometryCollectionRenderableType.h"
 
 #include "Components/DynamicMeshComponent.h"
-
 #include "Dataflow/DataflowRenderingViewMode.h"
-
 #include "DataflowRendering/DataflowRenderableComponents.h"
 #include "DataflowRendering/DataflowRenderableType.h"
 #include "DataflowRendering/DataflowRenderableTypeInstance.h"
@@ -29,9 +27,41 @@
 namespace GeometryCollectionSurfaceRenderableType::Private
 {
 	static const TCHAR* DefaultDataflowSelectionMaterialName = TEXT("/Engine/EditorMaterials/Dataflow/DataflowSelection");
+
+	static void AddLeavesToSelection(int32 TransformIndex, const TManagedArray<TSet<int32>>& ChildrenAttribute, FDataflowSelection& OutSelection)
+	{
+		const TSet<int32>& Children = ChildrenAttribute[TransformIndex];
+		if (Children.IsEmpty())
+		{
+			OutSelection.SetSelected(TransformIndex);
+		}
+		else
+		{
+			for (const int32 Child : Children)
+			{
+				AddLeavesToSelection(Child, ChildrenAttribute, OutSelection);
+			}
+		}
+	}
+
+	static void AddLeavesToSelection(const FManagedArrayCollection& Collection, const FDataflowSelection& InSelection, FDataflowSelection& OutSelection)
+	{
+		if (const TManagedArray<TSet<int32>>* ChildrenAttributePtr = Collection.FindAttribute<TSet<int32>>(FTransformCollection::ChildrenAttribute, FTransformCollection::TransformGroup))
+		{
+			const int32 NumTransforms = ChildrenAttributePtr->Num();
+			for (int32 TransformIndex = 0; TransformIndex < NumTransforms; ++TransformIndex)
+			{
+				if (InSelection.IsSelected(TransformIndex))
+				{
+					AddLeavesToSelection(TransformIndex, *ChildrenAttributePtr, OutSelection);
+				}
+			}
+		}
+	}
 }
 
-UDataflowGeometryCollectionRenderSettings::UDataflowGeometryCollectionRenderSettings(const FObjectInitializer& ObjectInitializer)
+////////////////////////////////////////////////////////////////////////////////////////////////////////////////////////////////
+UDataflowGeometryCollectionSurfaceRenderSettings::UDataflowGeometryCollectionSurfaceRenderSettings(const FObjectInitializer& ObjectInitializer)
 {
 	SelectionMaterial = LoadObject<UMaterialInterface>(nullptr, GeometryCollectionSurfaceRenderableType::Private::DefaultDataflowSelectionMaterialName);
 
@@ -42,13 +72,13 @@ UDataflowGeometryCollectionRenderSettings::UDataflowGeometryCollectionRenderSett
 #if WITH_EDITOR
 	ColorRamp.OnColorCurveChangedDelegate.AddWeakLambda(this, [this](const TArray<FRichCurve*>& CHangedCurves)
 		{
-			FPropertyChangedEvent PropertyChangedEvent(UDataflowGeometryCollectionRenderSettings::StaticClass()->FindPropertyByName(GET_MEMBER_NAME_CHECKED(UDataflowGeometryCollectionRenderSettings, ColorRamp)));
+			FPropertyChangedEvent PropertyChangedEvent(UDataflowGeometryCollectionSurfaceRenderSettings::StaticClass()->FindPropertyByName(GET_MEMBER_NAME_CHECKED(UDataflowGeometryCollectionSurfaceRenderSettings, ColorRamp)));
 			PostEditChangeProperty(PropertyChangedEvent);
 		});
 #endif
 }
 
-const int32 UDataflowGeometryCollectionRenderSettings::GetHashFromProperties() const
+const int32 UDataflowGeometryCollectionSurfaceRenderSettings::GetHashFromProperties() const
 {
 	int32 TypeHash = 0;
 	TypeHash = HashCombine(TypeHash, GetTypeHash(bShowSelection));
@@ -66,6 +96,8 @@ const int32 UDataflowGeometryCollectionRenderSettings::GetHashFromProperties() c
 	return TypeHash;
 }
 
+////////////////////////////////////////////////////////////////////////////////////////////////////////////////////////////////
+
 namespace UE::Dataflow::Private
 {
 	int32 CurveSurfaceVertexThreshold = 50000;
@@ -78,7 +110,7 @@ namespace UE::Dataflow::Private
 		UE_DATAFLOW_IRENDERABLE_PRIMARY_TYPE(FManagedArrayCollection, Collection);
 		UE_DATAFLOW_IRENDERABLE_RENDER_GROUP(Surface);
 		UE_DATAFLOW_IRENDERABLE_VIEW_MODE(FDataflowConstruction3DViewMode);
-		UE_DATAFLOW_IRENDERABLE_SETTINGS(UDataflowGeometryCollectionRenderSettings);
+		UE_DATAFLOW_IRENDERABLE_SETTINGS(UDataflowGeometryCollectionSurfaceRenderSettings);
 
 		virtual bool CanRender(const FRenderableTypeInstance& Instance) const
 		{
@@ -88,7 +120,7 @@ namespace UE::Dataflow::Private
 
 		virtual void GetPrimitiveComponents(const FRenderableTypeInstance& Instance, FRenderableComponents& OutComponents) const override
 		{
-			const UDataflowGeometryCollectionRenderSettings* Settings = Instance.GetTypedRenderSettings<UDataflowGeometryCollectionRenderSettings>();
+			const UDataflowGeometryCollectionSurfaceRenderSettings* Settings = Instance.GetTypedRenderSettings<UDataflowGeometryCollectionSurfaceRenderSettings>();
 
 			const bool bShowSelection = Settings ? Settings->bShowSelection : true;
 
@@ -217,15 +249,14 @@ namespace UE::Dataflow::Private
 			}
 			if (OutSelection.GetGroupName() == FTransformCollection::TransformGroup && OutSelection.AnySelected())
 			{
-				GeometryCollection::Facades::FCollectionTransformSelectionFacade SelectionFacade(Collection);
-				TArray<int32> SelectionToConvert = OutSelection.AsArray();
-				SelectionFacade.ConvertSelectionToRigidNodes(SelectionToConvert);
-				OutSelection.SetFromArray(SelectionToConvert);
+				const FDataflowSelection TransformSelection = OutSelection;
+				OutSelection.Clear();
+				GeometryCollectionSurfaceRenderableType::Private::AddLeavesToSelection(Collection, TransformSelection, OutSelection);
 			}
 		}
 
 		template<typename T>
-		void TransferTypedAttributeVertexColor(FGeometryCollection& GeomCollection, const FDataflowNode::FAttributeKey& ValueAttributeKey, const UDataflowGeometryCollectionRenderSettings* Settings) const
+		void TransferTypedAttributeVertexColor(FGeometryCollection& GeomCollection, const FDataflowNode::FAttributeKey& ValueAttributeKey, const UDataflowGeometryCollectionSurfaceRenderSettings* Settings) const
 		{
 			TManagedArray<FLinearColor>* ColorAttribute = GeomCollection.FindAttributeTyped<FLinearColor>(FGeometryCollection::ColorAttribute, FGeometryCollection::VerticesGroup);
 			const TManagedArray<T>* ValueAttribute = GeomCollection.FindAttributeTyped<T>(ValueAttributeKey.AttributeName, ValueAttributeKey.GroupName);
@@ -288,7 +319,7 @@ namespace UE::Dataflow::Private
 			}
 		}
 
-		void TransferAttributeToVisualizeToVertexColor(FGeometryCollection& GeomCollection, const FDataflowNode::FAttributeKey& ValueAttributeKey, const UDataflowGeometryCollectionRenderSettings* Settings) const
+		void TransferAttributeToVisualizeToVertexColor(FGeometryCollection& GeomCollection, const FDataflowNode::FAttributeKey& ValueAttributeKey, const UDataflowGeometryCollectionSurfaceRenderSettings* Settings) const
 		{
 			if (ValueAttributeKey.AttributeName.IsNone() || ValueAttributeKey.GroupName.IsNone())
 			{
@@ -373,6 +404,7 @@ namespace UE::Dataflow::Private
 		UE_DATAFLOW_IRENDERABLE_PRIMARY_TYPE(FManagedArrayCollection, Collection);
 		UE_DATAFLOW_IRENDERABLE_RENDER_GROUP(Uvs);
 		UE_DATAFLOW_IRENDERABLE_VIEW_MODE(FDataflowConstructionUVViewMode);
+		UE_DATAFLOW_IRENDERABLE_SETTINGS(UDataflowGeometryCollectionUvsRenderSettings);
 
 		virtual bool CanRender(const FRenderableTypeInstance& Instance) const
 		{
@@ -382,9 +414,12 @@ namespace UE::Dataflow::Private
 
 		virtual void GetPrimitiveComponents(const FRenderableTypeInstance& Instance, FRenderableComponents& OutComponents) const override
 		{
+			const UDataflowGeometryCollectionUvsRenderSettings* Settings = Instance.GetTypedRenderSettings<UDataflowGeometryCollectionUvsRenderSettings>();
+
 			const FManagedArrayCollection& Collection = GetCollection(Instance);
 
-			Rendering::AddUvDynamicMeshComponent(Collection, 0, OutComponents);
+			const int32 UVChannel = Settings ? Settings->UVChannel : 0;
+			Rendering::AddUvDynamicMeshComponent(Collection, UVChannel, OutComponents);
 		}
 	};
 
diff --git a/Engine/Plugins/Experimental/Dataflow/Source/DataflowEditor/Private/DataflowRendering/DataflowGeometryCollectionRenderableType.h b/Engine/Plugins/Experimental/Dataflow/Source/DataflowEditor/Private/DataflowRendering/DataflowGeometryCollectionRenderableType.h
--- a/Engine/Plugins/Experimental/Dataflow/Source/DataflowEditor/Private/DataflowRendering/DataflowGeometryCollectionRenderableType.h
+++ b/Engine/Plugins/Experimental/Dataflow/Source/DataflowEditor/Private/DataflowRendering/DataflowGeometryCollectionRenderableType.h
@@ -11,11 +11,11 @@
 class UMaterialInterface;
 
 UCLASS(MinimalAPI, AutoExpandCategories = "Surface|Selection, Surface|Attribute")
-class UDataflowGeometryCollectionRenderSettings : public UDataflowRenderableTypeSettings
+class UDataflowGeometryCollectionSurfaceRenderSettings : public UDataflowRenderableTypeSettings
 {
 	GENERATED_BODY()
 
-	UDataflowGeometryCollectionRenderSettings(const FObjectInitializer& ObjectInitializer = FObjectInitializer::Get());
+	UDataflowGeometryCollectionSurfaceRenderSettings(const FObjectInitializer& ObjectInitializer = FObjectInitializer::Get());
 
 public:
 	/** Show selected transforms when relevant to the selected node */
@@ -49,6 +49,18 @@ public:
 	const int32 GetHashFromProperties() const ;
 };
 
+
+UCLASS(MinimalAPI, AutoExpandCategories = "UVs")
+class UDataflowGeometryCollectionUvsRenderSettings : public UDataflowRenderableTypeSettings
+{
+	GENERATED_BODY()
+
+public:
+	/** UV channel to display */
+	UPROPERTY(EditAnywhere, Category = "UVs", meta = (UIMIn = 0, UIMax = 7, ClampMin = 0, ClampMax = 7))
+	int32 UVChannel = 0;
+};
+
 namespace UE::Dataflow::Private
 {
 	void RegisterGeometryCollectionRenderableTypes();
diff --git a/Engine/Plugins/Experimental/GeometryCollectionPlugin/Source/GeometryCollectionNodes/Public/Dataflow/GeometryCollectionUVNodes.h b/Engine/Plugins/Experimental/GeometryCollectionPlugin/Source/GeometryCollectionNodes/Public/Dataflow/GeometryCollectionUVNodes.h
--- a/Engine/Plugins/Experimental/GeometryCollectionPlugin/Source/GeometryCollectionNodes/Public/Dataflow/GeometryCollectionUVNodes.h
+++ b/Engine/Plugins/Experimental/GeometryCollectionPlugin/Source/GeometryCollectionNodes/Public/Dataflow/GeometryCollectionUVNodes.h
@@ -17,7 +17,6 @@ struct FAddUVChannelDataflowNode : public FDataflowNode
 {
 	GENERATED_USTRUCT_BODY()
 	DATAFLOW_NODE_DEFINE_INTERNAL(FAddUVChannelDataflowNode, "AddUVChannel", "GeometryCollection|UV", "")
-	DATAFLOW_NODE_RENDER_TYPE("SurfaceRender", FGeometryCollection::StaticType(), "Collection", "UVChannel")
 
 private:
 	/** Target Collection */
@@ -46,7 +45,6 @@ struct FAutoUnwrapUVDataflowNode : public FDataflowNode
 {
 	GENERATED_USTRUCT_BODY()
 	DATAFLOW_NODE_DEFINE_INTERNAL(FAutoUnwrapUVDataflowNode, "AutoUnwrapUV", "GeometryCollection|UV", "")
-	DATAFLOW_NODE_RENDER_TYPE("SurfaceRender", FGeometryCollection::StaticType(), "Collection", "UVChannel")
 
 private:
 	/** Target Collection */
@@ -79,7 +77,6 @@ struct FMergeUVIslandsDataflowNode : public FDataflowNode
 {
 	GENERATED_USTRUCT_BODY()
 	DATAFLOW_NODE_DEFINE_INTERNAL(FMergeUVIslandsDataflowNode, "MergeUVIslands", "GeometryCollection|UV", "")
-	DATAFLOW_NODE_RENDER_TYPE("SurfaceRender", FGeometryCollection::StaticType(), "Collection", "UVChannel")
 
 private:
 	/** Target Collection */
@@ -122,7 +119,6 @@ struct FBoxProjectUVDataflowNode : public FDataflowNode
 {
 	GENERATED_USTRUCT_BODY()
 	DATAFLOW_NODE_DEFINE_INTERNAL(FBoxProjectUVDataflowNode, "BoxProjectUV", "GeometryCollection|UV", "")
-	DATAFLOW_NODE_RENDER_TYPE("SurfaceRender", FGeometryCollection::StaticType(), "Collection", "UVChannel")
 
 private:
 	/** Target Collection */
