Dataflow : update a lot of nodes to use the new rendering system
- Also add support for metadata on properties in node  to specific  if the property is a UVChannel ( meta=(DataflowUVChannel) )

#Synced-CL 49372252

diff --git a/Engine/Plugins/Experimental/Dataflow/Source/DataflowEditor/Private/DataflowRendering/DataflowDynamicMeshRenderableType.cpp b/Engine/Plugins/Experimental/Dataflow/Source/DataflowEditor/Private/DataflowRendering/DataflowDynamicMeshRenderableType.cpp
--- a/Engine/Plugins/Experimental/Dataflow/Source/DataflowEditor/Private/DataflowRendering/DataflowDynamicMeshRenderableType.cpp
+++ b/Engine/Plugins/Experimental/Dataflow/Source/DataflowEditor/Private/DataflowRendering/DataflowDynamicMeshRenderableType.cpp
@@ -4,16 +4,20 @@
 
 #include "Components/DynamicMeshComponent.h"
 
+#include "Dataflow/DataflowMesh.h"
 #include "Dataflow/DataflowRenderingViewMode.h"
 
 #include "DataflowRendering/DataflowRenderableComponents.h"
 #include "DataflowRendering/DataflowRenderableType.h"
 #include "DataflowRendering/DataflowRenderableTypeInstance.h"
 #include "DataflowRendering/DataflowRenderableTypeRegistry.h"
+#include "DataflowRendering/DataflowRenderableTypeSettings.h"
 #include "DataflowRendering/DataflowRenderableTypeUtils.h"
 
 #include "UObject/ObjectPtr.h"
 
+UE_DATAFLOW_POLICY_DECLARE_TYPENAME(TObjectPtr<UMaterialInterface>);
+
 namespace UE::Dataflow::Private
 {
 	///////////////////////////////////////////////////////////////////////////////////////////////////////////////
@@ -31,6 +35,44 @@ namespace UE::Dataflow::Private
 				if (UDynamicMeshComponent* Component = OutComponents.AddNewComponent<UDynamicMeshComponent>())
 				{
 					Component->SetDynamicMesh(DynamicMesh);
+
+					using FMaterialArray = TArray<TObjectPtr<UMaterialInterface>>;
+					const FMaterialArray EmptyArray;
+					const FMaterialArray& Materials = Instance.GetOutputValueByType<FMaterialArray>(EmptyArray);
+					for (int32 Index = 0; Index < Materials.Num(); ++Index)
+					{
+						Component->SetMaterial(Index, Materials[Index]);
+					}
+				}
+			}
+		}
+	};
+
+	////////////////////////////////////////////////////////////////////////////////////////////////////////////////////////////
+
+	class FDataflowMeshSurfaceRenderableType : public IRenderableType
+	{
+		UE_DATAFLOW_IRENDERABLE_PRIMARY_TYPE(TObjectPtr<UDataflowMesh>, DataflowMesh);
+		UE_DATAFLOW_IRENDERABLE_RENDER_GROUP(Surface);
+		UE_DATAFLOW_IRENDERABLE_VIEW_MODE(FDataflowConstruction3DViewMode);
+
+		virtual void GetPrimitiveComponents(const FRenderableTypeInstance& Instance, FRenderableComponents& OutComponents) const override
+		{
+			if (const TObjectPtr<UDataflowMesh> DataflowMesh = GetDataflowMesh(Instance, nullptr))
+			{
+				if (const UE::Geometry::FDynamicMesh3* DynamicMesh = DataflowMesh->GetDynamicMesh())
+				{
+					if (UDynamicMeshComponent* Component = OutComponents.AddNewComponent<UDynamicMeshComponent>())
+					{
+						UE::Geometry::FDynamicMesh3 MeshToSet = *DynamicMesh;
+						Component->SetMesh(MoveTemp(MeshToSet));
+
+						const TArray<TObjectPtr<UMaterialInterface>>& Materials = DataflowMesh->GetMaterials();
+						for (int32 Index = 0; Index < Materials.Num(); ++Index)
+						{
+							Component->SetMaterial(Index, Materials[Index]);
+						}
+					}
 				}
 			}
 		}
@@ -68,12 +110,44 @@ namespace UE::Dataflow::Private
 		UE_DATAFLOW_IRENDERABLE_PRIMARY_TYPE(TObjectPtr<UDynamicMesh>, DynamicMesh);
 		UE_DATAFLOW_IRENDERABLE_RENDER_GROUP(Uvs);
 		UE_DATAFLOW_IRENDERABLE_VIEW_MODE(FDataflowConstructionUVViewMode);
+		UE_DATAFLOW_IRENDERABLE_SETTINGS(UDataflowUVsRenderSettings);
 
 		virtual void GetPrimitiveComponents(const FRenderableTypeInstance& Instance, FRenderableComponents& OutComponents) const override
 		{
+			const UDataflowUVsRenderSettings* Settings = Instance.GetTypedRenderSettings<UDataflowUVsRenderSettings>();
+
 			if (const TObjectPtr<UDynamicMesh> DynamicMesh = GetDynamicMesh(Instance, nullptr))
 			{
-				Rendering::AddUvDynamicMeshComponent(DynamicMesh->GetMeshRef(), 0, OutComponents);
+				const int32 UVChannel = Settings
+					? Settings->GetUVChannel(Rendering::GetUVChannelFromInstance(Instance))
+					: 0;
+				Rendering::AddUvDynamicMeshComponent(DynamicMesh->GetMeshRef(), UVChannel, OutComponents);
+			}
+		}
+	};
+
+	////////////////////////////////////////////////////////////////////////////////////////////////////////////////////////////
+
+	class FDataflowMeshUvsRenderableType : public IRenderableType
+	{
+		UE_DATAFLOW_IRENDERABLE_PRIMARY_TYPE(TObjectPtr<UDataflowMesh>, DataflowMesh);
+		UE_DATAFLOW_IRENDERABLE_RENDER_GROUP(Uvs);
+		UE_DATAFLOW_IRENDERABLE_VIEW_MODE(FDataflowConstructionUVViewMode);
+		UE_DATAFLOW_IRENDERABLE_SETTINGS(UDataflowUVsRenderSettings);
+
+		virtual void GetPrimitiveComponents(const FRenderableTypeInstance& Instance, FRenderableComponents& OutComponents) const override
+		{
+			const UDataflowUVsRenderSettings* Settings = Instance.GetTypedRenderSettings<UDataflowUVsRenderSettings>();
+
+			if (const TObjectPtr<UDataflowMesh> DataflowMesh = GetDataflowMesh(Instance, nullptr))
+			{
+				if (const UE::Geometry::FDynamicMesh3* DynamicMesh = DataflowMesh->GetDynamicMesh())
+				{
+					const int32 UVChannel = Settings
+						? Settings->GetUVChannel(Rendering::GetUVChannelFromInstance(Instance))
+						: 0;
+					Rendering::AddUvDynamicMeshComponent(*DynamicMesh, UVChannel, OutComponents);
+				}
 			}
 		}
 	};
@@ -83,8 +157,10 @@ namespace UE::Dataflow::Private
 	void RegisterDynamicMeshRenderableTypes()
 	{
 		UE_DATAFLOW_REGISTER_RENDERABLE_TYPE(FDynamicMeshSurfaceRenderableType);
+		UE_DATAFLOW_REGISTER_RENDERABLE_TYPE(FDataflowMeshSurfaceRenderableType);
 		UE_DATAFLOW_REGISTER_RENDERABLE_TYPE(FDynamicMeshArraySurfaceRenderableType);
 		UE_DATAFLOW_REGISTER_RENDERABLE_TYPE(FDynamicMeshUvsRenderableType);
+		UE_DATAFLOW_REGISTER_RENDERABLE_TYPE(FDataflowMeshUvsRenderableType);
 	}
 
 	///////////////////////////////////////////////////////////////////////////////////////////////////////////////
diff --git a/Engine/Plugins/Experimental/Dataflow/Source/DataflowEditor/Private/DataflowRendering/DataflowGeometryCollectionRenderableType.cpp b/Engine/Plugins/Experimental/Dataflow/Source/DataflowEditor/Private/DataflowRendering/DataflowGeometryCollectionRenderableType.cpp
--- a/Engine/Plugins/Experimental/Dataflow/Source/DataflowEditor/Private/DataflowRendering/DataflowGeometryCollectionRenderableType.cpp
+++ b/Engine/Plugins/Experimental/Dataflow/Source/DataflowEditor/Private/DataflowRendering/DataflowGeometryCollectionRenderableType.cpp
@@ -404,7 +404,7 @@ namespace UE::Dataflow::Private
 		UE_DATAFLOW_IRENDERABLE_PRIMARY_TYPE(FManagedArrayCollection, Collection);
 		UE_DATAFLOW_IRENDERABLE_RENDER_GROUP(Uvs);
 		UE_DATAFLOW_IRENDERABLE_VIEW_MODE(FDataflowConstructionUVViewMode);
-		UE_DATAFLOW_IRENDERABLE_SETTINGS(UDataflowGeometryCollectionUvsRenderSettings);
+		UE_DATAFLOW_IRENDERABLE_SETTINGS(UDataflowUVsRenderSettings);
 
 		virtual bool CanRender(const FRenderableTypeInstance& Instance) const
 		{
@@ -414,11 +414,13 @@ namespace UE::Dataflow::Private
 
 		virtual void GetPrimitiveComponents(const FRenderableTypeInstance& Instance, FRenderableComponents& OutComponents) const override
 		{
-			const UDataflowGeometryCollectionUvsRenderSettings* Settings = Instance.GetTypedRenderSettings<UDataflowGeometryCollectionUvsRenderSettings>();
+			const UDataflowUVsRenderSettings* Settings = Instance.GetTypedRenderSettings<UDataflowUVsRenderSettings>();
 
 			const FManagedArrayCollection& Collection = GetCollection(Instance);
 
-			const int32 UVChannel = Settings ? Settings->UVChannel : 0;
+			const int32 UVChannel = Settings
+				? Settings->GetUVChannel(Rendering::GetUVChannelFromInstance(Instance))
+				: 0;
 			Rendering::AddUvDynamicMeshComponent(Collection, UVChannel, OutComponents);
 		}
 	};
diff --git a/Engine/Plugins/Experimental/Dataflow/Source/DataflowEditor/Private/DataflowRendering/DataflowGeometryCollectionRenderableType.h b/Engine/Plugins/Experimental/Dataflow/Source/DataflowEditor/Private/DataflowRendering/DataflowGeometryCollectionRenderableType.h
--- a/Engine/Plugins/Experimental/Dataflow/Source/DataflowEditor/Private/DataflowRendering/DataflowGeometryCollectionRenderableType.h
+++ b/Engine/Plugins/Experimental/Dataflow/Source/DataflowEditor/Private/DataflowRendering/DataflowGeometryCollectionRenderableType.h
@@ -49,18 +49,6 @@ public:
 	const int32 GetHashFromProperties() const ;
 };
 
-
-UCLASS(MinimalAPI, AutoExpandCategories = "UVs")
-class UDataflowGeometryCollectionUvsRenderSettings : public UDataflowRenderableTypeSettings
-{
-	GENERATED_BODY()
-
-public:
-	/** UV channel to display */
-	UPROPERTY(EditAnywhere, Category = "UVs", meta = (UIMIn = 0, UIMax = 7, ClampMin = 0, ClampMax = 7))
-	int32 UVChannel = 0;
-};
-
 namespace UE::Dataflow::Private
 {
 	void RegisterGeometryCollectionRenderableTypes();
diff --git a/Engine/Plugins/Experimental/Dataflow/Source/DataflowEditor/Private/DataflowRendering/DataflowRenderableTypeInstance.cpp b/Engine/Plugins/Experimental/Dataflow/Source/DataflowEditor/Private/DataflowRendering/DataflowRenderableTypeInstance.cpp
--- a/Engine/Plugins/Experimental/Dataflow/Source/DataflowEditor/Private/DataflowRendering/DataflowRenderableTypeInstance.cpp
+++ b/Engine/Plugins/Experimental/Dataflow/Source/DataflowEditor/Private/DataflowRendering/DataflowRenderableTypeInstance.cpp
@@ -111,6 +111,15 @@ namespace UE::Dataflow
 		return ::HashCombine(OutputCacheKey, RenderGroupHash);
 	}
 
+	bool FRenderableTypeInstance::HasMetadata(const FDataflowOutput& Output, FName Metadata)
+	{
+		if (const FProperty* Property = Output.GetProperty())
+		{
+			return Property->HasMetaData(Metadata);
+		}
+		return false;
+	}
+
 	UDataflowRenderableTypeSettings* FRenderableTypeInstance::GetRenderSettings() const
 	{
 		if (RenderableType && Node)
diff --git a/Engine/Plugins/Experimental/Dataflow/Source/DataflowEditor/Private/DataflowRendering/DataflowRenderableTypeSettings.h b/Engine/Plugins/Experimental/Dataflow/Source/DataflowEditor/Private/DataflowRendering/DataflowRenderableTypeSettings.h
--- a/Engine/Plugins/Experimental/Dataflow/Source/DataflowEditor/Private/DataflowRendering/DataflowRenderableTypeSettings.h
+++ b/Engine/Plugins/Experimental/Dataflow/Source/DataflowEditor/Private/DataflowRendering/DataflowRenderableTypeSettings.h
@@ -12,3 +12,25 @@ class UDataflowRenderableTypeSettings : public UObject
 	GENERATED_BODY()
 };
 
+// Common types
+
+UCLASS(MinimalAPI, AutoExpandCategories = "UVs")
+class UDataflowUVsRenderSettings : public UDataflowRenderableTypeSettings
+{
+	GENERATED_BODY()
+
+public:
+	/** return the overriden channel if enabled or the default passed to this method */
+	int32 GetUVChannel(int32 DefaultChannel) const
+	{
+		return bOverrideUVchannel ? UVChannel : DefaultChannel;
+	}
+
+private:
+	UPROPERTY(EditAnywhere, Category = "UVs", meta = (InlineEditConditionToggle))
+	bool bOverrideUVchannel = false;
+
+	/** UV channel to display if oiverride is enabled */
+	UPROPERTY(EditAnywhere, Category = "UVs", meta = (UIMIn = 0, UIMax = 7, ClampMin = 0, ClampMax = 7, EditCondition = bOverrideUVchannel))
+	int32 UVChannel = 0;
+};
diff --git a/Engine/Plugins/Experimental/Dataflow/Source/DataflowEditor/Private/DataflowRendering/DataflowRenderableTypeUtils.cpp b/Engine/Plugins/Experimental/Dataflow/Source/DataflowEditor/Private/DataflowRendering/DataflowRenderableTypeUtils.cpp
--- a/Engine/Plugins/Experimental/Dataflow/Source/DataflowEditor/Private/DataflowRendering/DataflowRenderableTypeUtils.cpp
+++ b/Engine/Plugins/Experimental/Dataflow/Source/DataflowEditor/Private/DataflowRendering/DataflowRenderableTypeUtils.cpp
@@ -2,6 +2,7 @@
 
 #include "DataflowRendering/DataflowRenderableTypeUtils.h"
 #include "DataflowRendering/DataflowRenderableComponents.h"
+#include "DataflowRendering/DataflowRenderableTypeInstance.h"
 
 #include "Components/DynamicMeshComponent.h"
 #include "Engine/SkeletalMesh.h"
@@ -42,10 +43,14 @@ namespace UE::Dataflow::Rendering
 
 	TSharedPtr<UE::Geometry::FDynamicMesh3> GenerateUvMesh(const FMeshDescription& MeshDescription, int32 UvChannel)
 	{
-		UE::Geometry::FMeshDescriptionUVsToDynamicMesh Converter;
-		Converter.UVLayerIndex = UvChannel;
-		Converter.ScaleFactor = Private::UvScaleFactor;
-		return Converter.GetUVMesh(&MeshDescription);
+		if (UvChannel < MeshDescription.GetNumUVElementChannels())
+		{
+			UE::Geometry::FMeshDescriptionUVsToDynamicMesh Converter;
+			Converter.UVLayerIndex = UvChannel;
+			Converter.ScaleFactor = Private::UvScaleFactor;
+			return Converter.GetUVMesh(&MeshDescription);
+		}
+		return {};
 	}
 
 	///////////////////////////////////////////////////////////////////////////////////////////////////////////////
@@ -150,4 +155,11 @@ namespace UE::Dataflow::Rendering
 		}
 		return nullptr;
 	}
+
+	///////////////////////////////////////////////////////////////////////////////////////////////////////////////
+	int32 GetUVChannelFromInstance(const FRenderableTypeInstance& Instance, int32 DefaultValue)
+	{
+		const static FName DataflowUVChannelName(TEXT("DataflowUVChannel"));
+		return Instance.GetOutputValueByType<int32>(DefaultValue, DataflowUVChannelName);
+	}
 }
diff --git a/Engine/Plugins/Experimental/Dataflow/Source/DataflowEditor/Private/DataflowRendering/DataflowSkeletalMeshRenderableType.cpp b/Engine/Plugins/Experimental/Dataflow/Source/DataflowEditor/Private/DataflowRendering/DataflowSkeletalMeshRenderableType.cpp
--- a/Engine/Plugins/Experimental/Dataflow/Source/DataflowEditor/Private/DataflowRendering/DataflowSkeletalMeshRenderableType.cpp
+++ b/Engine/Plugins/Experimental/Dataflow/Source/DataflowEditor/Private/DataflowRendering/DataflowSkeletalMeshRenderableType.cpp
@@ -10,6 +10,7 @@
 #include "DataflowRendering/DataflowRenderableType.h"
 #include "DataflowRendering/DataflowRenderableTypeInstance.h"
 #include "DataflowRendering/DataflowRenderableTypeRegistry.h"
+#include "DataflowRendering/DataflowRenderableTypeSettings.h"
 #include "DataflowRendering/DataflowRenderableTypeUtils.h"
 
 #include "Engine/SkeletalMesh.h"
@@ -45,12 +46,18 @@ namespace UE::Dataflow::Private
 		UE_DATAFLOW_IRENDERABLE_PRIMARY_TYPE(TObjectPtr<USkeletalMesh>, SkeletalMesh);
 		UE_DATAFLOW_IRENDERABLE_RENDER_GROUP(Uvs);
 		UE_DATAFLOW_IRENDERABLE_VIEW_MODE(FDataflowConstructionUVViewMode);
+		UE_DATAFLOW_IRENDERABLE_SETTINGS(UDataflowUVsRenderSettings);
 
 		virtual void GetPrimitiveComponents(const FRenderableTypeInstance& Instance, FRenderableComponents& OutComponents) const override
 		{
+			const UDataflowUVsRenderSettings* Settings = Instance.GetTypedRenderSettings<UDataflowUVsRenderSettings>();
+
 			if (const TObjectPtr<USkeletalMesh> SkeletalMesh = GetSkeletalMesh(Instance, nullptr))
 			{
-				Rendering::AddUvDynamicMeshComponent(*SkeletalMesh, 0, OutComponents);
+				const int32 UVChannel = Settings
+					? Settings->GetUVChannel(Rendering::GetUVChannelFromInstance(Instance))
+					: 0;
+				Rendering::AddUvDynamicMeshComponent(*SkeletalMesh, UVChannel, OutComponents);
 			}
 		}
 	};
diff --git a/Engine/Plugins/Experimental/Dataflow/Source/DataflowEditor/Public/DataflowRendering/DataflowRenderableTypeInstance.h b/Engine/Plugins/Experimental/Dataflow/Source/DataflowEditor/Public/DataflowRendering/DataflowRenderableTypeInstance.h
--- a/Engine/Plugins/Experimental/Dataflow/Source/DataflowEditor/Public/DataflowRendering/DataflowRenderableTypeInstance.h
+++ b/Engine/Plugins/Experimental/Dataflow/Source/DataflowEditor/Public/DataflowRendering/DataflowRenderableTypeInstance.h
@@ -109,9 +109,9 @@ namespace UE::Dataflow
 			}
 			return Default;
 		}
-		/** Find the first output of a specific match type and return its value - non-anytype version */
+		/** Find the first output of a specific match type and optionally containing a specific metadata string and return its value - non-anytype version */
 		template <typename T UE_REQUIRES(!std::is_base_of_v<FDataflowAnyType, T>)>
-		const T& GetOutputValueByType(const T& Default) const
+		const T& GetOutputValueByType(const T& Default, FName Metadata = NAME_None) const
 		{
 			if (Node)
 			{
@@ -119,7 +119,10 @@ namespace UE::Dataflow
 				{
 					if (Output && Output->GetType() == TDataflowPolicyTypeName<T>::GetName())
 					{
-						return GetOutputValue<T>(Output->GetName(), Default);
+						if (Metadata.IsNone() || HasMetadata(*Output, Metadata))
+						{
+							return GetOutputValue<T>(Output->GetName(), Default);
+						}
 					}
 				}
 			}
@@ -128,7 +131,7 @@ namespace UE::Dataflow
 
 		/** Find the first output of a specific match type and return its value - anytype version */
 		template <typename T UE_REQUIRES(std::is_base_of_v<FDataflowAnyType, T>)>
-		const T::FStorageType& GetOutputValueByType(const T::FStorageType& Default) const
+		const T::FStorageType& GetOutputValueByType(const T::FStorageType& Default, FName Metadata = NAME_None) const
 		{
 			if (Node)
 			{
@@ -136,7 +139,10 @@ namespace UE::Dataflow
 				{
 					if (Output && T::FPolicyType::SupportsTypeStatic(Output->GetType()))
 					{
-						return GetOutputValue<typename T::FStorageType>(Output->GetName(), Default);
+						if (Metadata.IsNone() || HasMetadata(*Output, Metadata))
+						{
+							return GetOutputValue<typename T::FStorageType>(Output->GetName(), Default);
+						}
 					}
 				}
 			}
@@ -146,6 +152,8 @@ namespace UE::Dataflow
 	private:
 		UE_API FContextCacheKey ComputeCacheKey() const;
 
+		UE_API static bool HasMetadata(const FDataflowOutput& Output, FName Metadata);
+
 		FRenderableTypeInstance(FContext& InContext, const TSharedPtr<const FDataflowNode>& InNode, const FName InOutputName, const IRenderableType* InRenderableType)
 			: EvaluationContext(InContext)
 			, Node(InNode)
diff --git a/Engine/Plugins/Experimental/Dataflow/Source/DataflowEditor/Public/DataflowRendering/DataflowRenderableTypeUtils.h b/Engine/Plugins/Experimental/Dataflow/Source/DataflowEditor/Public/DataflowRendering/DataflowRenderableTypeUtils.h
--- a/Engine/Plugins/Experimental/Dataflow/Source/DataflowEditor/Public/DataflowRendering/DataflowRenderableTypeUtils.h
+++ b/Engine/Plugins/Experimental/Dataflow/Source/DataflowEditor/Public/DataflowRendering/DataflowRenderableTypeUtils.h
@@ -22,6 +22,7 @@ namespace UE::Geometry
 namespace UE::Dataflow
 {
 	struct FRenderableComponents;
+	struct FRenderableTypeInstance;
 
 	namespace Rendering
 	{
@@ -61,6 +62,8 @@ namespace UE::Dataflow
 				}
 			}
 		}
+
+		int32 GetUVChannelFromInstance(const FRenderableTypeInstance& Instance, int32 DefaultValue = 0);
 	}
 }
 
diff --git a/Engine/Plugins/Experimental/DataflowVolumetricPlugin/Source/DataflowVolumeNodes/Private/Dataflow/DataflowVolumeNodes.h b/Engine/Plugins/Experimental/DataflowVolumetricPlugin/Source/DataflowVolumeNodes/Private/Dataflow/DataflowVolumeNodes.h
--- a/Engine/Plugins/Experimental/DataflowVolumetricPlugin/Source/DataflowVolumeNodes/Private/Dataflow/DataflowVolumeNodes.h
+++ b/Engine/Plugins/Experimental/DataflowVolumetricPlugin/Source/DataflowVolumeNodes/Private/Dataflow/DataflowVolumeNodes.h
@@ -458,7 +458,6 @@ struct FConvertVolumeDataflowNode : public FDataflowNode
 {
 	GENERATED_USTRUCT_BODY()
 	DATAFLOW_NODE_DEFINE_INTERNAL(FConvertVolumeDataflowNode, "ConvertVolume", "Volume|Utilities", "")
-	DATAFLOW_NODE_RENDER_TYPE("SurfaceRender", FGeometryCollection::StaticType(), "Collection")
 
 private:
 	/** Float volume */
diff --git a/Engine/Plugins/Experimental/GeometryCollectionPlugin/Source/GeometryCollectionNodes/Public/Dataflow/GeometryCollectionTextureNodes.h b/Engine/Plugins/Experimental/GeometryCollectionPlugin/Source/GeometryCollectionNodes/Public/Dataflow/GeometryCollectionTextureNodes.h
--- a/Engine/Plugins/Experimental/GeometryCollectionPlugin/Source/GeometryCollectionNodes/Public/Dataflow/GeometryCollectionTextureNodes.h
+++ b/Engine/Plugins/Experimental/GeometryCollectionPlugin/Source/GeometryCollectionNodes/Public/Dataflow/GeometryCollectionTextureNodes.h
@@ -55,7 +55,6 @@ struct FBakeTextureFromCollectionDataflowNode : public FDataflowNode
 {
 	GENERATED_USTRUCT_BODY()
 	DATAFLOW_NODE_DEFINE_INTERNAL(FBakeTextureFromCollectionDataflowNode, "BakeTextureFromCollection", "GeometryCollection|Texture", "")
-	DATAFLOW_NODE_RENDER_TYPE("SurfaceRender", FGeometryCollection::StaticType(), "Collection", "UVChannel")
 
 private:
 	/** Target Collection */
@@ -79,7 +78,7 @@ private:
 	int32 GutterSize = 1;
 
 	/** Index of the added UV channel */
-	UPROPERTY(EditAnywhere, Category = ChannelOptions, meta = (DataflowInput, DataflowOutput, DataflowPassthrough = "UVChannel", ClampMin = 0, ClampMax = 7, UIMin = 0, UIMax = 7))
+	UPROPERTY(EditAnywhere, Category = ChannelOptions, meta = (DataflowInput, DataflowOutput, DataflowPassthrough = "UVChannel", DataflowUVChannel, ClampMin = 0, ClampMax = 7, UIMin = 0, UIMax = 7))
 	int32 UVChannel = 0;
 
 	/** Attribute to bake in the red channel */
diff --git a/Engine/Plugins/Experimental/GeometryCollectionPlugin/Source/GeometryCollectionNodes/Public/Dataflow/GeometryCollectionUVNodes.h b/Engine/Plugins/Experimental/GeometryCollectionPlugin/Source/GeometryCollectionNodes/Public/Dataflow/GeometryCollectionUVNodes.h
--- a/Engine/Plugins/Experimental/GeometryCollectionPlugin/Source/GeometryCollectionNodes/Public/Dataflow/GeometryCollectionUVNodes.h
+++ b/Engine/Plugins/Experimental/GeometryCollectionPlugin/Source/GeometryCollectionNodes/Public/Dataflow/GeometryCollectionUVNodes.h
@@ -24,7 +24,7 @@ private:
 	FManagedArrayCollection Collection;
 
 	/** Index of the added UV channel */
-	UPROPERTY(meta = (DataflowOutput))
+	UPROPERTY(meta = (DataflowOutput, DataflowUVChannel))
 	int32 UVChannel = 0;
 
 	/** Value to initialize the UV with */
@@ -56,7 +56,7 @@ private:
 	FDataflowFaceSelection FaceSelection;
 
 	/** UV channel to unwrap into ( 0 by default ) */
-	UPROPERTY(EditAnywhere, Category = Options, meta = (DataflowInput, DataflowOutput, DataflowPassthrough = "UVChannel", ClampMin = 0, ClampMax = 7, UIMin = 0, UIMax = 7))
+	UPROPERTY(EditAnywhere, Category = Options, meta = (DataflowInput, DataflowOutput, DataflowPassthrough = "UVChannel", DataflowUVChannel, ClampMin = 0, ClampMax = 7, UIMin = 0, UIMax = 7))
 	int32 UVChannel = 0;
 
 	/** Approximate space to leave between UV islands, measured in texels for 512x512 texture */
@@ -88,7 +88,7 @@ private:
 	FDataflowFaceSelection FaceSelection;
 
 	/** UV channel to unwrap into ( 0 by default ) */
-	UPROPERTY(EditAnywhere, Category = Options, meta = (DataflowInput, DataflowOutput, DataflowPassthrough = "UVChannel", ClampMin = 0, ClampMax = 7, UIMin = 0, UIMax = 7))
+	UPROPERTY(EditAnywhere, Category = Options, meta = (DataflowInput, DataflowOutput, DataflowPassthrough = "UVChannel", DataflowUVChannel, ClampMin = 0, ClampMax = 7, UIMin = 0, UIMax = 7))
 	int32 UVChannel = 0;
 
 	// Threshold for allowed area distortion from merging islands (when we use ExpMap to compute new UVs for the merged island)
@@ -130,7 +130,7 @@ private:
 	//FDataflowFaceSelection FaceSelection;
 
 	/** UV channel to unwrap into ( 0 by default ) */
-	UPROPERTY(EditAnywhere, Category = Options, meta = (DataflowInput, DataflowOutput, DataflowPassthrough = "UVChannel", ClampMin = 0, ClampMax = 7, UIMin = 0, UIMax = 7))
+	UPROPERTY(EditAnywhere, Category = Options, meta = (DataflowInput, DataflowOutput, DataflowPassthrough = "UVChannel", DataflowUVChannel, ClampMin = 0, ClampMax = 7, UIMin = 0, UIMax = 7))
 	int32 UVChannel = 0;
 
 	/** Approximate space to leave between UV islands, measured in texels for 512x512 texture */
diff --git a/Engine/Plugins/Experimental/GeometryDataflow/Source/GeometryDataflowNodes/Public/Dataflow/MeshBooleanNodes.h b/Engine/Plugins/Experimental/GeometryDataflow/Source/GeometryDataflowNodes/Public/Dataflow/MeshBooleanNodes.h
--- a/Engine/Plugins/Experimental/GeometryDataflow/Source/GeometryDataflowNodes/Public/Dataflow/MeshBooleanNodes.h
+++ b/Engine/Plugins/Experimental/GeometryDataflow/Source/GeometryDataflowNodes/Public/Dataflow/MeshBooleanNodes.h
@@ -53,7 +53,6 @@ struct FMeshBooleanDataflowNode : public FDataflowNode
 {
 	GENERATED_USTRUCT_BODY()
 	DATAFLOW_NODE_DEFINE_INTERNAL(FMeshBooleanDataflowNode, "MeshBoolean", "Mesh|Utilities", "")
-	DATAFLOW_NODE_RENDER_TYPE("SurfaceRender",FName("FDynamicMesh3"), "Mesh")
 
 public:
 
diff --git a/Engine/Plugins/Experimental/MeshResizing/Source/MeshResizingNodes/Private/MeshResizing/AlignUVMeshNode.h b/Engine/Plugins/Experimental/MeshResizing/Source/MeshResizingNodes/Private/MeshResizing/AlignUVMeshNode.h
--- a/Engine/Plugins/Experimental/MeshResizing/Source/MeshResizingNodes/Private/MeshResizing/AlignUVMeshNode.h
+++ b/Engine/Plugins/Experimental/MeshResizing/Source/MeshResizingNodes/Private/MeshResizing/AlignUVMeshNode.h
@@ -12,7 +12,6 @@ struct FAlignUVMeshNode : public FDataflowNode
 {
 	GENERATED_USTRUCT_BODY()
 	DATAFLOW_NODE_DEFINE_INTERNAL(FAlignUVMeshNode, "AlignUVMeshNode", "MeshResizing", "Align UV Mesh")
-	DATAFLOW_NODE_RENDER_TYPE("SurfaceRender", FName("UDataflowMesh"), "ResizingMesh", "UVChannelIndex")
 
 public:
 
@@ -26,7 +25,7 @@ private:
 	UPROPERTY(meta = (DataflowInput))
 	TObjectPtr<UDataflowMesh> BaseMesh;
 
-	UPROPERTY(EditAnywhere, Category = "Align", meta = (DataflowInput, DataflowOutput, DataflowPassthrough = "UVChannelIndex", ClampMin = 0, ClampMax = 7, UIMin = 0, UIMax = 7))
+	UPROPERTY(EditAnywhere, Category = "Align", meta = (DataflowInput, DataflowOutput, DataflowPassthrough = "UVChannelIndex", DataflowUVChannel, ClampMin = 0, ClampMax = 7, UIMin = 0, UIMax = 7))
 	int32 UVChannelIndex = 0;
 
 	/** Base UV channel index in case it differs from the ResizingMesh UV channel index, or -1 to use the same channel. */
diff --git a/Engine/Plugins/Experimental/MeshResizing/Source/MeshResizingNodes/Private/MeshResizing/BaseBodyDataflowNodes.h b/Engine/Plugins/Experimental/MeshResizing/Source/MeshResizingNodes/Private/MeshResizing/BaseBodyDataflowNodes.h
--- a/Engine/Plugins/Experimental/MeshResizing/Source/MeshResizingNodes/Private/MeshResizing/BaseBodyDataflowNodes.h
+++ b/Engine/Plugins/Experimental/MeshResizing/Source/MeshResizingNodes/Private/MeshResizing/BaseBodyDataflowNodes.h
@@ -22,7 +22,6 @@ struct FSkeletalMeshToMeshDataflowNode final : public FDataflowNode
 {
 	GENERATED_USTRUCT_BODY()
 	DATAFLOW_NODE_DEFINE_INTERNAL(FSkeletalMeshToMeshDataflowNode, "SkeletalMeshToMesh", "Mesh|Utilities", "")
-	DATAFLOW_NODE_RENDER_TYPE("SurfaceRender",FName("FDynamicMesh3"), {"Mesh", "MaterialArray"})
 
 public:
 
@@ -79,7 +78,6 @@ struct FGenerateResizableProxyDataflowNode final : public FDataflowNode
 {
 	GENERATED_USTRUCT_BODY()
 	DATAFLOW_NODE_DEFINE_INTERNAL(FGenerateResizableProxyDataflowNode, "GenerateResizableProxy", "MeshResizing", "")
-	DATAFLOW_NODE_RENDER_TYPE("SurfaceRender",FName("FDynamicMesh3"), {"TargetProxyMesh", "ProxyMaterialArray"})
 
 public:
 
@@ -143,7 +141,6 @@ struct FGenerateInterpolatedProxyDataflowNode final : public FDataflowNode
 {
 	GENERATED_USTRUCT_BODY()
 	DATAFLOW_NODE_DEFINE_INTERNAL(FGenerateInterpolatedProxyDataflowNode, "GenerateInterpolatedProxy", "MeshResizing", "")
-	DATAFLOW_NODE_RENDER_TYPE("SurfaceRender", FName("FDynamicMesh3"), { "ProxyMesh", "ProxyMaterialArray" })
 
 public:
 
diff --git a/Engine/Plugins/Experimental/MeshResizing/Source/MeshResizingNodes/Private/MeshResizing/MeshConstraintNodes.h b/Engine/Plugins/Experimental/MeshResizing/Source/MeshResizingNodes/Private/MeshResizing/MeshConstraintNodes.h
--- a/Engine/Plugins/Experimental/MeshResizing/Source/MeshResizingNodes/Private/MeshResizing/MeshConstraintNodes.h
+++ b/Engine/Plugins/Experimental/MeshResizing/Source/MeshResizingNodes/Private/MeshResizing/MeshConstraintNodes.h
@@ -16,7 +16,6 @@ struct FMeshConstrainedDeformationNode : public FDataflowNode
 {
 	GENERATED_USTRUCT_BODY()
 	DATAFLOW_NODE_DEFINE_INTERNAL(FMeshConstrainedDeformationNode, "MeshConstrainedDeformationTestPlayground", "MeshResizing", "Mesh Constrained Deformation")
-	DATAFLOW_NODE_RENDER_TYPE("SurfaceRender", FName("UDataflowMesh"), "ResizingMesh")
 
 public:
 
diff --git a/Engine/Plugins/Experimental/MeshResizing/Source/MeshResizingNodes/Private/MeshResizing/RBFInterpolationNodes.h b/Engine/Plugins/Experimental/MeshResizing/Source/MeshResizingNodes/Private/MeshResizing/RBFInterpolationNodes.h
--- a/Engine/Plugins/Experimental/MeshResizing/Source/MeshResizingNodes/Private/MeshResizing/RBFInterpolationNodes.h
+++ b/Engine/Plugins/Experimental/MeshResizing/Source/MeshResizingNodes/Private/MeshResizing/RBFInterpolationNodes.h
@@ -64,7 +64,6 @@ struct FApplyRBFResizingNode : public FDataflowNode
 {
 	GENERATED_USTRUCT_BODY()
 	DATAFLOW_NODE_DEFINE_INTERNAL(FApplyRBFResizingNode, "ApplyRBFResizing", "MeshResizing", "Apply RBF Resizing")
-	DATAFLOW_NODE_RENDER_TYPE("SurfaceRender", FName("UDataflowMesh"), "ResizedMesh")
 
 public:
 	FApplyRBFResizingNode(const UE::Dataflow::FNodeParameters& InParam, FGuid InGuid = FGuid::NewGuid());
diff --git a/Engine/Plugins/Experimental/MeshResizing/Source/MeshResizingNodes/Private/MeshResizing/UVMeshTransformNode.h b/Engine/Plugins/Experimental/MeshResizing/Source/MeshResizingNodes/Private/MeshResizing/UVMeshTransformNode.h
--- a/Engine/Plugins/Experimental/MeshResizing/Source/MeshResizingNodes/Private/MeshResizing/UVMeshTransformNode.h
+++ b/Engine/Plugins/Experimental/MeshResizing/Source/MeshResizingNodes/Private/MeshResizing/UVMeshTransformNode.h
@@ -12,7 +12,6 @@ struct FUVMeshTransformNode : public FDataflowNode
 {
 	GENERATED_USTRUCT_BODY()
 	DATAFLOW_NODE_DEFINE_INTERNAL(FUVMeshTransformNode, "UVMeshTransformNode", "MeshResizing", "UV Mesh Transform")
-	DATAFLOW_NODE_RENDER_TYPE("SurfaceRender", FName("UDataflowMesh"), "Mesh", "UVChannelIndex")
 
 public:
 
@@ -23,7 +22,7 @@ private:
 	UPROPERTY(meta = (DataflowInput, DataflowOutput, DataflowPassthrough = "Mesh"))
 	TObjectPtr<UDataflowMesh> Mesh;
 
-	UPROPERTY(EditAnywhere, Category = "Mesh", meta = (DataflowInput, DataflowOutput, DataflowPassthrough = "UVChannelIndex", ClampMin = 0, ClampMax = 7, UIMin = 0, UIMax = 7))
+	UPROPERTY(EditAnywhere, Category = "Mesh", meta = (DataflowInput, DataflowOutput, DataflowPassthrough = "UVChannelIndex", DataflowUVChannel, ClampMin = 0, ClampMax = 7, UIMin = 0, UIMax = 7))
 	int32 UVChannelIndex = 0;
 
 	UPROPERTY(EditAnywhere, Category = "Transform", Meta = (AllowPreserveRatio))
diff --git a/Engine/Plugins/Experimental/MeshResizing/Source/MeshResizingNodes/Private/MeshResizing/UVResizeControllerNode.h b/Engine/Plugins/Experimental/MeshResizing/Source/MeshResizingNodes/Private/MeshResizing/UVResizeControllerNode.h
--- a/Engine/Plugins/Experimental/MeshResizing/Source/MeshResizingNodes/Private/MeshResizing/UVResizeControllerNode.h
+++ b/Engine/Plugins/Experimental/MeshResizing/Source/MeshResizingNodes/Private/MeshResizing/UVResizeControllerNode.h
@@ -17,7 +17,6 @@ struct FUVResizeControllerNode : public FDataflowNode
 {
 	GENERATED_USTRUCT_BODY()
 	DATAFLOW_NODE_DEFINE_INTERNAL(FUVResizeControllerNode, "UVResizeController", "MeshResizing", "UV Resize Controller")
-	DATAFLOW_NODE_RENDER_TYPE("SurfaceRender", FName("UDataflowMesh"), "Mesh", "UVChannelIndex")
 
 public:
 	FUVResizeControllerNode(const UE::Dataflow::FNodeParameters& InParam, FGuid InGuid = FGuid::NewGuid());
diff --git a/Engine/Plugins/Experimental/MeshResizing/Source/MeshResizingNodes/Private/MeshResizing/UVUnwrapNode.h b/Engine/Plugins/Experimental/MeshResizing/Source/MeshResizingNodes/Private/MeshResizing/UVUnwrapNode.h
--- a/Engine/Plugins/Experimental/MeshResizing/Source/MeshResizingNodes/Private/MeshResizing/UVUnwrapNode.h
+++ b/Engine/Plugins/Experimental/MeshResizing/Source/MeshResizingNodes/Private/MeshResizing/UVUnwrapNode.h
@@ -21,7 +21,6 @@ struct FUVUnwrapNode : public FDataflowNode
 {
 	GENERATED_USTRUCT_BODY()
 	DATAFLOW_NODE_DEFINE_INTERNAL(FUVUnwrapNode, "UVUnwrapNode", "MeshResizing", "UV Unwrap")
-	DATAFLOW_NODE_RENDER_TYPE("SurfaceRender", FName("UDataflowMesh"), "Mesh", "UVChannelIndex")
 
 public:
 
@@ -32,7 +31,7 @@ private:
 	UPROPERTY(meta = (DataflowInput, DataflowOutput, DataflowPassthrough = "Mesh"))
 	TObjectPtr<UDataflowMesh> Mesh;
 
-	UPROPERTY(EditAnywhere, Category = "UV Unwrap", meta = (DataflowInput, DataflowOutput, DataflowPassthrough = "UVChannelIndex", ClampMin = 0, ClampMax = 7, UIMin = 0, UIMax = 7))
+	UPROPERTY(EditAnywhere, Category = "UV Unwrap", meta = (DataflowInput, DataflowOutput, DataflowPassthrough = "UVChannelIndex", DataflowUVChannel, ClampMin = 0, ClampMax = 7, UIMin = 0, UIMax = 7))
 	int32 UVChannelIndex = 0;
 
 	UPROPERTY(EditAnywhere, Category = "UV Unwrap")
diff --git a/Engine/Source/Runtime/Experimental/Dataflow/Core/Public/Dataflow/DataflowConvertNodes.h b/Engine/Source/Runtime/Experimental/Dataflow/Core/Public/Dataflow/DataflowConvertNodes.h
--- a/Engine/Source/Runtime/Experimental/Dataflow/Core/Public/Dataflow/DataflowConvertNodes.h
+++ b/Engine/Source/Runtime/Experimental/Dataflow/Core/Public/Dataflow/DataflowConvertNodes.h
@@ -214,7 +214,6 @@ struct FConvertSelectionTypesDataflowNode : public FDataflowNode
 {
 	GENERATED_USTRUCT_BODY()
 	DATAFLOW_NODE_DEFINE_INTERNAL(FConvertSelectionTypesDataflowNode, "ConvertSelectionTypes", DATAFLOW_CONVERT_NODES_CATEGORY, "")
-	DATAFLOW_NODE_RENDER_TYPE("SurfaceRender", FGeometryCollection::StaticType(), "Collection")
 
 private:
 	/** GeometryCollection for the selection */
@@ -278,7 +277,6 @@ struct FConvertIndexToSelectionTypesDataflowNode : public FDataflowNode
 {
 	GENERATED_USTRUCT_BODY()
 	DATAFLOW_NODE_DEFINE_INTERNAL(FConvertIndexToSelectionTypesDataflowNode, "ConvertIndexToSelection", DATAFLOW_CONVERT_NODES_CATEGORY, "")
-		DATAFLOW_NODE_RENDER_TYPE("SurfaceRender", FGeometryCollection::StaticType(), "Collection")
 
 public:
 	FConvertIndexToSelectionTypesDataflowNode(const UE::Dataflow::FNodeParameters& InParam, FGuid InGuid = FGuid::NewGuid());
@@ -310,7 +308,6 @@ struct FConvertIndexArrayToSelectionTypesDataflowNode : public FDataflowNode
 {
 	GENERATED_USTRUCT_BODY()
 	DATAFLOW_NODE_DEFINE_INTERNAL(FConvertIndexArrayToSelectionTypesDataflowNode, "ConvertIndexArrayToSelection", DATAFLOW_CONVERT_NODES_CATEGORY, "")
-	DATAFLOW_NODE_RENDER_TYPE("SurfaceRender", FGeometryCollection::StaticType(), "Collection")
 	
 public:
 	FConvertIndexArrayToSelectionTypesDataflowNode(const UE::Dataflow::FNodeParameters& InParam, FGuid InGuid = FGuid::NewGuid());
