Dataflow : fix the visualization menu in the cloth editor to not close when clicking on the label of the menu entries

#Synced-CL 49572843

diff --git a/Engine/Plugins/Experimental/Dataflow/Source/DataflowEditor/Private/Dataflow/DataflowSimulationViewport.cpp b/Engine/Plugins/Experimental/Dataflow/Source/DataflowEditor/Private/Dataflow/DataflowSimulationViewport.cpp
--- a/Engine/Plugins/Experimental/Dataflow/Source/DataflowEditor/Private/Dataflow/DataflowSimulationViewport.cpp
+++ b/Engine/Plugins/Experimental/Dataflow/Source/DataflowEditor/Private/Dataflow/DataflowSimulationViewport.cpp
@@ -166,8 +166,9 @@ void SDataflowSimulationViewport::BuildVisualizationMenu(FToolMenuSection& MenuS
 						}
 					}
 				}),
-			false,
-			FSlateIcon(FAppStyle::Get().GetStyleSetName(), "DetailsView.ViewOptions")
+			/*bInopenSubMenuOnClick*/ false,
+			FSlateIcon(FAppStyle::Get().GetStyleSetName(), "DetailsView.ViewOptions"),
+			/*bShouldCloseWindowAfterMenuSelection*/ false
 		);
 	}
 }
