Dataflow : fix vertex attribute painting tool invalidating nodes even when canceling the tool

#Synced-CL 49572843

diff --git a/Engine/Plugins/Experimental/Dataflow/Source/DataflowEditor/Private/DataflowEditorTools/DataflowEditorVertexAttributePaintTool.cpp b/Engine/Plugins/Experimental/Dataflow/Source/DataflowEditor/Private/DataflowEditorTools/DataflowEditorVertexAttributePaintTool.cpp
--- a/Engine/Plugins/Experimental/Dataflow/Source/DataflowEditor/Private/DataflowEditorTools/DataflowEditorVertexAttributePaintTool.cpp
+++ b/Engine/Plugins/Experimental/Dataflow/Source/DataflowEditor/Private/DataflowEditorTools/DataflowEditorVertexAttributePaintTool.cpp
@@ -472,10 +472,7 @@ void FDataflowEditorVertexAttributePaintToolData::Setup(FDynamicMesh3& InMesh, F
 
 void FDataflowEditorVertexAttributePaintToolData::Shutdown()
 {
-	if (NodeToUpdate)
-	{
-		NodeToUpdate->Invalidate();
-	}
+	// nothing to do - invalidation is now done in the commit results
 }
 
 void FDataflowEditorVertexAttributePaintToolData::BeginChange(UDynamicMeshComponent* Component)
@@ -577,6 +574,7 @@ void FDataflowEditorVertexAttributePaintToolData::UpdateNode(TObjectPtr<UDataflo
 				}
 				// WeightsToApply represents the weights of the selected sub-mesh, DynamicMeshToWeight allow to remap properly to the full collection mesh
 				NodeToUpdate->SetVertexAttributeValues(*DataflowContext, WeightsToApply, DynamicMeshToWeight);
+				NodeToUpdate->Invalidate();
 			}
 		}
 	}
