Dataflow : make sure the construction view never focus on an empty box

#Synced-CL 49572843

diff --git a/Engine/Plugins/Experimental/Dataflow/Source/DataflowEditor/Private/Dataflow/DataflowEditorMode.cpp b/Engine/Plugins/Experimental/Dataflow/Source/DataflowEditor/Private/Dataflow/DataflowEditorMode.cpp
--- a/Engine/Plugins/Experimental/Dataflow/Source/DataflowEditor/Private/Dataflow/DataflowEditorMode.cpp
+++ b/Engine/Plugins/Experimental/Dataflow/Source/DataflowEditor/Private/Dataflow/DataflowEditorMode.cpp
@@ -780,8 +780,17 @@ void UDataflowEditorMode::ViewportResized(FViewport* Viewport, uint32 /*Unused*/
 
 FBox UDataflowEditorMode::SceneBoundingBox() const
 {
+	const FBox DefaultFocusBox(FVector(-50, -50, 0), FVector(+50, +50, +100));
+
 	const TSharedPtr<FDataflowConstructionScene> ConstructionScenePtr = ConstructionScene.Pin();
-	return ConstructionScenePtr ? ConstructionScenePtr->GetBoundingBox() : FBox(FVector(-100.), FVector(100.));
+	FBox BoxToReturn = ConstructionScenePtr 
+		? ConstructionScenePtr->GetBoundingBox() 
+		: DefaultFocusBox;
+	if (!BoxToReturn.IsValid)
+	{
+		BoxToReturn = DefaultFocusBox;
+	}
+	return BoxToReturn;
 }
 
 FBox UDataflowEditorMode::SelectionBoundingBox() const
