Dataflow weightmap paint tool : add the ability to display remapped values on top of the painted mesh

#Synced-CL 49572843

diff --git a/Engine/Plugins/Experimental/Dataflow/Source/DataflowEditor/Private/DataflowEditorTools/DataflowEditorVertexAttributePaintTool.cpp b/Engine/Plugins/Experimental/Dataflow/Source/DataflowEditor/Private/DataflowEditorTools/DataflowEditorVertexAttributePaintTool.cpp
--- a/Engine/Plugins/Experimental/Dataflow/Source/DataflowEditor/Private/DataflowEditorTools/DataflowEditorVertexAttributePaintTool.cpp
+++ b/Engine/Plugins/Experimental/Dataflow/Source/DataflowEditor/Private/DataflowEditorTools/DataflowEditorVertexAttributePaintTool.cpp
@@ -2120,6 +2120,76 @@ void UDataflowEditorVertexAttributePaintTool::DrawHUD(FCanvas* Canvas, IToolsCon
 			Canvas->DrawItem(TextItem);
 		}
 	}
+
+	if (ToolProperties && ToolProperties->DisplayProperties.bShowValues)
+	{
+		if (Canvas && RenderAPI)
+		{
+			if (const FDynamicMesh3* Mesh = GetSculptMesh())
+			{
+				FViewCameraState StateOut;
+				GetToolManager()->GetContextQueriesAPI()->GetCurrentViewState(StateOut);
+				const FTransform3d LocalToWorld = UE::ToolTarget::GetLocalToWorldTransform(Target);
+				const FVector3d LocalEyePosition(LocalToWorld.InverseTransformPosition(StateOut.Position));
+
+				auto IsVertexVisible =
+					[&LocalEyePosition, &Octree = Octree](const FDynamicMesh3& Mesh, const int32 VtxIndex)
+					{
+						UE::Geometry::FVertexInfo VertexInfo;
+						Mesh.GetVertex(VtxIndex, VertexInfo, true, false, false);
+						const FVector3d Centroid = VertexInfo.Position;
+						const FVector3d FaceNormal = FVector3d(VertexInfo.Normal);
+						const FVector3d CameraToVertexNormal = UE::Geometry::Normalized(Centroid - LocalEyePosition);
+						if (!FaceNormal.IsZero() && FaceNormal.Dot(CameraToVertexNormal) > 0)
+						{
+							return false;
+						}
+						const FVector RayDirection = CameraToVertexNormal;
+						const FVector RayOrigin = LocalEyePosition;
+						const double MaxDistance = (LocalEyePosition - Centroid).Size() - UE_SMALL_NUMBER;
+						const int32 HitTID = Octree.FindNearestHitObject(FRay3d(RayOrigin, RayDirection), MaxDistance);
+						if (HitTID != IndexConstants::InvalidID && Mesh.IsTriangle(HitTID))
+						{
+							return false;
+						}
+						return true;
+					};
+
+				FNumberFormattingOptions NumberFormattingOptions;
+				NumberFormattingOptions.AlwaysSign = false;
+				NumberFormattingOptions.UseGrouping = false;
+				NumberFormattingOptions.RoundingMode = ERoundingMode::HalfFromZero;
+				NumberFormattingOptions.MinimumIntegralDigits = 1;
+				NumberFormattingOptions.MaximumIntegralDigits = 6;
+				NumberFormattingOptions.MinimumFractionalDigits = 2;
+				NumberFormattingOptions.MaximumFractionalDigits = 2;
+
+				const float MinValue = ToolProperties->DisplayProperties.MinValue;
+				const float MaxValue = ToolProperties->DisplayProperties.MaxValue;
+				const float ValueRange = MaxValue - MinValue;
+
+				for (int32 VertexID : Mesh->VertexIndicesItr())
+				{
+					if (IsVertexVisible(*Mesh, VertexID))
+					{
+						const float NormalizedValue = VertexData.GetValue(VertexID);
+						const float RemappedValue = MinValue + NormalizedValue * ValueRange;
+						const FText ValueText = FText::AsNumber(RemappedValue, &NumberFormattingOptions);
+
+						const FVector3d VtxPosition = Mesh->GetVertex(VertexID);
+						const FVector4 ScreenBrushPos = RenderAPI->GetSceneView()->WorldToScreen(VtxPosition);
+						FVector2D ScreenBrushPos2D = FVector2D(ScreenBrushPos.X, ScreenBrushPos.Y);
+						RenderAPI->GetSceneView()->ScreenToPixel(ScreenBrushPos, ScreenBrushPos2D);
+
+						FCanvasTextItem TextItem(ScreenBrushPos2D, ValueText, GEngine->GetLargeFont(), FLinearColor::White);
+						TextItem.EnableShadow(FLinearColor::Black);
+						TextItem.bCentreX = true;
+						Canvas->DrawItem(TextItem);
+					}
+				}
+			}
+		}
+	}
 }
 
 void UDataflowEditorVertexAttributePaintTool::Render(IToolsContextRenderAPI* RenderAPI)
diff --git a/Engine/Plugins/Experimental/Dataflow/Source/DataflowEditor/Private/DataflowEditorTools/DataflowEditorVertexAttributePaintToolDetailCustomization.cpp b/Engine/Plugins/Experimental/Dataflow/Source/DataflowEditor/Private/DataflowEditorTools/DataflowEditorVertexAttributePaintToolDetailCustomization.cpp
--- a/Engine/Plugins/Experimental/Dataflow/Source/DataflowEditor/Private/DataflowEditorTools/DataflowEditorVertexAttributePaintToolDetailCustomization.cpp
+++ b/Engine/Plugins/Experimental/Dataflow/Source/DataflowEditor/Private/DataflowEditorTools/DataflowEditorVertexAttributePaintToolDetailCustomization.cpp
@@ -444,6 +444,92 @@ namespace UE::Dataflow::Editor
 		}
 	}
 
+	void FVertexAttributePaintToolDetailCustomization::AddValuesDisplayRows(IDetailCategoryBuilder& MeshDisplayCategory) const
+	{
+		MeshDisplayCategory.AddCustomRow(LOCTEXT("ValuesDisplay_ShowValues_Row", "Values Display - Show Values"), false)
+			.NameContent()
+			[
+				SNew(STextBlock)
+					.Text(LOCTEXT("ValuesDisplay_ShowValues_Label", "Show Values"))
+					.Font(FAppStyle::GetFontStyle(TEXT("PropertyWindow.NormalFont")))
+					.ToolTipText(LOCTEXT("ValuesDisplay_ShowValues_Tooltip", "When enable, display the values on each vertex using Min/Max remapping "))
+			]
+			.ValueContent()
+			[
+				SNew(SCheckBox)
+					.IsChecked_Lambda([ToolProperties = ToolProperties]()
+						{
+							const bool bIsChecked = (ToolProperties.IsValid() && ToolProperties->DisplayProperties.bShowValues);
+							return bIsChecked ? ECheckBoxState::Checked : ECheckBoxState::Unchecked;
+						})
+					.OnCheckStateChanged_Lambda([ToolProperties = ToolProperties](ECheckBoxState InCheckBoxState)
+						{
+							if (ToolProperties.IsValid())
+							{
+								ToolProperties->DisplayProperties.bShowValues = (InCheckBoxState == ECheckBoxState::Checked);
+							}
+						})
+			];
+
+		MeshDisplayCategory.AddCustomRow(LOCTEXT("ValuesDisplay_MinValue_Row", "Values Display - Min Value"), false)
+			.NameContent()
+			[
+				SNew(STextBlock)
+					.Text(LOCTEXT("ValuesDisplay_MinValue_Label", "Min Value"))
+					.Font(FAppStyle::GetFontStyle(TEXT("PropertyWindow.NormalFont")))
+					.ToolTipText(LOCTEXT("ValuesDisplay_MinValue_Tooltip", "Minimum value mapping to the zero value of the painted map"))
+			]
+			.ValueContent()
+			[
+				SNew(SNumericEntryBox<float>)
+					.IsEnabled_Lambda([ToolProperties = ToolProperties]()
+						{
+							return (ToolProperties.IsValid() && ToolProperties->DisplayProperties.bShowValues);
+						})
+					.Value_Lambda([ToolProperties = ToolProperties]()
+						{
+							const float MinValue = ToolProperties.IsValid() ? ToolProperties->DisplayProperties.MinValue : 0;
+							return MinValue;
+						})
+					.OnValueChanged_Lambda([ToolProperties = ToolProperties](float Value)
+						{
+							if (ToolProperties.IsValid())
+							{
+								ToolProperties->DisplayProperties.MinValue = Value;
+							}
+						})
+			];
+
+		MeshDisplayCategory.AddCustomRow(LOCTEXT("ValuesDisplay_MaxValue_Row", "Values Display - Max Value"), false)
+			.NameContent()
+			[
+				SNew(STextBlock)
+					.Text(LOCTEXT("ValuesDisplay_MaxValue_Label", "Max Value"))
+					.Font(FAppStyle::GetFontStyle(TEXT("PropertyWindow.NormalFont")))
+					.ToolTipText(LOCTEXT("ValuesDisplay_MaxValue_Tooltip", "Maximum value mapping to the one value of the painted map"))
+			]
+			.ValueContent()
+			[
+				SNew(SNumericEntryBox<float>)
+					.IsEnabled_Lambda([ToolProperties = ToolProperties]()
+						{
+							return (ToolProperties.IsValid() && ToolProperties->DisplayProperties.bShowValues);
+						})
+					.Value_Lambda([ToolProperties = ToolProperties]()
+						{
+							const float MaxValue = ToolProperties.IsValid() ? ToolProperties->DisplayProperties.MaxValue : 0;
+							return MaxValue;
+						})
+					.OnValueChanged_Lambda([ToolProperties = ToolProperties](float Value)
+						{
+							if (ToolProperties.IsValid())
+							{
+								ToolProperties->DisplayProperties.MaxValue = Value;
+							}
+						})
+			];
+	}
+
 	///////////////////////////////////////////////////////////////////////////////////////////////////////////////
 
 	void FVertexAttributePaintToolDetailCustomization::CustomizeDetails(IDetailLayoutBuilder& DetailBuilder)
@@ -476,7 +562,7 @@ namespace UE::Dataflow::Editor
 			}
 		}
 
-		// MeQuery category
+		// Value Query category
 		IDetailCategoryBuilder& QueryCategory = DetailBuilder.EditCategory("Query", FText::GetEmpty(), ECategoryPriority::Important);
 		QueryCategory.InitiallyCollapsed(false);
 		{
@@ -500,6 +586,8 @@ namespace UE::Dataflow::Editor
 			AddColorModeRow(MeshDisplayCategory);
 
 			AddColorRampRow(MeshDisplayCategory);
+
+			AddValuesDisplayRows(MeshDisplayCategory);
 		}
 
 		// Hide all customized properties 
diff --git a/Engine/Plugins/Experimental/Dataflow/Source/DataflowEditor/Private/DataflowEditorTools/DataflowEditorVertexAttributePaintToolDetailCustomization.h b/Engine/Plugins/Experimental/Dataflow/Source/DataflowEditor/Private/DataflowEditorTools/DataflowEditorVertexAttributePaintToolDetailCustomization.h
--- a/Engine/Plugins/Experimental/Dataflow/Source/DataflowEditor/Private/DataflowEditorTools/DataflowEditorVertexAttributePaintToolDetailCustomization.h
+++ b/Engine/Plugins/Experimental/Dataflow/Source/DataflowEditor/Private/DataflowEditorTools/DataflowEditorVertexAttributePaintToolDetailCustomization.h
@@ -36,7 +36,8 @@ namespace UE::Dataflow::Editor
 		void AddColorModeRow(IDetailCategoryBuilder& EditModeCategory) const;
 		void AddColorRampRow(IDetailCategoryBuilder& EditModeCategory) const;
 		TSharedRef<SWidget> BuildColorRampMenu() const;
-		
+		void AddValuesDisplayRows(IDetailCategoryBuilder& EditModeCategory) const;
+
 		void AddBrushUI(IDetailLayoutBuilder& DetailBuilder) const;
 		void AddBrushModeRow(IDetailCategoryBuilder& BrushCategory) const;
 		void AddBrushFalloffModeRow(IDetailCategoryBuilder& BrushCategory) const;
diff --git a/Engine/Plugins/Experimental/Dataflow/Source/DataflowEditor/Public/DataflowEditorTools/DataflowEditorVertexAttributePaintTool.h b/Engine/Plugins/Experimental/Dataflow/Source/DataflowEditor/Public/DataflowEditorTools/DataflowEditorVertexAttributePaintTool.h
--- a/Engine/Plugins/Experimental/Dataflow/Source/DataflowEditor/Public/DataflowEditorTools/DataflowEditorVertexAttributePaintTool.h
+++ b/Engine/Plugins/Experimental/Dataflow/Source/DataflowEditor/Public/DataflowEditorTools/DataflowEditorVertexAttributePaintTool.h
@@ -199,6 +199,15 @@ struct FDataflowEditorVertexAttributePaintToolDisplayProperties
 	// Used by UI customization 
 	FLinearColorRamp GreyScaleColorRamp;
 	FLinearColorRamp WhiteColorRamp;
+
+	UPROPERTY()
+	bool bShowValues = false;
+
+	UPROPERTY()
+	float MinValue = 0;
+
+	UPROPERTY()
+	float MaxValue = 100;
 };
 
 USTRUCT()
