Dataflow paint tool : add option for mirror to use Object space (on by default)

#Synced-CL 49572843

diff --git a/Engine/Plugins/Experimental/Dataflow/Source/DataflowEditor/Private/DataflowEditorTools/DataflowEditorVertexAttributePaintTool.cpp b/Engine/Plugins/Experimental/Dataflow/Source/DataflowEditor/Private/DataflowEditorTools/DataflowEditorVertexAttributePaintTool.cpp
--- a/Engine/Plugins/Experimental/Dataflow/Source/DataflowEditor/Private/DataflowEditorTools/DataflowEditorVertexAttributePaintTool.cpp
+++ b/Engine/Plugins/Experimental/Dataflow/Source/DataflowEditor/Private/DataflowEditorTools/DataflowEditorVertexAttributePaintTool.cpp
@@ -2654,7 +2654,7 @@ void UDataflowEditorVertexAttributePaintTool::MirrorValues()
 
 		if (FDynamicMesh3* Mesh = GetSculptMesh())
 		{
-			MirrorData.EnsureMirrorDataIsUpdated(*Mesh, ToolProperties->MirrorProperties.MirrorAxis, ToolProperties->MirrorProperties.MirrorDirection);
+			MirrorData.EnsureMirrorDataIsUpdated(*Mesh, ToolProperties->MirrorProperties.MirrorAxis, ToolProperties->MirrorProperties.MirrorDirection, ToolProperties->MirrorProperties.bObjectSpace);
 
 			TArray<int32> VerticesToUpdate;
 			MirrorData.FindMirroredIndices(*Mesh, VertexROI, VerticesToUpdate);
@@ -2700,20 +2700,21 @@ bool UDataflowEditorVertexAttributePaintTool::FMirrorData::IsPointOnTargetMirror
 	{
 		return false;
 	}
-	if (Direction == EDataflowEditorToolMirrorDirection::PositiveToNegative && InPoint[Axis - 1] >= 0.f)
+	const int32 AxisIndex = (Axis - 1);
+	if (Direction == EDataflowEditorToolMirrorDirection::PositiveToNegative && InPoint[AxisIndex] >= PlaneOffsets[AxisIndex])
 	{
 		return false; // target is negative side, but point is on positive side
 	}
-	if (Direction == EDataflowEditorToolMirrorDirection::NegativeToPositive && InPoint[Axis - 1] <= 0.f)
+	if (Direction == EDataflowEditorToolMirrorDirection::NegativeToPositive && InPoint[AxisIndex] <= PlaneOffsets[AxisIndex])
 	{
 		return false; // target is positive side, but vertex is on negative side
 	}
 	return true;
 }
 
-void UDataflowEditorVertexAttributePaintTool::FMirrorData::EnsureMirrorDataIsUpdated(const FDynamicMesh3& Mesh, EAxis::Type InMirrorAxis, EDataflowEditorToolMirrorDirection InMirrorDirection)
+void UDataflowEditorVertexAttributePaintTool::FMirrorData::EnsureMirrorDataIsUpdated(const FDynamicMesh3& Mesh, EAxis::Type InMirrorAxis, EDataflowEditorToolMirrorDirection InMirrorDirection, bool bInObjectSpace)
 {
-	if (VertexMap.Num() > 0 && InMirrorAxis == Axis && InMirrorDirection == Direction)
+	if (VertexMap.Num() > 0 && InMirrorAxis == Axis && InMirrorDirection == Direction && bObjectSpace == bInObjectSpace)
 	{
 		// already initialized, just re-use cached data
 		return;
@@ -2722,6 +2723,7 @@ void UDataflowEditorVertexAttributePaintTool::FMirrorData::EnsureMirrorDataIsUpd
 	// need to re-initialize
 	Axis = InMirrorAxis;
 	Direction = InMirrorDirection;
+	bObjectSpace = bInObjectSpace;
 	VertexMap.Reset();
 
 	const int32 NumVertices = Mesh.VertexCount();
@@ -2731,6 +2733,13 @@ void UDataflowEditorVertexAttributePaintTool::FMirrorData::EnsureMirrorDataIsUpd
 	UE::Geometry::TPointHashGrid3f<int32> VertHash(HashGridCellSize, INDEX_NONE);
 	VertHash.Reserve(NumVertices);
 
+	PlaneOffsets = FVector{ 0, 0, 0 };
+	if (bObjectSpace)
+	{
+		const UE::Geometry::FAxisAlignedBox3d MeshBounds = Mesh.GetBounds(true);
+		PlaneOffsets = MeshBounds.Center();
+	}
+
 	for (int32 VertexID : Mesh.VertexIndicesItr())
 	{
 		const FVector VertexPos = Mesh.GetVertex(VertexID);
@@ -2744,7 +2753,8 @@ void UDataflowEditorVertexAttributePaintTool::FMirrorData::EnsureMirrorDataIsUpd
 
 		// flip position across the mirror axis
 		FVector3f MirroredPosition = FVector3f(SourcePosition);
-		MirroredPosition[Axis - 1] *= -1.f;
+		const int32 AxisIndex = (Axis - 1);
+		MirroredPosition[AxisIndex] = ((MirroredPosition[AxisIndex] - PlaneOffsets[AxisIndex]) *  (-1.0f)) + PlaneOffsets[AxisIndex];
 
 		// Query spatial hash near mirrored position, gradually increasing search radius until at least 1 point is found
 		TPair<int32, double> ClosestMirroredPoint = { INDEX_NONE, TNumericLimits<double>::Max() };
diff --git a/Engine/Plugins/Experimental/Dataflow/Source/DataflowEditor/Private/DataflowEditorTools/DataflowEditorVertexAttributePaintToolDetailCustomization.cpp b/Engine/Plugins/Experimental/Dataflow/Source/DataflowEditor/Private/DataflowEditorTools/DataflowEditorVertexAttributePaintToolDetailCustomization.cpp
--- a/Engine/Plugins/Experimental/Dataflow/Source/DataflowEditor/Private/DataflowEditorTools/DataflowEditorVertexAttributePaintToolDetailCustomization.cpp
+++ b/Engine/Plugins/Experimental/Dataflow/Source/DataflowEditor/Private/DataflowEditorTools/DataflowEditorVertexAttributePaintToolDetailCustomization.cpp
@@ -1681,6 +1681,24 @@ namespace UE::Dataflow::Editor
 					return FReply::Handled();
 				});
 
+		TSharedRef<SCheckBox> ObjectSpaceCheckBox =
+			SNew(SCheckBox)
+			.IsChecked_Lambda([ToolProperties = ToolProperties]()
+				{
+					const bool bIsChecked = (ToolProperties.IsValid() && ToolProperties->MirrorProperties.bObjectSpace);
+					return bIsChecked ? ECheckBoxState::Checked : ECheckBoxState::Unchecked;
+				})
+			.OnCheckStateChanged_Lambda([ToolProperties = ToolProperties](ECheckBoxState InCheckBoxState)
+				{
+					if (ToolProperties.IsValid())
+					{
+						ToolProperties->MirrorProperties.bObjectSpace = (InCheckBoxState == ECheckBoxState::Checked);
+					}
+				});
+
+		constexpr float HorizontalPadding = 2.0f;
+		constexpr float VerticalPadding = 4.0f;
+
 		// MIRROR WEIGHTS category
 		EditValuesCategory.AddCustomRow(LOCTEXT("MirrorWeightsRow", "Mirror"), false)
 			.WholeRowContent()
@@ -1691,7 +1709,7 @@ namespace UE::Dataflow::Editor
 							return Tool.IsValid() ? Tool->HasSelection() : false; 
 						})
 					+ SVerticalBox::Slot()
-					.Padding(WeightEditHorizontalPadding, WeightEditVerticalPadding)
+					.Padding(HorizontalPadding, VerticalPadding)
 					[
 						SNew(SHorizontalBox)
 						+ SHorizontalBox::Slot()
@@ -1719,6 +1737,25 @@ namespace UE::Dataflow::Editor
 						]
 					]
 					+ SVerticalBox::Slot()
+					.Padding(HorizontalPadding, VerticalPadding)
+					[
+						SNew(SHorizontalBox)
+							+ SHorizontalBox::Slot()
+							.VAlign(VAlign_Center)
+							.FillWidth(WeightEditingLabelsPercent)
+							[
+								SNew(STextBlock)
+									.Text(LOCTEXT("MirrorObjectSpaceLabel", "Object Space"))
+									.Font(FAppStyle::GetFontStyle(TEXT("PropertyWindow.NormalFont")))
+									.ToolTipText(LOCTEXT("MirrorObjectSpaceTooltip", "When checked, the planes will be cenetred on the object bounds."))
+							]
+							+ SHorizontalBox::Slot()
+							.FillWidth(1.f)
+							[
+								ObjectSpaceCheckBox
+							]
+					]
+					+ SVerticalBox::Slot()
 					.Padding(0.f, WeightEditVerticalPadding)
 					[
 						SNew(SBox)
diff --git a/Engine/Plugins/Experimental/Dataflow/Source/DataflowEditor/Public/DataflowEditorTools/DataflowEditorVertexAttributePaintTool.h b/Engine/Plugins/Experimental/Dataflow/Source/DataflowEditor/Public/DataflowEditorTools/DataflowEditorVertexAttributePaintTool.h
--- a/Engine/Plugins/Experimental/Dataflow/Source/DataflowEditor/Public/DataflowEditorTools/DataflowEditorVertexAttributePaintTool.h
+++ b/Engine/Plugins/Experimental/Dataflow/Source/DataflowEditor/Public/DataflowEditorTools/DataflowEditorVertexAttributePaintTool.h
@@ -215,6 +215,10 @@ struct FDataflowEditorVertexAttributePaintToolMirrorProperties
 {
 	GENERATED_BODY()
 
+	/* if true, the mirror planes will be centered on the bounds of the object */
+	UPROPERTY()
+	bool bObjectSpace = true;
+
 	UPROPERTY()
 	TEnumAsByte<EAxis::Type> MirrorAxis = EAxis::X;
 
@@ -566,7 +570,7 @@ private:
 	{
 	public:
 		// lazily updates the mirror data tables for the current skeleton/mesh/mirror plane
-		void EnsureMirrorDataIsUpdated(const FDynamicMesh3& Mesh, EAxis::Type InMirrorAxis, EDataflowEditorToolMirrorDirection InMirrorDirection);
+		void EnsureMirrorDataIsUpdated(const FDynamicMesh3& Mesh, EAxis::Type InMirrorAxis, EDataflowEditorToolMirrorDirection InMirrorDirection, bool bInObjectSpace);
 
 		void FindMirroredIndices(const FDynamicMesh3& Mesh, const TArray<int32>& SelectedVertices, TArray<int32>& OutVerticesToUpdate);
 
@@ -576,6 +580,8 @@ private:
 
 		EAxis::Type Axis = EAxis::X;
 		EDataflowEditorToolMirrorDirection Direction = EDataflowEditorToolMirrorDirection::PositiveToNegative;
+		bool bObjectSpace = true;
+		FVector PlaneOffsets = FVector{ 0, 0, 0 };
 		TMap<int32, int32> VertexMap; // <Target, Source>
 	};
 	FMirrorData MirrorData;
