Dataflow : texture to attribute node : listen to texture change for real time update of the graph

#Synced-CL 49632091

diff --git a/Engine/Plugins/Experimental/Dataflow/Source/DataflowNodes/Private/Dataflow/DataflowTextureToAttributeNode.cpp b/Engine/Plugins/Experimental/Dataflow/Source/DataflowNodes/Private/Dataflow/DataflowTextureToAttributeNode.cpp
--- a/Engine/Plugins/Experimental/Dataflow/Source/DataflowNodes/Private/Dataflow/DataflowTextureToAttributeNode.cpp
+++ b/Engine/Plugins/Experimental/Dataflow/Source/DataflowNodes/Private/Dataflow/DataflowTextureToAttributeNode.cpp
@@ -9,6 +9,10 @@
 #include "GeometryCollection/Facades/CollectionMeshFacade.h"
 #include "ImageCore.h"
 
+#if WITH_EDITOR
+#include "UObject/UObjectGlobals.h"
+#endif
+
 #include UE_INLINE_GENERATED_CPP_BY_NAME(DataflowTextureToAttributeNode)
 
 #define LOCTEXT_NAMESPACE "DataflowTextureToAttributeNode"
@@ -57,6 +61,18 @@ FDataflowTextureToAttributeNode::FDataflowTextureToAttributeNode(const UE::Dataf
 	RegisterOutputConnection(&Collection, &Collection);
 }
 
+FDataflowTextureToAttributeNode::~FDataflowTextureToAttributeNode()
+{
+	StopListeningToTextureChanges();
+}
+
+void FDataflowTextureToAttributeNode::PostSerialize(const FArchive& Ar)
+{
+	FDataflowNode::PostSerialize(Ar);
+	StopListeningToTextureChanges();
+	StartListeningToTextureChanges();
+}
+
 void FDataflowTextureToAttributeNode::Evaluate(UE::Dataflow::FContext& Context, const FDataflowOutput* Out) const
 {
 	if (Out->IsA(&Collection))
@@ -180,5 +196,50 @@ FDataflowNode::FAttributeKey FDataflowTextureToAttributeNode::GetVertexAttribute
 	};
 }
 
+void FDataflowTextureToAttributeNode::OnPropertyChanged(UE::Dataflow::FContext& Context, const FPropertyChangedEvent& InPropertyChangedEvent)
+{
+	if (InPropertyChangedEvent.GetPropertyName() == GET_MEMBER_NAME_CHECKED(FDataflowTextureToAttributeNode, Texture))
+	{
+		StopListeningToTextureChanges();
+		StartListeningToTextureChanges();
+	}
+}
+
+void FDataflowTextureToAttributeNode::StartListeningToTextureChanges()
+{
+#if WITH_EDITOR
+	check(!TextureChangeDelegateHandle.IsValid());
+	TextureChangeDelegateHandle = FCoreUObjectDelegates::OnObjectPropertyChanged.AddSP(this, &FDataflowTextureToAttributeNode::OnTextureChanged);
+#endif
+}
+
+void FDataflowTextureToAttributeNode::StopListeningToTextureChanges()
+{
+#if WITH_EDITOR
+	if (TextureChangeDelegateHandle.IsValid())
+	{
+		FCoreUObjectDelegates::OnObjectPropertyChanged.RemoveAll(this);
+		TextureChangeDelegateHandle.Reset();
+	}
+#endif
+}
+
+void FDataflowTextureToAttributeNode::OnTextureChanged(UObject* Object, FPropertyChangedEvent& Event)
+{
+	if (IsValid(Texture) && IsValid(Object) && Texture == Object)
+	{
+		TWeakPtr<FDataflowNode> LocalWeakThis = this->AsWeak();
+		ExecuteOnGameThread(TEXT("FDataflowTextureToAttributeNode::OnTextureChanged"),
+			[LocalWeakThis]()
+			{
+				if (TSharedPtr<FDataflowNode> SharedThis = LocalWeakThis.Pin())
+				{
+					SharedThis->Invalidate();
+				}
+			}
+		);
+	}
+}
+
 #undef LOCTEXT_NAMESPACE 
 
diff --git a/Engine/Plugins/Experimental/Dataflow/Source/DataflowNodes/Private/Dataflow/DataflowTextureToAttributeNode.h b/Engine/Plugins/Experimental/Dataflow/Source/DataflowNodes/Private/Dataflow/DataflowTextureToAttributeNode.h
--- a/Engine/Plugins/Experimental/Dataflow/Source/DataflowNodes/Private/Dataflow/DataflowTextureToAttributeNode.h
+++ b/Engine/Plugins/Experimental/Dataflow/Source/DataflowNodes/Private/Dataflow/DataflowTextureToAttributeNode.h
@@ -5,6 +5,7 @@
 #include "Dataflow/DataflowNode.h"
 #include "GeometryCollection/ManagedArrayCollection.h"
 #include "Dataflow/DataflowCollectionAddScalarVertexPropertyNode.h"
+#include "Delegates/IDelegateInstance.h"
 
 #include "DataflowTextureToAttributeNode.generated.h"
 
@@ -23,6 +24,7 @@ struct FDataflowTextureToAttributeNode : public FDataflowNode
 
 public:
 	FDataflowTextureToAttributeNode(const UE::Dataflow::FNodeParameters& InParam, FGuid InGuid = FGuid::NewGuid());
+	virtual ~FDataflowTextureToAttributeNode();
 
 private:
 	/** Collection to transfer the attribute data to */
@@ -47,6 +49,13 @@ private:
 
 	virtual void Evaluate(UE::Dataflow::FContext& Context, const FDataflowOutput* Out) const override;
 	virtual FAttributeKey GetVertexAttributeToVisualize(UE::Dataflow::FContext& Context, const FName OutputName, const FName RenderGroup) const override;
+
+	virtual void PostSerialize(const FArchive& Ar) override;
+	virtual void OnPropertyChanged(UE::Dataflow::FContext& Context, const FPropertyChangedEvent& InPropertyChangedEvent) override;
+	void StartListeningToTextureChanges();
+	void StopListeningToTextureChanges();
+	void OnTextureChanged(UObject* Object, FPropertyChangedEvent& Event);
+	FDelegateHandle TextureChangeDelegateHandle;
 };
 
 namespace UE::Dataflow
