Dataflow paint tool : add support for value picker and update shortcut to match the new skeletal weight painting tool

#Synced-CL 49632091

diff --git a/Engine/Plugins/Experimental/Dataflow/Source/DataflowEditor/Private/Dataflow/DataflowEditorModule.cpp b/Engine/Plugins/Experimental/Dataflow/Source/DataflowEditor/Private/Dataflow/DataflowEditorModule.cpp
--- a/Engine/Plugins/Experimental/Dataflow/Source/DataflowEditor/Private/Dataflow/DataflowEditorModule.cpp
+++ b/Engine/Plugins/Experimental/Dataflow/Source/DataflowEditor/Private/Dataflow/DataflowEditorModule.cpp
@@ -116,6 +116,23 @@ namespace UE::Dataflow::Private
 		}
 	};
 
+	class FDataflowEditorVertexAttributePaintToolActionCommands : public TInteractiveToolCommands<FDataflowEditorVertexAttributePaintToolActionCommands>
+	{
+	public:
+		FDataflowEditorVertexAttributePaintToolActionCommands() :
+			TInteractiveToolCommands<FDataflowEditorVertexAttributePaintToolActionCommands>(
+				TEXT("DataflowEditorVertexAttributePaintTooContext"),
+				LOCTEXT("DataflowEditorVertexAttributePaintTooContext", "Dataflow Vertex Attribute Paint Tool Context"),
+				NAME_None,
+				FAppStyle::GetAppStyleSetName())
+		{}
+
+		virtual void GetToolDefaultObjectList(TArray<UInteractiveTool*>& ToolCDOs) override
+		{
+			ToolCDOs.Add(GetMutableDefault<UDataflowEditorVertexAttributePaintTool>());
+		}
+	};
+
 	class FDataflowToolActionCommandBindings : public UE::Dataflow::FDataflowToolRegistry::IDataflowToolActionCommands
 	{
 	public:
@@ -124,6 +141,7 @@ namespace UE::Dataflow::Private
 			FDataflowEditorWeightMapPaintToolActionCommands::Register();
 			FDataflowEditorSkinWeightPaintToolActionCommands::Register();
 			FDataflowEditorEditSkeletonBonesToolActionCommands::Register();
+			FDataflowEditorVertexAttributePaintToolActionCommands::Register();
 		}
 
 		virtual void UnbindActiveCommands(const TSharedPtr<FUICommandList>& UICommandList) const override
@@ -136,6 +154,9 @@ namespace UE::Dataflow::Private
 
 			checkf(FDataflowEditorEditSkeletonBonesToolActionCommands::IsRegistered(), TEXT("Expected SkeletonEditTool actions to have been registered"));
 			FDataflowEditorEditSkeletonBonesToolActionCommands::Get().UnbindActiveCommands(UICommandList);
+
+			checkf(FDataflowEditorVertexAttributePaintToolActionCommands::IsRegistered(), TEXT("Expected VertexAttributePaintTool actions to have been registered"));
+			FDataflowEditorVertexAttributePaintToolActionCommands::Get().UnbindActiveCommands(UICommandList);
 		}
 
 		virtual void BindCommandsForCurrentTool(const TSharedPtr<FUICommandList>& UICommandList, UInteractiveTool* Tool) const override
@@ -155,6 +176,11 @@ namespace UE::Dataflow::Private
 				checkf(FDataflowEditorEditSkeletonBonesToolActionCommands::IsRegistered(), TEXT("Expected SkeletonEditTool actions to have been registered"));
 				FDataflowEditorEditSkeletonBonesToolActionCommands::Get().BindCommandsForCurrentTool(UICommandList, Tool);
 			}
+			if (ExactCast<UDataflowEditorVertexAttributePaintTool>(Tool))
+			{
+				checkf(FDataflowEditorVertexAttributePaintToolActionCommands::IsRegistered(), TEXT("Expected VertexAttributePaintTool actions to have been registered"));
+				FDataflowEditorVertexAttributePaintToolActionCommands::Get().BindCommandsForCurrentTool(UICommandList, Tool);
+			}
 		}
 	};
 
diff --git a/Engine/Plugins/Experimental/Dataflow/Source/DataflowEditor/Private/DataflowEditorTools/DataflowEditorVertexAttributePaintTool.cpp b/Engine/Plugins/Experimental/Dataflow/Source/DataflowEditor/Private/DataflowEditorTools/DataflowEditorVertexAttributePaintTool.cpp
--- a/Engine/Plugins/Experimental/Dataflow/Source/DataflowEditor/Private/DataflowEditorTools/DataflowEditorVertexAttributePaintTool.cpp
+++ b/Engine/Plugins/Experimental/Dataflow/Source/DataflowEditor/Private/DataflowEditorTools/DataflowEditorVertexAttributePaintTool.cpp
@@ -861,17 +861,161 @@ void UDataflowEditorVertexAttributePaintTool::CommitResult(UBaseDynamicMeshCompo
 	}
 }
 
+void UDataflowEditorVertexAttributePaintTool::EnableValuePicker(bool bEnable)
+{
+	bPendingPickWeight = bEnable;
+}
+
+bool UDataflowEditorVertexAttributePaintTool::IsValuePickerEnabled()
+{
+	return bPendingPickWeight;
+}
 
 void UDataflowEditorVertexAttributePaintTool::RegisterActions(FInteractiveToolActionSet& ActionSet)
 {
-	UMeshSculptToolBase::RegisterActions(ActionSet);
+	// IMPORTANT : we do not register the base class actions as they do not map enough to the wanted actions for this tool
+
+	int32 ActionID = (int32)EStandardToolActions::BaseClientDefinedActionID + 500;
+
+	//-------------------------------------------------------------------------------------------------------
+	// Color modes
+
+	ActionSet.RegisterAction(this, ActionID,
+		TEXT("MeshDisplayGreyscale"),
+		LOCTEXT("MeshDisplayGreyscale", "Greyscale Display"),
+		LOCTEXT("MeshDisplayGreyscaleTooltip", "Use greyscale color ramp to display values"),
+		EModifierKey::Control, EKeys::NumPadSeven,
+		[this]() { SetColorMode(EDataflowEditorToolColorMode::Greyscale); });
+	++ActionID;
+
+	ActionSet.RegisterAction(this, ActionID,
+		TEXT("MeshDisplayColorRamp"),
+		LOCTEXT("MeshDisplayColorRamp", "Color Ramp Display"),
+		LOCTEXT("MeshDisplayColorRampTooltip", "Use custom color ramp to display values"),
+		EModifierKey::Control, EKeys::NumPadEight,
+		[this]() { SetColorMode(EDataflowEditorToolColorMode::Ramp); });
+	++ActionID;
+
+	ActionSet.RegisterAction(this, ActionID,
+		TEXT("MeshDisplayFullMaterial"),
+		LOCTEXT("MeshDisplayFullMaterial", "Full Material Display"),
+		LOCTEXT("MeshDisplayFullMaterialTooltip", "Use normal material color to display values"),
+		EModifierKey::Control, EKeys::NumPadSix,
+		[this]() { SetColorMode(EDataflowEditorToolColorMode::FullMaterial); });
+	++ActionID;
+
+	//-------------------------------------------------------------------------------------------------------
+	// Component selection mode
+
+	ActionSet.RegisterAction(this, ActionID,
+		TEXT("ComponentSelectionModeVertices"),
+		LOCTEXT("ComponentSelectionModeVertices", "Vertex Selection Mode"),
+		LOCTEXT("ComponentSelectionModeVerticesTooltip", "Enable vertex selection mode"),
+		EModifierKey::None, EKeys::One,
+		[this]() { SetComponentSelectionMode(EComponentSelectionMode::Vertices); });
+	++ActionID;
+
+	ActionSet.RegisterAction(this, ActionID,
+		TEXT("ComponentSelectionModeEdges"),
+		LOCTEXT("ComponentSelectionModeEdges", "Edge Selection Mode"),
+		LOCTEXT("ComponentSelectionModeEdgesTooltip", "Enable edge selection mode"),
+		EModifierKey::None, EKeys::Two,
+		[this]() { SetComponentSelectionMode(EComponentSelectionMode::Edges); });
+	++ActionID;
+
+	ActionSet.RegisterAction(this, ActionID,
+		TEXT("ComponentSelectionModeFaces"),
+		LOCTEXT("ComponentSelectionModeFaces", "Face Selection Mode"),
+		LOCTEXT("ComponentSelectionModeFacesTooltip", "Enable face selection mode"),
+		EModifierKey::None, EKeys::Three,
+		[this]() { SetComponentSelectionMode(EComponentSelectionMode::Faces); });
+	++ActionID;
+
+	//-------------------------------------------------------------------------------------------------------
+	// Selection actions
+
+	ActionSet.RegisterAction(this, ActionID,
+		TEXT("SelectionGrow"),
+		LOCTEXT("SelectionGrow", "Grow Selection"),
+		LOCTEXT("SelectionGrowTooltip", "Grow the current selection"),
+		EModifierKey::None, EKeys::Period,
+		[this]() { GrowSelection(); });
+	++ActionID;
+
+	ActionSet.RegisterAction(this, ActionID,
+		TEXT("SelectionShrink"),
+		LOCTEXT("SelectionShrink", "Shrink Selection"),
+		LOCTEXT("SelectionShrinkTooltip", "Shrink the current selection"),
+		EModifierKey::None, EKeys::Comma,
+		[this]() { ShrinkSelection(); });
+	++ActionID;
+
+	ActionSet.RegisterAction(this, ActionID,
+		TEXT("SelectionFlood"),
+		LOCTEXT("SelectionFlood", "Flood Selection"),
+		LOCTEXT("SelectionFloodTooltip", "Flood the current selection"),
+		EModifierKey::Control, EKeys::A,
+		[this]() { FloodSelection(); });
+	++ActionID;
+
+	//-------------------------------------------------------------------------------------------------------
+	// Mirror actions
+
+	ActionSet.RegisterAction(this, ActionID,
+		TEXT("MirrorDirection"),
+		LOCTEXT("MirrorDirection", "Invert Mirror Direction"),
+		LOCTEXT("MirrorDirectionTooltip", "Invert the current mirror direction"),
+		EModifierKey::Control, EKeys::M,
+		[this]() 
+		{ 
+			if (ToolProperties)
+			{
+				ToolProperties->MirrorProperties.MirrorDirection =
+					(ToolProperties->MirrorProperties.MirrorDirection == EDataflowEditorToolMirrorDirection::PositiveToNegative)
+					? EDataflowEditorToolMirrorDirection::NegativeToPositive
+					: EDataflowEditorToolMirrorDirection::PositiveToNegative;
+			}
+		});
+	++ActionID;
+
+	ActionSet.RegisterAction(this, ActionID,
+		TEXT("MirrorValues"),
+		LOCTEXT("MirrorValues", "Mirror Values"),
+		LOCTEXT("MirrorValuesTooltip", "Mirror the current selection values"),
+		EModifierKey::None, EKeys::M,
+		[this]() { MirrorValues(); });
+	++ActionID;
+
+	//-------------------------------------------------------------------------------------------------------
+	// Copy Paste
+
+	ActionSet.RegisterAction(this, ActionID,
+		TEXT("CopyValue"),
+		LOCTEXT("CopyValue", "Copy Average Value"),
+		LOCTEXT("CopyValueTooltip", "Copy the current selection average value"),
+		EModifierKey::Control, EKeys::C,
+		[this]() { CopyAverageFromSelectionToClipboard(); });
+	++ActionID;
+
+	ActionSet.RegisterAction(this, ActionID,
+		TEXT("PasteValue"),
+		LOCTEXT("PasteValue", "Paste Value"),
+		LOCTEXT("PasteValueTooltip", "Paste value onto the current selection"),
+		EModifierKey::Control, EKeys::V,
+		[this]() { PasteValueToSelectionFromClipboard(); });
+	++ActionID;
+
+	//-------------------------------------------------------------------------------------------------------
+	// Value picking
+
+	ActionSet.RegisterAction(this, ActionID,
+		TEXT("PickValueUnderCursor"),
+		LOCTEXT("PickValueUnderCursor", "Pick Value"),
+		LOCTEXT("PickValueUnderCursorTooltip", "Set the active value to that currently under the cursor"),
+		EModifierKey::None, EKeys::T,
+		[this]() { EnableValuePicker(!IsValuePickerEnabled()); });
+	++ActionID;
 
-	ActionSet.RegisterAction(this, (int32)EStandardToolActions::BaseClientDefinedActionID + 500,
-		TEXT("PickWeightValueUnderCursor"),
-		LOCTEXT("PickWeightValueUnderCursor", "Pick Weight Value"),
-		LOCTEXT("PickWeightValueUnderCursorTooltip", "Set the active weight painting value to that currently under the cursor"),
-		EModifierKey::Shift, EKeys::G,
-		[this]() { bPendingPickWeight = true; });
 };
 
 void UDataflowEditorVertexAttributePaintTool::OnPropertyModified(UObject* PropertySet, FProperty* Property)
@@ -1083,6 +1227,15 @@ void UDataflowEditorVertexAttributePaintTool::OnEndStroke()
 
 	// close change record
 	EndChange();
+
+	if (bPendingPickWeight)
+	{
+		if (GetSculptMesh()->IsTriangle(GetBrushTriangleID()))
+		{
+			ToolProperties->BrushProperties.AttributeValue = GetCurrentWeightValueUnderBrush();
+			NotifyOfPropertyChangeByTool(ToolProperties);
+		}
+	}
 }
 
 void UDataflowEditorVertexAttributePaintTool::OnCancelStroke()
@@ -2112,7 +2265,9 @@ void UDataflowEditorVertexAttributePaintTool::DrawHUD(FCanvas* Canvas, IToolsCon
 			RenderAPI->GetSceneView()->ScreenToPixel(ScreenBrushPos, ScreenBrushPos2D);
 
 			// Render value under at brush even in mesh mode
-			const FText BaseText = LOCTEXT("DataflowWeightMapPaintTool_ValueAtBrush_TextFormat", "Value at brush: {0}");
+			const FText BaseText = bPendingPickWeight
+				? LOCTEXT("DataflowWeightMapPaintTool_PickValueAtBrush_TextFormat", "Pick Value: {0}")
+				: LOCTEXT("DataflowWeightMapPaintTool_ValueAtBrush_TextFormat", "Value at brush: {0}");
 			const FText ValueAtBrushText = FText::Format(BaseText, FText::AsNumber(ToolProperties->BrushProperties.ValueAtBrush));
 			FCanvasTextItem TextItem(ScreenBrushPos2D, ValueAtBrushText, GEngine->GetLargeFont(), FLinearColor::White);
 			TextItem.EnableShadow(FLinearColor::Black);
@@ -2243,7 +2398,7 @@ void UDataflowEditorVertexAttributePaintTool::OnTick(float DeltaTime)
 	}
 
 	ConfigureIndicator(IsVolumetricBrush());
-	SetIndicatorVisibility(!bIsLasso && !bIsMeshEditMode && bBrushIsHoverMesh);
+	SetIndicatorVisibility(!bIsLasso && !bIsMeshEditMode && bBrushIsHoverMesh && !bPendingPickWeight);
 
 	SCOPE_CYCLE_COUNTER(VertexAttributePaintTool_Tick);
 
@@ -2272,18 +2427,18 @@ void UDataflowEditorVertexAttributePaintTool::OnTick(float DeltaTime)
 		{
 			const double HitWeightValue = GetCurrentWeightValueUnderBrush();
 
-			if (bShouldPickWeight)
-			{
-				ToolProperties->BrushProperties.AttributeValue = HitWeightValue;
-				NotifyOfPropertyChangeByTool(ToolProperties);
-			}
+			//if (bShouldPickWeight)
+			//{
+			//	ToolProperties->BrushProperties.AttributeValue = HitWeightValue;
+			//	NotifyOfPropertyChangeByTool(ToolProperties);
+			//}
 
 			/*if (bShouldUpdateValueAtBrush)*/
 			{
 				ToolProperties->BrushProperties.ValueAtBrush = HitWeightValue;
 			}
 		}
-		bPendingPickWeight = false;
+		//bPendingPickWeight = false;
 	}
 
 	auto ExecuteStampOperation = [this](int StampIndex, const FRay& StampRay)
@@ -2315,8 +2470,7 @@ void UDataflowEditorVertexAttributePaintTool::OnTick(float DeltaTime)
 			AccumulateROI.Wait();
 		};
 
-
-	if (IsInBrushMode())
+	if (IsInBrushMode() && !bPendingPickWeight)
 	{
 		ProcessPerTickStamps(
 			[this](const FRay& StampRay) -> bool {
@@ -2324,6 +2478,18 @@ void UDataflowEditorVertexAttributePaintTool::OnTick(float DeltaTime)
 			}, ExecuteStampOperation);
 	}
 
+	if (UToolsContextCursorAPI* ToolsContextCursorAPI = GetToolManager()->GetContextObjectStore()->FindContext<UToolsContextCursorAPI>())
+	{
+		if (bPendingPickWeight)
+		{
+			ToolsContextCursorAPI->SetCursorOverride(EMouseCursor::EyeDropper);
+		}
+		else
+		{
+			ToolsContextCursorAPI->ClearCursorOverride();
+		}
+	}
+
 }
 
 bool UDataflowEditorVertexAttributePaintTool::CanAccept() const
diff --git a/Engine/Plugins/Experimental/Dataflow/Source/DataflowEditor/Private/DataflowEditorTools/DataflowEditorVertexAttributePaintToolDetailCustomization.cpp b/Engine/Plugins/Experimental/Dataflow/Source/DataflowEditor/Private/DataflowEditorTools/DataflowEditorVertexAttributePaintToolDetailCustomization.cpp
--- a/Engine/Plugins/Experimental/Dataflow/Source/DataflowEditor/Private/DataflowEditorTools/DataflowEditorVertexAttributePaintToolDetailCustomization.cpp
+++ b/Engine/Plugins/Experimental/Dataflow/Source/DataflowEditor/Private/DataflowEditorTools/DataflowEditorVertexAttributePaintToolDetailCustomization.cpp
@@ -764,31 +764,64 @@ namespace UE::Dataflow::Editor
 		]
 		.ValueContent()
 		[
-			SNew(SSpinBox<float>)
-			.MinValue_Lambda(GetMinValue)
-			.MaxValue_Lambda(GetMaxValue)
-			.Value(1.0f)
-			.SupportDynamicSliderMaxValue(false)
-			.Value_Lambda([this]()
-				{
-					return ToolProperties.IsValid() ? ToolProperties->BrushProperties.AttributeValue : 1.f;
-				})
-			.OnValueChanged_Lambda([this, GetMinValue, GetMaxValue](float NewValue)
-				{
-					if (ToolProperties.IsValid())
+			SNew(SHorizontalBox)
+			+ SHorizontalBox::Slot()
+			.Padding(2.f)
+			.FillWidth(1.0f)
+			[
+				SNew(SSpinBox<float>)
+				.MinValue_Lambda(GetMinValue)
+				.MaxValue_Lambda(GetMaxValue)
+				.Value(1.0f)
+				.SupportDynamicSliderMaxValue(false)
+				.Value_Lambda([this]()
 					{
-						ToolProperties->BrushProperties.AttributeValue = FMath::Clamp(NewValue, GetMinValue(), GetMaxValue());
-						FPropertyChangedEvent PropertyChangedEvent(UBrushBaseProperties::StaticClass()->FindPropertyByName(GET_MEMBER_NAME_CHECKED(FDataflowEditorVertexAttributePaintToolBrushProperties, AttributeValue)));
-						ToolProperties->PostEditChangeProperty(PropertyChangedEvent);
-					}
-				})
-			.OnValueCommitted_Lambda([this](float NewValue, ETextCommit::Type CommitType)
-				{
-					if (ToolProperties.IsValid())
+						return ToolProperties.IsValid() ? ToolProperties->BrushProperties.AttributeValue : 1.f;
+					})
+				.OnValueChanged_Lambda([this, GetMinValue, GetMaxValue](float NewValue)
 					{
-						// TODO(ccaillaud) ==>> ToolProperties->SaveConfig();
-					}
-				})
+						if (ToolProperties.IsValid())
+						{
+							ToolProperties->BrushProperties.AttributeValue = FMath::Clamp(NewValue, GetMinValue(), GetMaxValue());
+							FPropertyChangedEvent PropertyChangedEvent(UBrushBaseProperties::StaticClass()->FindPropertyByName(GET_MEMBER_NAME_CHECKED(FDataflowEditorVertexAttributePaintToolBrushProperties, AttributeValue)));
+							ToolProperties->PostEditChangeProperty(PropertyChangedEvent);
+						}
+					})
+				.OnValueCommitted_Lambda([this](float NewValue, ETextCommit::Type CommitType)
+					{
+						if (ToolProperties.IsValid())
+						{
+							// TODO(ccaillaud) ==>> ToolProperties->SaveConfig();
+						}
+					})
+			]
+			+ SHorizontalBox::Slot()
+			.Padding(2.f)
+			.AutoWidth()
+			[
+				SNew(SCheckBox)
+				.Style(FAppStyle::Get(), "ToggleButtonCheckbox")
+				.IsChecked_Lambda([Tool = Tool]()-> ECheckBoxState
+					{
+						const bool bPickerEnabled = (Tool.IsValid() && Tool->IsValuePickerEnabled());
+						return bPickerEnabled? ECheckBoxState::Checked: ECheckBoxState::Unchecked;
+					})
+				.OnCheckStateChanged_Lambda([Tool = Tool](const ECheckBoxState NewState)
+					{
+						if (Tool.IsValid())
+						{
+							Tool->EnableValuePicker(NewState == ECheckBoxState::Checked);
+						}
+					})
+				.Padding(2.0f)
+				.ToolTipText(LOCTEXT("ValuePickerTooltip", "Toggle for value picker"))
+				.HAlign(HAlign_Center)
+				[
+					SNew(SImage)
+					.Image(FAppStyle::GetBrush("Icons.EyeDropper"))
+					.ColorAndOpacity(FSlateColor::UseForeground())
+				]
+			]
 		];
 	}
 
diff --git a/Engine/Plugins/Experimental/Dataflow/Source/DataflowEditor/Public/DataflowEditorTools/DataflowEditorVertexAttributePaintTool.h b/Engine/Plugins/Experimental/Dataflow/Source/DataflowEditor/Public/DataflowEditorTools/DataflowEditorVertexAttributePaintTool.h
--- a/Engine/Plugins/Experimental/Dataflow/Source/DataflowEditor/Public/DataflowEditorTools/DataflowEditorVertexAttributePaintTool.h
+++ b/Engine/Plugins/Experimental/Dataflow/Source/DataflowEditor/Public/DataflowEditorTools/DataflowEditorVertexAttributePaintTool.h
@@ -457,6 +457,9 @@ protected:
 	virtual bool SharesBrushPropertiesChanges() const override { return false; }
 	// end UMeshSculptToolBase API
 
+	void EnableValuePicker(bool bEnable);
+	bool IsValuePickerEnabled();
+
 protected:
 	//
 	// Marquee Support
