Dataflow paint tools : Add keyboard shortcut information to tooltips when relevant
( value picker, mirror buttons, selection grow, shrink, flood, color ramp options )
- Added an method to inputChord class to get the modifier key from it

#Synced-CL 49632091

diff --git a/Engine/Plugins/Experimental/Dataflow/Source/DataflowEditor/Private/DataflowEditorTools/DataflowEditorVertexAttributePaintTool.cpp b/Engine/Plugins/Experimental/Dataflow/Source/DataflowEditor/Private/DataflowEditorTools/DataflowEditorVertexAttributePaintTool.cpp
--- a/Engine/Plugins/Experimental/Dataflow/Source/DataflowEditor/Private/DataflowEditorTools/DataflowEditorVertexAttributePaintTool.cpp
+++ b/Engine/Plugins/Experimental/Dataflow/Source/DataflowEditor/Private/DataflowEditorTools/DataflowEditorVertexAttributePaintTool.cpp
@@ -871,6 +871,26 @@ bool UDataflowEditorVertexAttributePaintTool::IsValuePickerEnabled()
 	return bPendingPickWeight;
 }
 
+FInputChord UDataflowEditorVertexAttributePaintTool::GreyscaleDisplayInputChord = FInputChord(EModifierKey::Control, EKeys::NumPadSeven);
+FInputChord UDataflowEditorVertexAttributePaintTool::ColorMapDisplayInputChord = FInputChord(EModifierKey::Control, EKeys::NumPadEight);
+FInputChord UDataflowEditorVertexAttributePaintTool::MaterialColorDisplayInputChord = FInputChord(EModifierKey::Control, EKeys::NumPadSix);
+
+FInputChord UDataflowEditorVertexAttributePaintTool::VerticesSelectionModeInputChord = FInputChord(EModifierKey::None, EKeys::One);
+FInputChord UDataflowEditorVertexAttributePaintTool::EdgesSelectionModeInputChord = FInputChord(EModifierKey::None, EKeys::Two);
+FInputChord UDataflowEditorVertexAttributePaintTool::FacesSelectionModeInputChord = FInputChord(EModifierKey::None, EKeys::Three);
+
+FInputChord UDataflowEditorVertexAttributePaintTool::GrowSelectionInputChord = FInputChord(EModifierKey::None, EKeys::Period);
+FInputChord UDataflowEditorVertexAttributePaintTool::ShrinkSelectionInputChord = FInputChord(EModifierKey::None, EKeys::Comma);
+FInputChord UDataflowEditorVertexAttributePaintTool::FloodSelectionInputChord = FInputChord(EModifierKey::Control, EKeys::A);
+
+FInputChord UDataflowEditorVertexAttributePaintTool::InvertMirrorDirectionInputChord = FInputChord(EModifierKey::Control, EKeys::M);
+FInputChord UDataflowEditorVertexAttributePaintTool::MirrorValuesInputChord = FInputChord(EModifierKey::None, EKeys::M);
+
+FInputChord UDataflowEditorVertexAttributePaintTool::CopyAverageValueInputChord = FInputChord(EModifierKey::Control, EKeys::C);
+FInputChord UDataflowEditorVertexAttributePaintTool::PasteAverageValueInputChord = FInputChord(EModifierKey::Control, EKeys::V);
+
+FInputChord UDataflowEditorVertexAttributePaintTool::ToggleValuePickerInputChord = FInputChord(EModifierKey::None, EKeys::T);
+
 void UDataflowEditorVertexAttributePaintTool::RegisterActions(FInteractiveToolActionSet& ActionSet)
 {
 	// IMPORTANT : we do not register the base class actions as they do not map enough to the wanted actions for this tool
@@ -884,7 +904,7 @@ void UDataflowEditorVertexAttributePaintTool::RegisterActions(FInteractiveToolAc
 		TEXT("MeshDisplayGreyscale"),
 		LOCTEXT("MeshDisplayGreyscale", "Greyscale Display"),
 		LOCTEXT("MeshDisplayGreyscaleTooltip", "Use greyscale color ramp to display values"),
-		EModifierKey::Control, EKeys::NumPadSeven,
+		GreyscaleDisplayInputChord.GetModifierKey(), GreyscaleDisplayInputChord.Key,
 		[this]() { SetColorMode(EDataflowEditorToolColorMode::Greyscale); });
 	++ActionID;
 
@@ -892,7 +912,7 @@ void UDataflowEditorVertexAttributePaintTool::RegisterActions(FInteractiveToolAc
 		TEXT("MeshDisplayColorRamp"),
 		LOCTEXT("MeshDisplayColorRamp", "Color Ramp Display"),
 		LOCTEXT("MeshDisplayColorRampTooltip", "Use custom color ramp to display values"),
-		EModifierKey::Control, EKeys::NumPadEight,
+		ColorMapDisplayInputChord.GetModifierKey(), ColorMapDisplayInputChord.Key,
 		[this]() { SetColorMode(EDataflowEditorToolColorMode::Ramp); });
 	++ActionID;
 
@@ -900,7 +920,7 @@ void UDataflowEditorVertexAttributePaintTool::RegisterActions(FInteractiveToolAc
 		TEXT("MeshDisplayFullMaterial"),
 		LOCTEXT("MeshDisplayFullMaterial", "Full Material Display"),
 		LOCTEXT("MeshDisplayFullMaterialTooltip", "Use normal material color to display values"),
-		EModifierKey::Control, EKeys::NumPadSix,
+		MaterialColorDisplayInputChord.GetModifierKey(), MaterialColorDisplayInputChord.Key,
 		[this]() { SetColorMode(EDataflowEditorToolColorMode::FullMaterial); });
 	++ActionID;
 
@@ -911,7 +931,7 @@ void UDataflowEditorVertexAttributePaintTool::RegisterActions(FInteractiveToolAc
 		TEXT("ComponentSelectionModeVertices"),
 		LOCTEXT("ComponentSelectionModeVertices", "Vertex Selection Mode"),
 		LOCTEXT("ComponentSelectionModeVerticesTooltip", "Enable vertex selection mode"),
-		EModifierKey::None, EKeys::One,
+		VerticesSelectionModeInputChord.GetModifierKey(), VerticesSelectionModeInputChord.Key,
 		[this]() { SetComponentSelectionMode(EComponentSelectionMode::Vertices); });
 	++ActionID;
 
@@ -919,7 +939,7 @@ void UDataflowEditorVertexAttributePaintTool::RegisterActions(FInteractiveToolAc
 		TEXT("ComponentSelectionModeEdges"),
 		LOCTEXT("ComponentSelectionModeEdges", "Edge Selection Mode"),
 		LOCTEXT("ComponentSelectionModeEdgesTooltip", "Enable edge selection mode"),
-		EModifierKey::None, EKeys::Two,
+		EdgesSelectionModeInputChord.GetModifierKey(), EdgesSelectionModeInputChord.Key,
 		[this]() { SetComponentSelectionMode(EComponentSelectionMode::Edges); });
 	++ActionID;
 
@@ -927,7 +947,7 @@ void UDataflowEditorVertexAttributePaintTool::RegisterActions(FInteractiveToolAc
 		TEXT("ComponentSelectionModeFaces"),
 		LOCTEXT("ComponentSelectionModeFaces", "Face Selection Mode"),
 		LOCTEXT("ComponentSelectionModeFacesTooltip", "Enable face selection mode"),
-		EModifierKey::None, EKeys::Three,
+		FacesSelectionModeInputChord.GetModifierKey(), FacesSelectionModeInputChord.Key,
 		[this]() { SetComponentSelectionMode(EComponentSelectionMode::Faces); });
 	++ActionID;
 
@@ -938,7 +958,7 @@ void UDataflowEditorVertexAttributePaintTool::RegisterActions(FInteractiveToolAc
 		TEXT("SelectionGrow"),
 		LOCTEXT("SelectionGrow", "Grow Selection"),
 		LOCTEXT("SelectionGrowTooltip", "Grow the current selection"),
-		EModifierKey::None, EKeys::Period,
+		GrowSelectionInputChord.GetModifierKey(), GrowSelectionInputChord.Key,
 		[this]() { GrowSelection(); });
 	++ActionID;
 
@@ -946,7 +966,7 @@ void UDataflowEditorVertexAttributePaintTool::RegisterActions(FInteractiveToolAc
 		TEXT("SelectionShrink"),
 		LOCTEXT("SelectionShrink", "Shrink Selection"),
 		LOCTEXT("SelectionShrinkTooltip", "Shrink the current selection"),
-		EModifierKey::None, EKeys::Comma,
+		ShrinkSelectionInputChord.GetModifierKey(), ShrinkSelectionInputChord.Key,
 		[this]() { ShrinkSelection(); });
 	++ActionID;
 
@@ -954,7 +974,7 @@ void UDataflowEditorVertexAttributePaintTool::RegisterActions(FInteractiveToolAc
 		TEXT("SelectionFlood"),
 		LOCTEXT("SelectionFlood", "Flood Selection"),
 		LOCTEXT("SelectionFloodTooltip", "Flood the current selection"),
-		EModifierKey::Control, EKeys::A,
+		FloodSelectionInputChord.GetModifierKey(), FloodSelectionInputChord.Key,
 		[this]() { FloodSelection(); });
 	++ActionID;
 
@@ -965,7 +985,7 @@ void UDataflowEditorVertexAttributePaintTool::RegisterActions(FInteractiveToolAc
 		TEXT("MirrorDirection"),
 		LOCTEXT("MirrorDirection", "Invert Mirror Direction"),
 		LOCTEXT("MirrorDirectionTooltip", "Invert the current mirror direction"),
-		EModifierKey::Control, EKeys::M,
+		InvertMirrorDirectionInputChord.GetModifierKey(), InvertMirrorDirectionInputChord.Key,
 		[this]() 
 		{ 
 			if (ToolProperties)
@@ -982,7 +1002,7 @@ void UDataflowEditorVertexAttributePaintTool::RegisterActions(FInteractiveToolAc
 		TEXT("MirrorValues"),
 		LOCTEXT("MirrorValues", "Mirror Values"),
 		LOCTEXT("MirrorValuesTooltip", "Mirror the current selection values"),
-		EModifierKey::None, EKeys::M,
+		MirrorValuesInputChord.GetModifierKey(), MirrorValuesInputChord.Key,
 		[this]() { MirrorValues(); });
 	++ActionID;
 
@@ -993,7 +1013,7 @@ void UDataflowEditorVertexAttributePaintTool::RegisterActions(FInteractiveToolAc
 		TEXT("CopyValue"),
 		LOCTEXT("CopyValue", "Copy Average Value"),
 		LOCTEXT("CopyValueTooltip", "Copy the current selection average value"),
-		EModifierKey::Control, EKeys::C,
+		CopyAverageValueInputChord.GetModifierKey(), CopyAverageValueInputChord.Key,
 		[this]() { CopyAverageFromSelectionToClipboard(); });
 	++ActionID;
 
@@ -1001,7 +1021,7 @@ void UDataflowEditorVertexAttributePaintTool::RegisterActions(FInteractiveToolAc
 		TEXT("PasteValue"),
 		LOCTEXT("PasteValue", "Paste Value"),
 		LOCTEXT("PasteValueTooltip", "Paste value onto the current selection"),
-		EModifierKey::Control, EKeys::V,
+		PasteAverageValueInputChord.GetModifierKey(), PasteAverageValueInputChord.Key,
 		[this]() { PasteValueToSelectionFromClipboard(); });
 	++ActionID;
 
@@ -1012,7 +1032,7 @@ void UDataflowEditorVertexAttributePaintTool::RegisterActions(FInteractiveToolAc
 		TEXT("PickValueUnderCursor"),
 		LOCTEXT("PickValueUnderCursor", "Pick Value"),
 		LOCTEXT("PickValueUnderCursorTooltip", "Set the active value to that currently under the cursor"),
-		EModifierKey::None, EKeys::T,
+		ToggleValuePickerInputChord.GetModifierKey(), ToggleValuePickerInputChord.Key,
 		[this]() { EnableValuePicker(!IsValuePickerEnabled()); });
 	++ActionID;
 
diff --git a/Engine/Plugins/Experimental/Dataflow/Source/DataflowEditor/Private/DataflowEditorTools/DataflowEditorVertexAttributePaintToolDetailCustomization.cpp b/Engine/Plugins/Experimental/Dataflow/Source/DataflowEditor/Private/DataflowEditorTools/DataflowEditorVertexAttributePaintToolDetailCustomization.cpp
--- a/Engine/Plugins/Experimental/Dataflow/Source/DataflowEditor/Private/DataflowEditorTools/DataflowEditorVertexAttributePaintToolDetailCustomization.cpp
+++ b/Engine/Plugins/Experimental/Dataflow/Source/DataflowEditor/Private/DataflowEditorTools/DataflowEditorVertexAttributePaintToolDetailCustomization.cpp
@@ -18,6 +18,9 @@
 #include "Widgets/Input/SComboBox.h"
 #include "Framework/MultiBox/MultiBoxBuilder.h"
 #include "Drawing/MeshElementsVisualizer.h"
+#include "Framework/Commands/InputChord.h"
+#include "IDocumentation.h"
+#include "Widgets/SToolTip.h"
 
 #define LOCTEXT_NAMESPACE "VertexAttributePaintToolDetailCustomization"
 
@@ -58,6 +61,18 @@ namespace UE::Dataflow::Editor
 		}
 	}
 
+	static TSharedRef<SToolTip> MakeTooltipWithShortcut(const FText& TooltipText, const FInputChord& InputChord)
+	{
+		//return FText::Format(LOCTEXT("TooltipWithShorcut", "{0} {1}"), { TooltipText, InputChord.GetInputText() });
+
+		//FFormatNamedArguments Args;
+		//Args.Add(TEXT("ToolTipDescription"), TooltipText);
+		//Args.Add(TEXT("Keybinding"), InputChord.GetInputText());
+		//return FText::Format(NSLOCTEXT("ToolBar", "ToolTip + Keybinding", "{ToolTipDescription} ({Keybinding})"), Args);
+
+		return IDocumentation::Get()->CreateToolTip(TooltipText, nullptr, FString(), FString(), InputChord.GetInputText());
+	}
+
 	///////////////////////////////////////////////////////////////////////////////////////////////////////////////
 	FVertexAttributePaintToolDetailCustomization::FVertexAttributePaintToolDetailCustomization()
 	{
@@ -175,10 +190,13 @@ namespace UE::Dataflow::Editor
 					})
 				+ SSegmentedControl<EDataflowEditorToolColorMode>::Slot(EDataflowEditorToolColorMode::Greyscale)
 				.Text(LOCTEXT("GreyscaleMode", "Greyscale"))
+				.ToolTipWidget(MakeTooltipWithShortcut(LOCTEXT("GreyscaleModeTooltip", "Greyscale Mode"), UDataflowEditorVertexAttributePaintTool::GreyscaleDisplayInputChord))
 				+ SSegmentedControl<EDataflowEditorToolColorMode>::Slot(EDataflowEditorToolColorMode::Ramp)
 				.Text(LOCTEXT("RampMode", "Ramp"))
+				.ToolTipWidget(MakeTooltipWithShortcut(LOCTEXT("RampModeModeTooltip", "Color Ramp Mode"), UDataflowEditorVertexAttributePaintTool::ColorMapDisplayInputChord))
 				+ SSegmentedControl<EDataflowEditorToolColorMode>::Slot(EDataflowEditorToolColorMode::FullMaterial)
 				.Text(LOCTEXT("MaterialMode", "Full Material"))
+				.ToolTipWidget(MakeTooltipWithShortcut(LOCTEXT("MaterialModeTooltip", "Material Mode"), UDataflowEditorVertexAttributePaintTool::MaterialColorDisplayInputChord))
 			]
 		];
 	}
@@ -814,7 +832,7 @@ namespace UE::Dataflow::Editor
 						}
 					})
 				.Padding(2.0f)
-				.ToolTipText(LOCTEXT("ValuePickerTooltip", "Toggle for value picker"))
+				.ToolTip(MakeTooltipWithShortcut(LOCTEXT("ValuePickerTooltip", "Toggle for value picker"), UDataflowEditorVertexAttributePaintTool::ToggleValuePickerInputChord))
 				.HAlign(HAlign_Center)
 				[
 					SNew(SImage)
@@ -1170,15 +1188,17 @@ namespace UE::Dataflow::Editor
 				.Padding(2.f, WeightEditVerticalPadding)
 				[
 					SNew(SButton)
-					.IsEnabled_Lambda([this]
-						{
-							return Tool.IsValid() ? Tool->HasSelection() : false;
-						})
-					.HAlign(HAlign_Center)
-					.VAlign(VAlign_Center)
-					.Text(LOCTEXT("GrowSelectionButtonLabel", "Grow"))
-					.ToolTipText(LOCTEXT("GrowSelectionTooltip",
-						"Grow the current selection by adding connected neighbors to current selection.\n"))
+						.IsEnabled_Lambda([this]
+							{
+								return Tool.IsValid() ? Tool->HasSelection() : false;
+							})
+						.HAlign(HAlign_Center)
+						.VAlign(VAlign_Center)
+						.Text(LOCTEXT("GrowSelectionButtonLabel", "Grow"))
+						.ToolTip(MakeTooltipWithShortcut(
+							LOCTEXT("GrowSelectionTooltip", "Grow the current selection by adding connected neighbors to current selection."), 
+							UDataflowEditorVertexAttributePaintTool::GrowSelectionInputChord
+							))
 					.OnClicked_Lambda([this]()
 						{
 							if (Tool.IsValid())
@@ -1199,8 +1219,10 @@ namespace UE::Dataflow::Editor
 					.HAlign(HAlign_Center)
 					.VAlign(VAlign_Center)
 					.Text(LOCTEXT("ShrinkSelectionButtonLabel", "Shrink"))
-					.ToolTipText(LOCTEXT("ShrinkSelectionTooltip",
-						"Shrink the current selection by removing components on the border of the current selection.\n"))
+					.ToolTip(MakeTooltipWithShortcut(
+						LOCTEXT("ShrinkSelectionTooltip", "Shrink the current selection by removing components on the border of the current selection."),
+						UDataflowEditorVertexAttributePaintTool::ShrinkSelectionInputChord
+						))
 					.OnClicked_Lambda([this]()
 						{
 							if (Tool.IsValid())
@@ -1226,8 +1248,10 @@ namespace UE::Dataflow::Editor
 					.HAlign(HAlign_Center)
 					.VAlign(VAlign_Center)
 					.Text(LOCTEXT("FloodSelectionButtonLabel", "Flood"))
-					.ToolTipText(LOCTEXT("FloodSelectionTooltip",
-						"Flood the current selection by adding all connected components to the current selection.\n"))
+						.ToolTip(MakeTooltipWithShortcut(
+							LOCTEXT("FloodSelectionTooltip", "Flood the current selection by adding all connected components to the current selection."),
+							UDataflowEditorVertexAttributePaintTool::FloodSelectionInputChord
+						))
 					.OnClicked_Lambda([this]()
 						{
 							if (Tool.IsValid())
@@ -1618,9 +1642,7 @@ namespace UE::Dataflow::Editor
 					SNew(SButton)
 					.HAlign(HAlign_Center)
 					.Text(LOCTEXT("CopyValuesButtonLabel", "Copy"))
-					.ToolTipText(LOCTEXT("CopyButtonTooltip",
-						"Copy the average of the selected values to the clipboard. \n"
-						"This is designed to work with the Paste command."))
+					.ToolTip(MakeTooltipWithShortcut(LOCTEXT("CopyButtonTooltip", "Copy the average of the selected values to the clipboard.\nThis is designed to work with the Paste command."), UDataflowEditorVertexAttributePaintTool::CopyAverageValueInputChord))
 					.OnClicked_Lambda([this]()
 						{
 							if (Tool.IsValid())
@@ -1638,9 +1660,7 @@ namespace UE::Dataflow::Editor
 				SNew(SButton)
 				.HAlign(HAlign_Center)
 				.Text(LOCTEXT("PasteValuesButtonLabel", "Paste"))
-				.ToolTipText(LOCTEXT("PasteButtonTooltip",
-					"Paste the value on the selected vertices.\n"
-					"This command requires the clipboard contain the selection average value from the Copy command."))
+				.ToolTip(MakeTooltipWithShortcut(LOCTEXT("PasteButtonTooltip", "Paste the value on the selected vertices.\nThis command requires the clipboard contain the selection average value from the Copy command."), UDataflowEditorVertexAttributePaintTool::PasteAverageValueInputChord))
 				.OnClicked_Lambda([this]()
 					{
 						if (ToolProperties.IsValid())
@@ -1679,32 +1699,33 @@ namespace UE::Dataflow::Editor
 			+ SSegmentedControl<EAxis::Type>::Slot(EAxis::Z)
 			.Text(LOCTEXT("MirrorZLabel", "Z"));
 
-			TSharedRef<SSegmentedControl<EDataflowEditorToolMirrorDirection>> MirrortSideButtonsWidget =
-				SNew(SSegmentedControl<EDataflowEditorToolMirrorDirection>)
-				.ToolTipText(LOCTEXT("MirrorDirectionTooltip", "The direction that determines what side of the plane to copy weights from."))
-				.Value_Lambda([this]()
-					{
-						return ToolProperties.IsValid() ? ToolProperties->MirrorProperties.MirrorDirection : EDataflowEditorToolMirrorDirection::PositiveToNegative;
-					})
-				.OnValueChanged_Lambda([this](EDataflowEditorToolMirrorDirection Mode)
+		TSharedRef<SSegmentedControl<EDataflowEditorToolMirrorDirection>> MirrortSideButtonsWidget =
+			SNew(SSegmentedControl<EDataflowEditorToolMirrorDirection>)
+			.ToolTipText(LOCTEXT("MirrorDirectionTooltip", "The direction that determines what side of the plane to copy weights from."))
+			.Value_Lambda([this]()
+				{
+					return ToolProperties.IsValid() ? ToolProperties->MirrorProperties.MirrorDirection : EDataflowEditorToolMirrorDirection::PositiveToNegative;
+				})
+			.OnValueChanged_Lambda([this](EDataflowEditorToolMirrorDirection Mode)
+				{
+					if (ToolProperties.IsValid())
 					{
-						if (ToolProperties.IsValid())
-						{
-							ToolProperties->MirrorProperties.MirrorDirection = Mode;
-						}
-					})
-				+ SSegmentedControl<EDataflowEditorToolMirrorDirection>::Slot(EDataflowEditorToolMirrorDirection::PositiveToNegative)
-				.Text(LOCTEXT("MirrorPosToNegLabel", "+ to -"))
-				+ SSegmentedControl<EDataflowEditorToolMirrorDirection>::Slot(EDataflowEditorToolMirrorDirection::NegativeToPositive)
-				.Text(LOCTEXT("MirrorNegToPosLabel", "- to +"));
+						ToolProperties->MirrorProperties.MirrorDirection = Mode;
+					}
+				})
+			+ SSegmentedControl<EDataflowEditorToolMirrorDirection>::Slot(EDataflowEditorToolMirrorDirection::PositiveToNegative)
+			.Text(LOCTEXT("MirrorPosToNegLabel", "+ to -"))
+			.ToolTipWidget(MakeTooltipWithShortcut(LOCTEXT("MirrorDirectionPosToNegTooltip", "Mirror the values from the positive side of the axis to the negative side."), UDataflowEditorVertexAttributePaintTool::InvertMirrorDirectionInputChord))
+			+ SSegmentedControl<EDataflowEditorToolMirrorDirection>::Slot(EDataflowEditorToolMirrorDirection::NegativeToPositive)
+			.Text(LOCTEXT("MirrorNegToPosLabel", "- to +"))
+			.ToolTipWidget(MakeTooltipWithShortcut(LOCTEXT("MirrorDirectionNegToPosTooltip", "Mirror the values from the negative side of the axis to the positive side."), UDataflowEditorVertexAttributePaintTool::InvertMirrorDirectionInputChord))
+			;
 
 		TSharedRef<SButton> MirrorButtonWidget = 
 			SNew(SButton)
 			.HAlign(HAlign_Center)
 			.Text(LOCTEXT("MirrorWeightsButtonLabel", "Mirror"))
-			.ToolTipText(LOCTEXT("MirrorButtonTooltip",
-				"Values are copied across the given plane in the given direction.\n"
-				"This command operates on the selected vertices."))
+			.ToolTip(MakeTooltipWithShortcut(LOCTEXT("MirrorButtonTooltip", "Values are copied across the given plane in the given direction.\nThis command operates on the selected vertices."), UDataflowEditorVertexAttributePaintTool::MirrorValuesInputChord))
 			.OnClicked_Lambda([this]()
 				{
 					if (Tool.IsValid())
diff --git a/Engine/Plugins/Experimental/Dataflow/Source/DataflowEditor/Public/DataflowEditorTools/DataflowEditorVertexAttributePaintTool.h b/Engine/Plugins/Experimental/Dataflow/Source/DataflowEditor/Public/DataflowEditorTools/DataflowEditorVertexAttributePaintTool.h
--- a/Engine/Plugins/Experimental/Dataflow/Source/DataflowEditor/Public/DataflowEditorTools/DataflowEditorVertexAttributePaintTool.h
+++ b/Engine/Plugins/Experimental/Dataflow/Source/DataflowEditor/Public/DataflowEditorTools/DataflowEditorVertexAttributePaintTool.h
@@ -28,6 +28,8 @@
 #include "DataflowEditorTools/DataflowEditorToolEnums.h"
 #include "ToolMeshSelector.h"
 
+#include "Framework/Commands/InputChord.h"
+
 #include "DataflowEditorVertexAttributePaintTool.generated.h"
 
 
@@ -435,6 +437,27 @@ public:
 
 	UE_API TArray<int32> GetSelectedVertices() const;
 
+public:
+	static FInputChord GreyscaleDisplayInputChord;
+	static FInputChord ColorMapDisplayInputChord;
+	static FInputChord MaterialColorDisplayInputChord;
+
+	static FInputChord VerticesSelectionModeInputChord;
+	static FInputChord EdgesSelectionModeInputChord;
+	static FInputChord FacesSelectionModeInputChord;
+
+	static FInputChord GrowSelectionInputChord;
+	static FInputChord ShrinkSelectionInputChord;
+	static FInputChord FloodSelectionInputChord;
+
+	static FInputChord InvertMirrorDirectionInputChord;
+	static FInputChord MirrorValuesInputChord;
+
+	static FInputChord CopyAverageValueInputChord;
+	static FInputChord PasteAverageValueInputChord;
+
+	static FInputChord ToggleValuePickerInputChord;
+
 protected:
 	// UMeshSculptToolBase API
 	virtual UBaseDynamicMeshComponent* GetSculptMeshComponent() override;
diff --git a/Engine/Source/Runtime/Slate/Private/Framework/Commands/InputChord.cpp b/Engine/Source/Runtime/Slate/Private/Framework/Commands/InputChord.cpp
--- a/Engine/Source/Runtime/Slate/Private/Framework/Commands/InputChord.cpp
+++ b/Engine/Source/Runtime/Slate/Private/Framework/Commands/InputChord.cpp
@@ -18,6 +18,11 @@ FInputChord::FInputChord()
 {
 }
 
+EModifierKey::Type FInputChord::GetModifierKey() const
+{
+	return EModifierKey::FromBools(bCtrl, bAlt, bShift, bCmd);
+}
+
 /**
  * Returns the friendly, localized string name of this key binding
  */
diff --git a/Engine/Source/Runtime/Slate/Public/Framework/Commands/InputChord.h b/Engine/Source/Runtime/Slate/Public/Framework/Commands/InputChord.h
--- a/Engine/Source/Runtime/Slate/Public/Framework/Commands/InputChord.h
+++ b/Engine/Source/Runtime/Slate/Public/Framework/Commands/InputChord.h
@@ -173,6 +173,14 @@ struct FInputChord
 	bool NeedsAlt() const { return bAlt; }
 	bool NeedsShift() const { return bShift; }
 
+
+	/**
+	 * Gets the corresponding EModifierKey::Type
+	 *
+	 * @return A localized string.
+	 */
+	SLATE_API EModifierKey::Type GetModifierKey() const;
+
 	/** 
 	 * Gets a localized string that represents the chord.
 	 *
