Dataflow : add node to compute bounds from any collection ( including Cloth one )

#Synced-CL 49661615

diff --git a/Engine/Plugins/Experimental/Dataflow/Source/DataflowNodes/Private/Dataflow/DataflowGetCollectionBoundsNode.cpp b/Engine/Plugins/Experimental/Dataflow/Source/DataflowNodes/Private/Dataflow/DataflowGetCollectionBoundsNode.cpp
new file mode 100644
--- /dev/null
+++ b/Engine/Plugins/Experimental/Dataflow/Source/DataflowNodes/Private/Dataflow/DataflowGetCollectionBoundsNode.cpp
@@ -0,0 +1,88 @@
+// Copyright Epic Games, Inc. All Rights Reserved.
+
+#include "Dataflow/DataflowGetCollectionBoundsNode.h"
+
+#include "Dataflow/DataflowNodeFactory.h"
+#include "GeometryCollection/GeometryCollection.h"
+#include "GeometryCollection/Facades/CollectionBoundsFacade.h"
+
+#include UE_INLINE_GENERATED_CPP_BY_NAME(DataflowGetCollectionBoundsNode)
+
+#define LOCTEXT_NAMESPACE "DataflowTextureToAttributeNode"
+
+namespace UE::Dataflow
+{
+	void RegisterGetCollectionBoundsNode()
+	{
+		DATAFLOW_NODE_REGISTER_CREATION_FACTORY(FDataflowGetCollectionBoundsNode);
+	}
+};
+
+
+
+FDataflowGetCollectionBoundsNode::FDataflowGetCollectionBoundsNode(const UE::Dataflow::FNodeParameters& InParam, FGuid InGuid)
+	: FDataflowNode(InParam, InGuid)
+{
+	RegisterInputConnection(&Collection);
+	RegisterOutputConnection(&Collection, &Collection);
+	RegisterOutputConnection(&Bounds);
+}
+
+void FDataflowGetCollectionBoundsNode::Evaluate(UE::Dataflow::FContext& Context, const FDataflowOutput* Out) const
+{
+	if (Out->IsA(&Collection))
+	{
+		SafeForwardInput(Context, &Collection, &Collection);
+		return;
+	}
+
+	if (Out->IsA(&Bounds))
+	{
+		
+
+		FManagedArrayCollection InCollection = GetValue(Context, &Collection, Collection);
+
+		const FName InVertexGroupName = VertexGroup.Name;
+		if (InCollection.HasGroup(InVertexGroupName))
+		{
+			GeometryCollection::Facades::FBoundsFacade BoundsFacade(InCollection);
+			if (BoundsFacade.IsValid())
+			{
+				const FBox OutBounds = BoundsFacade.GetBoundingBoxInCollectionSpace();
+				SetValue(Context, OutBounds, &Bounds);
+				return;
+			}
+
+			// handling of cloth collection 
+			// cannot really include cloth here so we need to directly serach for the vertex position attribute in the render vertices ( eventually we may converge to use only geometry collections )
+			const TManagedArray<FVector3f>* VertexAttribute = InCollection.FindAttribute<FVector3f>("RenderPosition", InVertexGroupName);
+			if (!VertexAttribute)
+			{
+				// cannot really include cloth here so we need to directly serach for the vertex position attribute in the render vertices ( eventually we may converge to use only geometry collections )
+				VertexAttribute = InCollection.FindAttribute<FVector3f>("SimPosition3D", InVertexGroupName);
+			}
+			if (VertexAttribute && VertexAttribute->Num() > 0)
+			{
+				FBox OutBounds;
+				for (int32 Index = 0; Index < VertexAttribute->Num(); ++Index)
+				{
+					OutBounds += FVector((*VertexAttribute)[Index]);
+				}
+				SetValue(Context, OutBounds, &Bounds);
+				return;
+			}
+		
+			Context.Warning(LOCTEXT("NoVertexPosition_Msg", "No valid vertex position attribute could be found in this collection, bounds will be empty"));
+		}
+		else
+		{
+			Context.Error(LOCTEXT("InvalidVertexGroup_Msg", "Vertex group does not exists, bounds will be empty"));
+		}
+
+		const FBox OutBounds(EForceInit::ForceInitToZero);
+		SetValue(Context, OutBounds, &Bounds);
+	}
+}
+
+#undef LOCTEXT_NAMESPACE 
+
diff --git a/Engine/Plugins/Experimental/Dataflow/Source/DataflowNodes/Private/Dataflow/DataflowGetCollectionBoundsNode.h b/Engine/Plugins/Experimental/Dataflow/Source/DataflowNodes/Private/Dataflow/DataflowGetCollectionBoundsNode.h
new file mode 100644
--- /dev/null
+++ b/Engine/Plugins/Experimental/Dataflow/Source/DataflowNodes/Private/Dataflow/DataflowGetCollectionBoundsNode.h
@@ -0,0 +1,43 @@
+// Copyright Epic Games, Inc. All Rights Reserved.
+
+#pragma once 
+
+#include "Dataflow/DataflowNode.h"
+#include "GeometryCollection/ManagedArrayCollection.h"
+#include "Dataflow/DataflowCollectionAddScalarVertexPropertyNode.h"
+
+#include "DataflowGetCollectionBoundsNode.generated.h"
+
+/**
+* Get the bounds of the vertices of a specific collection
+*/
+USTRUCT()
+struct FDataflowGetCollectionBoundsNode : public FDataflowNode
+{
+	GENERATED_USTRUCT_BODY()
+	DATAFLOW_NODE_DEFINE_INTERNAL(FDataflowGetCollectionBoundsNode, "Collection Bounds", "Collection", "BoundingBox Size Limits")
+
+public:
+	FDataflowGetCollectionBoundsNode(const UE::Dataflow::FNodeParameters& InParam, FGuid InGuid = FGuid::NewGuid());
+
+private:
+	/** Collection to compute the bounds from */
+	UPROPERTY(meta = (DataflowInput, DataflowOutput, DataflowPassthrough = Collection))
+	FManagedArrayCollection Collection;
+
+	/** resulting bounds of the collection */
+	UPROPERTY(EditAnywhere, Category = Target, meta = (DataflowOutput))
+	FBox Bounds;
+
+	/** Vertex group to use */
+	UPROPERTY(EditAnywhere, Category = Target)
+	FScalarVertexPropertyGroup VertexGroup;
+
+	virtual void Evaluate(UE::Dataflow::FContext& Context, const FDataflowOutput* Out) const override;
+};
+
+namespace UE::Dataflow
+{
+	void RegisterGetCollectionBoundsNode();
+}
+
diff --git a/Engine/Plugins/Experimental/Dataflow/Source/DataflowNodes/Private/Dataflow/DataflowNodesPlugin.cpp b/Engine/Plugins/Experimental/Dataflow/Source/DataflowNodes/Private/Dataflow/DataflowNodesPlugin.cpp
--- a/Engine/Plugins/Experimental/Dataflow/Source/DataflowNodes/Private/Dataflow/DataflowNodesPlugin.cpp
+++ b/Engine/Plugins/Experimental/Dataflow/Source/DataflowNodes/Private/Dataflow/DataflowNodesPlugin.cpp
@@ -18,6 +18,7 @@
 #include "Dataflow/DataflowVertexColorToAttributeNode.h"
 #include "Dataflow/DataflowVisualizeAttributeNode.h"
 #include "Dataflow/DataflowValueSamplerToAttributeNode.h"
+#include "Dataflow/DataflowGetCollectionBoundsNode.h"
 #include "Dataflow/SamplerNodes/DataflowSamplerTypes.h"
 #include "Dataflow/SamplerNodes/DataflowGradientSamplerNode.h"
 
@@ -57,6 +58,7 @@ void IDataflowNodesPlugin::StartupModule()
 	UE::Dataflow::RegisterStaticMeshNodes();
 	UE::Dataflow::RegisterSelectionNodes();
 	UE::Dataflow::RegisterContextOverridesNodes();
+	UE::Dataflow::RegisterGetCollectionBoundsNode();
 
 	UE::Dataflow::RegisterTextureToAttributeNodes();
 	UE::Dataflow::RegisterVertexColorToAttributeNodes();
