Dataflow : Add linear gradient from box sampler node

#Synced-CL 49661615

diff --git a/Engine/Plugins/Experimental/Dataflow/Source/DataflowNodes/Private/Dataflow/SamplerNodes/DataflowGradientSamplerNode.cpp b/Engine/Plugins/Experimental/Dataflow/Source/DataflowNodes/Private/Dataflow/SamplerNodes/DataflowGradientSamplerNode.cpp
--- a/Engine/Plugins/Experimental/Dataflow/Source/DataflowNodes/Private/Dataflow/SamplerNodes/DataflowGradientSamplerNode.cpp
+++ b/Engine/Plugins/Experimental/Dataflow/Source/DataflowNodes/Private/Dataflow/SamplerNodes/DataflowGradientSamplerNode.cpp
@@ -8,6 +8,7 @@ namespace UE::Dataflow
 	void RegisterGradientValueSamplerNodes()
 	{
 		DATAFLOW_NODE_REGISTER_CREATION_FACTORY(FDataflowLinearGradientValueSamplerNode);
+		DATAFLOW_NODE_REGISTER_CREATION_FACTORY(FDataflowLinearGradientFromBoxValueSamplerNode);
 		DATAFLOW_NODE_REGISTER_CREATION_FACTORY(FDataflowRadialGradientValueSamplerNode);
 	}
 }
@@ -19,7 +20,7 @@ void FDataflowLinearGradientValueSampler::Sample(TArrayView<const FVector3f> Pos
 	const FVector3f GradientDirection = EndPoint - StartPoint;
 	const double SquareLength = GradientDirection.SquaredLength();
 
-	if (SquareLength == UE_SMALL_NUMBER)
+	if (SquareLength < UE_SMALL_NUMBER)
 	{
 		for (float& OutValue : OutValues)
 		{
@@ -69,6 +70,64 @@ void FDataflowLinearGradientValueSamplerNode::Evaluate(UE::Dataflow::FContext& C
 
 ////////////////////////////////////////////////////////////////////////////////////////////////////////////////////////
 
+void FDataflowLinearGradientFromBoxValueSampler::Sample(TArrayView<const FVector3f> Positions, TArrayView<float> OutValues) const
+{
+	const FVector BoxCenter = Box.GetCenter();
+	const FVector BoxHalfExtent = Box.GetExtent();
+
+	FVector GradientHalfExtent{ 0 };
+	switch (Axis)
+	{
+	case EAxis::X: GradientHalfExtent.X = BoxHalfExtent.X; break;
+	case EAxis::Y: GradientHalfExtent.Y = BoxHalfExtent.Y; break;
+	case EAxis::Z: GradientHalfExtent.Z = BoxHalfExtent.Z; break;
+	case EAxis::None: break;
+	}
+
+	GradientHalfExtent = Transform.TransformVector(GradientHalfExtent);
+
+	const FVector StartPoint = Transform.GetTranslation() + BoxCenter - GradientHalfExtent;
+	const FVector EndPoint = Transform.GetTranslation() + BoxCenter + GradientHalfExtent;
+
+	FDataflowLinearGradientValueSampler RegularGradientSampler;
+	RegularGradientSampler.StartPoint = FVector3f(StartPoint);
+	RegularGradientSampler.StartValue = StartValue;
+	RegularGradientSampler.EndPoint = FVector3f(EndPoint);
+	RegularGradientSampler.EndValue= EndValue;
+	RegularGradientSampler.bClamp = bClamp;
+
+	RegularGradientSampler.Sample(Positions, OutValues);
+}
+
+FDataflowLinearGradientFromBoxValueSamplerNode::FDataflowLinearGradientFromBoxValueSamplerNode(const UE::Dataflow::FNodeParameters& InParam, FGuid InGuid)
+	: FDataflowNode(InParam, InGuid)
+{
+	RegisterInputConnection(&Gradient.Box, GET_MEMBER_NAME_CHECKED(FDataflowLinearGradientFromBoxValueSampler, Box));
+	RegisterInputConnection(&Gradient.Axis, GET_MEMBER_NAME_CHECKED(FDataflowLinearGradientFromBoxValueSampler, Axis));
+	RegisterInputConnection(&Gradient.Transform, GET_MEMBER_NAME_CHECKED(FDataflowLinearGradientFromBoxValueSampler, Transform));
+	RegisterInputConnection(&Gradient.StartValue, GET_MEMBER_NAME_CHECKED(FDataflowLinearGradientFromBoxValueSampler, StartValue));
+	RegisterInputConnection(&Gradient.EndValue, GET_MEMBER_NAME_CHECKED(FDataflowLinearGradientFromBoxValueSampler, EndValue));
+	RegisterInputConnection(&Gradient.bClamp, GET_MEMBER_NAME_CHECKED(FDataflowLinearGradientFromBoxValueSampler, bClamp));
+	RegisterOutputConnection(&Gradient);
+}
+
+void FDataflowLinearGradientFromBoxValueSamplerNode::Evaluate(UE::Dataflow::FContext& Context, const FDataflowOutput* Out) const
+{
+	if (Out->IsA(&Gradient))
+	{
+		FDataflowLinearGradientFromBoxValueSampler OutGradient{ Gradient };
+		OutGradient.Box = GetValue(Context, &Gradient.Box, Gradient.Box);
+		OutGradient.Axis = GetValue(Context, &Gradient.Axis, Gradient.Axis);
+		OutGradient.Transform = GetValue(Context, &Gradient.Transform, Gradient.Transform);
+		OutGradient.StartValue= GetValue(Context, &Gradient.StartValue, Gradient.StartValue);
+		OutGradient.EndValue = GetValue(Context, &Gradient.EndValue, Gradient.EndValue);
+		OutGradient.bClamp = GetValue(Context, &Gradient.bClamp, Gradient.bClamp);
+		SetValue(Context, OutGradient, &Gradient);
+	}
+}
+
+////////////////////////////////////////////////////////////////////////////////////////////////////////////////////////
+
 void FDataflowRadialGradientValueSampler::Sample(TArrayView<const FVector3f> Positions, TArrayView<float> OutValues) const
 {
 	if (Radius == UE_SMALL_NUMBER)
diff --git a/Engine/Plugins/Experimental/Dataflow/Source/DataflowNodes/Private/Dataflow/SamplerNodes/DataflowGradientSamplerNode.h b/Engine/Plugins/Experimental/Dataflow/Source/DataflowNodes/Private/Dataflow/SamplerNodes/DataflowGradientSamplerNode.h
--- a/Engine/Plugins/Experimental/Dataflow/Source/DataflowNodes/Private/Dataflow/SamplerNodes/DataflowGradientSamplerNode.h
+++ b/Engine/Plugins/Experimental/Dataflow/Source/DataflowNodes/Private/Dataflow/SamplerNodes/DataflowGradientSamplerNode.h
@@ -60,6 +60,61 @@ private:
 
 ////////////////////////////////////////////////////////////////////////////////////////////////////////////////////////
 
+USTRUCT()
+struct FDataflowLinearGradientFromBoxValueSampler : public FDataflowValueSampler
+{
+	GENERATED_BODY();
+
+	virtual ~FDataflowLinearGradientFromBoxValueSampler() override = default;
+	virtual void Sample(TArrayView<const FVector3f> Positions, TArrayView<float> OutValues) const override;
+
+public:
+	// Box to use for the gradient - one axis is chosen to be the axis of the gradient
+	UPROPERTY(EditAnywhere, Category = Gradient, meta = (DataflowInput))
+	FBox Box;
+
+	// Transform to apply to the box
+	UPROPERTY(EditAnywhere, Category = Gradient, meta = (DataflowInput))
+	FTransform Transform = FTransform::Identity;
+
+	// Axis of the box to use for the gradient
+	UPROPERTY(EditAnywhere, Category = Gradient, meta = (DataflowInput))
+	TEnumAsByte<EAxis::Type> Axis = EAxis::Z;
+
+	// Value defined at min point along the axis
+	UPROPERTY(EditAnywhere, Category = Gradient, meta = (DataflowInput))
+	float StartValue = 0.0;
+
+	// Value defined at max point along the axis
+	UPROPERTY(EditAnywhere, Category = Gradient, meta = (DataflowInput))
+	float EndValue = 1.0;
+
+	// If true generated value will be clamped beyond end and start point, otherwise, they will be extrapolated
+	UPROPERTY(EditAnywhere, Category = Gradient, meta = (DataflowInput))
+	bool bClamp = true;
+};
+
+/**
+ * Output a linear gradient sampler
+ */
+USTRUCT()
+struct FDataflowLinearGradientFromBoxValueSamplerNode : public FDataflowNode
+{
+	GENERATED_BODY()
+	DATAFLOW_NODE_DEFINE_INTERNAL(FDataflowLinearGradientFromBoxValueSamplerNode, "Linear Gradient From Box Sampler", "Samplers", "Bounds Bounding Box Axis")
+
+public:
+	FDataflowLinearGradientFromBoxValueSamplerNode(const UE::Dataflow::FNodeParameters& InParam, FGuid InGuid = FGuid::NewGuid());
+
+private:
+	UPROPERTY(EditAnywhere, Category = Sampler, meta = (DataflowOutput, ShowOnlyInnerProperties))
+	FDataflowLinearGradientFromBoxValueSampler Gradient;
+
+	void Evaluate(UE::Dataflow::FContext& Context, const FDataflowOutput* Out) const override;
+};
+
+////////////////////////////////////////////////////////////////////////////////////////////////////////////////////////
+
 USTRUCT()
 struct FDataflowRadialGradientValueSampler : public FDataflowValueSampler
 {
