Dataflow : add maximum number of curve vertices to be changed through the rendering settings

#Synced-CL 49754362

diff --git a/Engine/Plugins/Experimental/Dataflow/Source/DataflowEditor/Private/DataflowRendering/DataflowGeometryCollectionRenderableType.cpp b/Engine/Plugins/Experimental/Dataflow/Source/DataflowEditor/Private/DataflowRendering/DataflowGeometryCollectionRenderableType.cpp
--- a/Engine/Plugins/Experimental/Dataflow/Source/DataflowEditor/Private/DataflowRendering/DataflowGeometryCollectionRenderableType.cpp
+++ b/Engine/Plugins/Experimental/Dataflow/Source/DataflowEditor/Private/DataflowRendering/DataflowGeometryCollectionRenderableType.cpp
@@ -93,6 +93,9 @@ const int32 UDataflowGeometryCollectionSurfaceRenderSettings::GetHashFromPropert
 	{
 		TypeHash = HashCombine(TypeHash, GetTypeHash(Curve));
 	}
+
+	TypeHash = HashCombine(TypeHash, GetTypeHash(MaxCurveVertices));
+
 	return TypeHash;
 }
 
@@ -100,9 +103,6 @@ const int32 UDataflowGeometryCollectionSurfaceRenderSettings::GetHashFromPropert
 
 namespace UE::Dataflow::Private
 {
-	int32 CurveSurfaceVertexThreshold = 50000;
-	FAutoConsoleVariableRef CVarCurveSurfaceVertexThreshold(TEXT("p.Dataflow.Curve.Surface.VertexThreshold"), CurveSurfaceVertexThreshold, TEXT("Vertex limit used to reduce curve based rendering in the contruction viewport for performance control"));
-
 	////////////////////////////////////////////////////////////////////////////////////////////////////////////////////////////
 
 	class FGeometryCollectionSurfaceRenderableType : public IRenderableType
@@ -161,7 +161,7 @@ namespace UE::Dataflow::Private
 					// TODO(dataflow) : need a version of ConvertGeometryCollectionToDynamicMesh that only takes a managed array collection
 					if (TUniquePtr<FGeometryCollection> GeomCollection = TUniquePtr<FGeometryCollection>(Collection.NewCopy<FGeometryCollection>()))
 					{
-						CullCurvesIfNeeded(*GeomCollection);
+						CullCurvesIfNeeded(*GeomCollection, Settings? Settings->MaxCurveVertices: -1);
 
 						// prepare the dynamic meshes
 						DynamicMeshes.SetNum(TransformAttribute->Num());
@@ -349,48 +349,50 @@ namespace UE::Dataflow::Private
 			}
 		}
 
-		void CullCurvesIfNeeded(FGeometryCollection& GeomCollection) const
+		void CullCurvesIfNeeded(FGeometryCollection& GeomCollection, int32 MaxCurveVerticesFromSettings) const
 		{
 			// Performance control for curve based geometry that can ghave a lot of geometry
-					// if the threshold is reached, we only render a certain percentage of the curves
+			// if the threshold is reached, we only render a certain percentage of the curves
 			const int32 NumVertices = GeomCollection.NumElements(FGeometryCollection::VerticesGroup);
-			if (CurveSurfaceVertexThreshold > 0 && NumVertices > CurveSurfaceVertexThreshold)
+
+			float CullingRatio = 0.0f;
+			if (MaxCurveVerticesFromSettings > 0 && NumVertices > MaxCurveVerticesFromSettings)
+			{
+				CullingRatio = 1.0 - FMath::Clamp((float)MaxCurveVerticesFromSettings / (float)NumVertices, 0, 1);
+			}
+
+			GeometryCollection::Facades::FCollectionCurveGeometryFacade CurveFacade(GeomCollection);
+			const bool bHasCurves = CurveFacade.IsValid();
+			if (bHasCurves)
 			{
-				const float CullingRatio = 1.0 - FMath::Clamp((float)CurveSurfaceVertexThreshold / (float)NumVertices, 0, 1);
+				TManagedArray<bool>& FaceVisible = GeomCollection.AddAttribute<bool>(FGeometryCollection::FaceVisibleAttribute, FGeometryCollection::FacesGroup);
 
-				GeometryCollection::Facades::FCollectionCurveGeometryFacade CurveFacade(GeomCollection);
-				const bool bHasCurves = CurveFacade.IsValid();
-				if (bHasCurves)
-				{
-					TManagedArray<bool>& FaceVisible = GeomCollection.AddAttribute<bool>(FGeometryCollection::FaceVisibleAttribute, FGeometryCollection::FacesGroup);
+				FRandomStream Random(0);
 
-					FRandomStream Random(0);
+				// assumption that vertex index are Pointindex * 2
+				const TArray<int32>& CurvePointOffsets = CurveFacade.GetCurvePointOffsets();
+				int32 StartFace = 0;
+				const int32 NumCurves = CurvePointOffsets.Num();
+				for (int32 CurveIndex = 0; CurveIndex < NumCurves; ++CurveIndex)
+				{
+					const float CurveRand = Random.GetFraction();
+					const bool bCurveVisible = (CurveRand >= CullingRatio);
 
-					// assumption that vertex index are Pointindex * 2
-					const TArray<int32>& CurvePointOffsets = CurveFacade.GetCurvePointOffsets();
-					int32 StartFace = 0;
-					const int32 NumCurves = CurvePointOffsets.Num();
-					for (int32 CurveIndex = 0; CurveIndex < NumCurves; ++CurveIndex)
+					const int32 CurvPointIndex = CurvePointOffsets[CurveIndex];
+					const int32 PrevPointIndex = (CurveIndex > 0) ? CurvePointOffsets[CurveIndex - 1] : 0;
+					const int32 CurveLength = FMath::Max(0, CurvPointIndex - PrevPointIndex);
+					if (CurveLength > 0)
 					{
-						const float CurveRand = Random.GetFraction();
-						const bool bCurveVisible = (CurveRand > CullingRatio);
-
-						const int32 CurvPointIndex = CurvePointOffsets[CurveIndex];
-						const int32 PrevPointIndex = (CurveIndex > 0) ? CurvePointOffsets[CurveIndex - 1] : 0;
-						const int32 CurveLength = FMath::Max(0, CurvPointIndex - PrevPointIndex);
-						if (CurveLength > 0)
+						const int32 CurveNumFaces = (CurveLength - 1) * 2;
+						for (int32 FaceIndex = 0; FaceIndex < CurveNumFaces; ++FaceIndex)
 						{
-							const int32 CurveNumFaces = (CurveLength - 1) * 2;
-							for (int32 FaceIndex = 0; FaceIndex < CurveNumFaces; ++FaceIndex)
+							const int32 GlobalFaceIndex = StartFace + FaceIndex;
+							if (FaceVisible.IsValidIndex(GlobalFaceIndex))
 							{
-								const int32 GlobalFaceIndex = StartFace + FaceIndex;
-								if (FaceVisible.IsValidIndex(GlobalFaceIndex))
-								{
-									FaceVisible[GlobalFaceIndex] = bCurveVisible;
-								}
+								FaceVisible[GlobalFaceIndex] = bCurveVisible;
 							}
-							StartFace += CurveNumFaces;
 						}
+						StartFace += CurveNumFaces;
 					}
 				}
 			}
diff --git a/Engine/Plugins/Experimental/Dataflow/Source/DataflowEditor/Private/DataflowRendering/DataflowGeometryCollectionRenderableType.h b/Engine/Plugins/Experimental/Dataflow/Source/DataflowEditor/Private/DataflowRendering/DataflowGeometryCollectionRenderableType.h
--- a/Engine/Plugins/Experimental/Dataflow/Source/DataflowEditor/Private/DataflowRendering/DataflowGeometryCollectionRenderableType.h
+++ b/Engine/Plugins/Experimental/Dataflow/Source/DataflowEditor/Private/DataflowRendering/DataflowGeometryCollectionRenderableType.h
@@ -10,7 +10,7 @@
 
 class UMaterialInterface;
 
-UCLASS(MinimalAPI, AutoExpandCategories = "Surface|Selection, Surface|Attribute")
+UCLASS(MinimalAPI, AutoExpandCategories = "Surface|Selection, Surface|Attribute, Surface|Curves")
 class UDataflowGeometryCollectionSurfaceRenderSettings : public UDataflowRenderableTypeSettings
 {
 	GENERATED_BODY()
@@ -46,6 +46,11 @@ public:
 	UPROPERTY(EditAnywhere, Category = "Surface|Attribute", meta = (EditCondition = "bShowVertexAttributeAsVertexColor && bUseColorRamp"))
 	FLinearColorRamp ColorRamp;
 
+	/** Maxmimum number of vertixes to render for curves (used for groom) , negative number to not cull any curves */
+	UPROPERTY(EditAnywhere, Category = "Surface|Curves")
+	int32 MaxCurveVertices = 50000;
+
+
 	const int32 GetHashFromProperties() const ;
 };
 
