Dataflow : add vertex selection when rendering geometry collection

#Synced-CL 49754362

diff --git a/Engine/Plugins/Experimental/Dataflow/Source/DataflowEditor/Private/DataflowRendering/DataflowGeometryCollectionRenderableType.cpp b/Engine/Plugins/Experimental/Dataflow/Source/DataflowEditor/Private/DataflowRendering/DataflowGeometryCollectionRenderableType.cpp
--- a/Engine/Plugins/Experimental/Dataflow/Source/DataflowEditor/Private/DataflowRendering/DataflowGeometryCollectionRenderableType.cpp
+++ b/Engine/Plugins/Experimental/Dataflow/Source/DataflowEditor/Private/DataflowRendering/DataflowGeometryCollectionRenderableType.cpp
@@ -3,6 +3,8 @@
 #include "DataflowRendering/DataflowGeometryCollectionRenderableType.h"
 
 #include "Components/DynamicMeshComponent.h"
+#include "Drawing/PointSetComponent.h"
+
 #include "Dataflow/DataflowRenderingViewMode.h"
 #include "DataflowRendering/DataflowRenderableComponents.h"
 #include "DataflowRendering/DataflowRenderableType.h"
@@ -16,6 +18,7 @@
 #include "GeometryCollection/ManagedArrayCollection.h"
 
 #include "GeometryCollection/Facades/CollectionCurveFacade.h"
+#include "GeometryCollection/Facades/CollectionTransformFacade.h"
 #include "GeometryCollection/Facades/CollectionTransformSelectionFacade.h"
 
 #include "UObject/ObjectPtr.h"
@@ -201,12 +204,14 @@ namespace UE::Dataflow::Private
 					Instance.CacheValue(CurrentSettingsHash, SettingsHashCacheName);
 				}
 
-				FDataflowSelection Selection;
+				FDataflowTransformSelection TransformSelection;
+				FDataflowVertexSelection VertexSelection;
 				if (bShowSelection)
 				{
-					GetSelection(Instance, Collection, Selection);
+					GetSelection(Instance, Collection, TransformSelection, VertexSelection);
 				}
-				const bool bIsValidTransformSelection = (Selection.GetGroupName() == FTransformCollection::TransformGroup && Selection.IsValidForCollection(Collection));
+				const bool bIsValidTransformSelection = (TransformSelection.GetGroupName() == FTransformCollection::TransformGroup && TransformSelection.IsValidForCollection(Collection));
+				const bool bIsValidVertexSelection = (VertexSelection.GetGroupName() == FGeometryCollection::VerticesGroup && VertexSelection.IsValidForCollection(Collection));
 
 				UMaterialInterface* SelectionMaterial = Settings
 					? Settings->SelectionMaterial.Get()
@@ -225,7 +230,7 @@ namespace UE::Dataflow::Private
 							Component->SetCastShadow(false);
 							Component->SetMesh(MoveTemp(DynamicMeshes[TransformIndex]));
 
-							if (bIsValidTransformSelection && Selection.IsSelected(TransformIndex))
+							if (bIsValidTransformSelection && TransformSelection.IsSelected(TransformIndex))
 							{
 								Component->SetOverrideRenderMaterial(SelectionMaterial);
 							}
@@ -236,22 +241,67 @@ namespace UE::Dataflow::Private
 						}
 					}
 				}
+
+				// Create point component for vertex selection
+				if (bIsValidVertexSelection)
+				{
+					const TManagedArray<FVector3f>* VertexAttribute = Collection.FindAttribute<FVector3f>(FGeometryCollection::VertexPositionAttribute, FGeometryCollection::VerticesGroup);
+					const TManagedArray<int32>* VertexToTransformIndex = Collection.FindAttribute<int32>(FGeometryCollection::VertexBoneMapAttribute, FGeometryCollection::VerticesGroup);
+					if (VertexAttribute)
+					{
+						const GeometryCollection::Facades::FCollectionTransformFacade TransformFacade(Collection);
+						const TArray<FTransform> CollectionSpaceTransforms = TransformFacade.ComputeCollectionSpaceTransforms();
+
+						const int32 NumVertices = VertexSelection.Num();
+						TArray<FVector> Points;
+						Points.Reserve(NumVertices);
+
+						for (int32 Index = 0; Index < NumVertices; ++Index)
+						{
+							if (VertexSelection.IsSelected(Index))
+							{
+								const int32 TransformIndex = VertexToTransformIndex ? (*VertexToTransformIndex)[Index] : INDEX_NONE;
+								const FTransform Transform = CollectionSpaceTransforms.IsValidIndex(TransformIndex) ? CollectionSpaceTransforms[TransformIndex] : FTransform::Identity;
+								const FVector LocalPoint = FVector((*VertexAttribute)[Index]);
+								Points.Add(Transform.TransformPosition(LocalPoint));
+							}
+						}
+
+						static const FName PointComponentName = TEXT("VertexSelection");
+						if (UPointSetComponent* PointComponent = OutComponents.AddNewComponent<UPointSetComponent>(PointComponentName))
+						{
+							constexpr float PointSize = 8.f;
+							PointComponent->ReservePoints(Points.Num());
+							PointComponent->AddPoints(Points, FColor::Orange, PointSize);
+							UMaterialInterface* Material = LoadObject<UMaterialInterface>(nullptr, TEXT("/MeshModelingToolsetExp/Materials/PointSetOverlaidComponentMaterial_Round"));
+							PointComponent->SetMaterial(0, Material);
+
+						}
+					}
+				}
 			}
 		}
 
-		void GetSelection(const FRenderableTypeInstance& Instance, const FManagedArrayCollection& Collection, FDataflowSelection& OutSelection) const
+		void GetSelection(const FRenderableTypeInstance& Instance, const FManagedArrayCollection& Collection, FDataflowTransformSelection& OutTransformSelection, FDataflowVertexSelection& OutVertexSelection) const
 		{
-			Instance.GetSelectionToVisualize(OutSelection);
-			if (!OutSelection.IsValidForCollection(Collection))
+			OutTransformSelection.Clear();
+			OutVertexSelection.Clear();
+
+			FDataflowSelection Selection;
+			Instance.GetSelectionToVisualize(Selection);
+			if (!Selection.IsValidForCollection(Collection))
 			{
 				FDataflowSelection DefaultValue;
-				OutSelection = Instance.GetOutputValueByType<FDataflowSelectionTypes>(DefaultValue);
+				Selection = Instance.GetOutputValueByType<FDataflowSelectionTypes>(DefaultValue);
+			}
+			if (Selection.GetGroupName() == FTransformCollection::TransformGroup && Selection.AnySelected())
+			{
+				OutTransformSelection.Initialize(Selection);
+				GeometryCollectionSurfaceRenderableType::Private::AddLeavesToSelection(Collection, Selection, OutTransformSelection);
 			}
-			if (OutSelection.GetGroupName() == FTransformCollection::TransformGroup && OutSelection.AnySelected())
+			else if (Selection.GetGroupName() == FGeometryCollection::VerticesGroup && Selection.AnySelected())
 			{
-				const FDataflowSelection TransformSelection = OutSelection;
-				OutSelection.Clear();
-				GeometryCollectionSurfaceRenderableType::Private::AddLeavesToSelection(Collection, TransformSelection, OutSelection);
+				OutVertexSelection.Initialize(Selection);
 			}
 		}
 
