Dataflow editor : add a way to change the type of an anytype output pin

#Synced-CL 49797956

diff --git a/Engine/Plugins/Experimental/Dataflow/Source/DataflowEditor/Private/Dataflow/DataflowSchema.cpp b/Engine/Plugins/Experimental/Dataflow/Source/DataflowEditor/Private/Dataflow/DataflowSchema.cpp
--- a/Engine/Plugins/Experimental/Dataflow/Source/DataflowEditor/Private/Dataflow/DataflowSchema.cpp
+++ b/Engine/Plugins/Experimental/Dataflow/Source/DataflowEditor/Private/Dataflow/DataflowSchema.cpp
@@ -110,6 +110,52 @@ void UDataflowSchema::GetContextMenuActions(class UToolMenu* Menu, class UGraphN
 		{
 			if (FDataflowConnection* Connection = UDataflowEdNode::GetConnectionFromPin(Context->Pin))
 			{
+				if (Connection->IsAnyType() && !Connection->IsTypeLocked())
+				{
+					TArray<FName> SupportedTypes;
+					Connection->GetSupportedTypes(SupportedTypes);
+					if (SupportedTypes.Num() > 0)
+					{
+						FToolMenuSection& Section = Menu->AddSection("DataflowSchema_PinContextMenu_SectionTypes", LOCTEXT("DataflowSchema_PinContextMenu_SectionTypes_Text", "Compatible Types"));
+						for (FName SupportedType : SupportedTypes)
+						{
+							Section.AddMenuEntry(
+								SupportedType,
+								FText::Format(LOCTEXT("PinConvertType_Label", "Convert to {0}"), FText::FromName(SupportedType)),
+								FText::Format(LOCTEXT("PinConvertType_Tooltip", "Convert the output type {0}"), FText::FromName(SupportedType)),
+								FSlateIcon(),
+								FUIAction(
+									// Fine to capture raw pointers because the lifetime of teh menu is shorter than the connection and node lifetime
+									FExecuteAction::CreateLambda([SupportedType, Connection, EdNode]()
+										{
+											if (FDataflowNode* Node = Connection->GetOwningNode())
+											{
+												Node->SetOutputConnectionConcreteType(Connection, SupportedType);
+												Node->NotifyConnectionTypeChanged(Connection);
+
+												// teh context only gives a const node, but we need to update the pins 
+												if (UDataflowEdNode* MutableEdNode = const_cast<UDataflowEdNode*>(EdNode))
+												{
+													MutableEdNode->UpdatePinsFromDataflowNode();
+												}
+											}
+										}),
+									// Fine to capture raw pointers because the lifetime of teh menu is shorter than the connection and node lifetime
+									FCanExecuteAction::CreateLambda([SupportedType, Connection]()
+										{
+											return (Connection 
+												&& !Connection->IsTypeLocked() 
+												&& !Connection->IsConnected() 
+												&& Connection->GetType() != SupportedType
+												&& !Connection->IsAnytypeDependencyConnected()
+												);
+										})
+								)
+							);
+						}
+					}
+				}
+
 				FToolMenuSection& Section = Menu->AddSection("DataflowSchema_PinContextMenu_SectionDebug", LOCTEXT("DataflowSchema_PinContextMenu_SectionDebug_Text", "Debug"));
 				const bool bConnectionWatched = EdNode->IsConnectionWatched(*Connection);
 				Section.AddMenuEntry(bConnectionWatched? FGraphEditorCommands::Get().StopWatchingPin: FGraphEditorCommands::Get().StartWatchingPin);
diff --git a/Engine/Source/Runtime/Experimental/Dataflow/Core/Private/Dataflow/DataflowConnection.cpp b/Engine/Source/Runtime/Experimental/Dataflow/Core/Private/Dataflow/DataflowConnection.cpp
--- a/Engine/Source/Runtime/Experimental/Dataflow/Core/Private/Dataflow/DataflowConnection.cpp
+++ b/Engine/Source/Runtime/Experimental/Dataflow/Core/Private/Dataflow/DataflowConnection.cpp
@@ -129,6 +129,14 @@ bool FDataflowConnection::SupportsType(FName InType) const
 	return (InType == GetType());
 }
 
+void FDataflowConnection::GetSupportedTypes(TArray<FName>& OutSupportedTypes) const
+{
+	if (bIsAnyType && TypePolicy)
+	{
+		TypePolicy->GetSupportedTypes(OutSupportedTypes);
+	}
+}
+
 void FDataflowConnection::ForceTypeDependencyGroup(FName InTypeDependencyGroup)
 {
 	if (IsAnyType())
diff --git a/Engine/Source/Runtime/Experimental/Dataflow/Core/Private/Dataflow/DataflowNode.cpp b/Engine/Source/Runtime/Experimental/Dataflow/Core/Private/Dataflow/DataflowNode.cpp
--- a/Engine/Source/Runtime/Experimental/Dataflow/Core/Private/Dataflow/DataflowNode.cpp
+++ b/Engine/Source/Runtime/Experimental/Dataflow/Core/Private/Dataflow/DataflowNode.cpp
@@ -1611,6 +1611,15 @@ bool FDataflowNode::SetOutputConcreteType(const UE::Dataflow::FConnectionReferen
 	return SetConnectionConcreteType(FindOutput(OutputReference), NewType);
 }
 
+bool FDataflowNode::SetOutputConnectionConcreteType(FDataflowConnection* Output, FName NewType)
+{
+	if (Output && Output->GetOwningNode() == this && Output->GetDirection() == UE::Dataflow::FPin::EDirection::OUTPUT)
+	{
+		return SetConnectionConcreteType(Output, NewType);
+	}
+	return false;
+}
+
 bool FDataflowNode::SetAllConnectionConcreteType(FName NewType)
 {
 	bool bChanged = false;
diff --git a/Engine/Source/Runtime/Experimental/Dataflow/Core/Public/Dataflow/DataflowConnection.h b/Engine/Source/Runtime/Experimental/Dataflow/Core/Public/Dataflow/DataflowConnection.h
--- a/Engine/Source/Runtime/Experimental/Dataflow/Core/Public/Dataflow/DataflowConnection.h
+++ b/Engine/Source/Runtime/Experimental/Dataflow/Core/Public/Dataflow/DataflowConnection.h
@@ -270,6 +270,13 @@ public:
 	DATAFLOWCORE_API void SetTypePolicy(IDataflowTypePolicy* InTypePolicy);
 	DATAFLOWCORE_API bool SupportsType(FName InType) const;
 
+	/**
+	* return the short list of supported types
+	* if the type of the connection is a broad anytype and that the list of supported type is not finite, then nothing is returned
+	*/
+	DATAFLOWCORE_API void GetSupportedTypes(TArray<FName>& OutSupportedTypes) const;
+
+
 	DATAFLOWCORE_API FDataflowConnection& SetTypeDependencyGroup(FName DependencyGroupName);
 	FName GetTypeDependencyGroup() const { return TypeDependencyGroup;	}
 
diff --git a/Engine/Source/Runtime/Experimental/Dataflow/Core/Public/Dataflow/DataflowNode.h b/Engine/Source/Runtime/Experimental/Dataflow/Core/Public/Dataflow/DataflowNode.h
--- a/Engine/Source/Runtime/Experimental/Dataflow/Core/Public/Dataflow/DataflowNode.h
+++ b/Engine/Source/Runtime/Experimental/Dataflow/Core/Public/Dataflow/DataflowNode.h
@@ -786,6 +786,9 @@ PRAGMA_ENABLE_DEPRECATION_WARNINGS
 	// returns true if the output type was changed successfully
 	DATAFLOWCORE_API bool SetOutputConcreteType(const UE::Dataflow::FConnectionReference& OutputReference, FName NewType);
 
+	// returns true if the output type was changed successfully
+	DATAFLOWCORE_API bool SetOutputConnectionConcreteType(FDataflowConnection* Output, FName NewType);
+
 	/** Check if any connection matching a specific type dependency group are connected */
 	DATAFLOWCORE_API bool IsAnytypeDependencyConnected(FName DependencyGroup, const FDataflowConnection* IgnoreConnection = nullptr) const;
 
diff --git a/Engine/Source/Runtime/Experimental/Dataflow/Core/Public/Dataflow/DataflowTypePolicy.h b/Engine/Source/Runtime/Experimental/Dataflow/Core/Public/Dataflow/DataflowTypePolicy.h
--- a/Engine/Source/Runtime/Experimental/Dataflow/Core/Public/Dataflow/DataflowTypePolicy.h
+++ b/Engine/Source/Runtime/Experimental/Dataflow/Core/Public/Dataflow/DataflowTypePolicy.h
@@ -104,7 +104,12 @@ UE_DATAFLOW_POLICY_DECLARE_TYPENAME(FDataflowCurveSelection)
 
 struct IDataflowTypePolicy
 {
+	// Check if a specific type if supported 
 	virtual bool SupportsType(FName Type) const = 0;
+
+	// Add the supported list of types if there's a short list of them 
+	// for full anytype, no types will be added to the array
+	virtual void GetSupportedTypes(TArray<FName>& OutSupportedTypes) const {};
 };
 
 struct FDataflowAllTypesPolicy : public IDataflowTypePolicy
@@ -141,6 +146,16 @@ struct TDataflowSingleTypePolicy : public IDataflowTypePolicy
 		return (InType == TypeName);
 	}
 
+	virtual void GetSupportedTypes(TArray<FName>& OutSupportedTypes) const override
+	{
+		GetSupportedTypesStatic(OutSupportedTypes);
+	}
+
+	static void GetSupportedTypesStatic(TArray<FName>& OutSupportedTypes)
+	{
+		OutSupportedTypes.Add(TypeName);
+	}
+
 	template <typename TVisitor>
 	static bool VisitPolicyByType(FName RequestedType, TVisitor Visitor)
 	{
@@ -201,6 +216,12 @@ struct TDataflowMultiTypePolicy<T, TTypes...>: public TDataflowMultiTypePolicy<T
 			|| Super::SupportsTypeStatic(InType);
 	}
 
+	virtual void GetSupportedTypes(TArray<FName>& OutSupportedTypes) const override
+	{
+		Super::GetSupportedTypes(OutSupportedTypes);
+		TDataflowSingleTypePolicy<T>::GetSupportedTypesStatic(OutSupportedTypes);
+	}
+
 	template <typename TVisitor>
 	static bool VisitPolicyByType(FName RequestedType, TVisitor Visitor)
 	{
