Dataflow : Texture objects can now render in the construction view

#virtualized
#Synced-CL 49874723

diff --git a/Engine/Plugins/Experimental/Dataflow/Source/DataflowEditor/Private/Dataflow/DataflowEditorModule.cpp b/Engine/Plugins/Experimental/Dataflow/Source/DataflowEditor/Private/Dataflow/DataflowEditorModule.cpp
--- a/Engine/Plugins/Experimental/Dataflow/Source/DataflowEditor/Private/Dataflow/DataflowEditorModule.cpp
+++ b/Engine/Plugins/Experimental/Dataflow/Source/DataflowEditor/Private/Dataflow/DataflowEditorModule.cpp
@@ -45,6 +45,7 @@
 #include "DataflowRendering/DataflowBoxRenderableType.h"
 #include "DataflowRendering/DataflowSphereRenderableType.h"
 #include "DataflowRendering/DataflowMaterialRenderableType.h"
+#include "DataflowRendering/DataflowTextureRenderableType.h"
 
 #include "ContentBrowserModule.h"
 #include "Modules/ModuleManager.h"
@@ -317,6 +318,7 @@ void FDataflowEditorModule::StartupModule()
 	UE::Dataflow::Private::RegisterBoxRenderableTypes();
 	UE::Dataflow::Private::RegisterSphereRenderableTypes();
 	UE::Dataflow::Private::RegisterMaterialRenderableTypes();
+	UE::Dataflow::Private::RegisterTextureRenderableTypes();
 
 	// register an browser extension to show when an asset is bound to a dataflow asset
 	if (FContentBrowserModule* ContentBrowserModule = FModuleManager::Get().GetModulePtr<FContentBrowserModule>(TEXT("ContentBrowser")))
diff --git a/Engine/Plugins/Experimental/Dataflow/Source/DataflowEditor/Private/DataflowRendering/DataflowTextureRenderableType.cpp b/Engine/Plugins/Experimental/Dataflow/Source/DataflowEditor/Private/DataflowRendering/DataflowTextureRenderableType.cpp
new file mode 100644
--- /dev/null
+++ b/Engine/Plugins/Experimental/Dataflow/Source/DataflowEditor/Private/DataflowRendering/DataflowTextureRenderableType.cpp
@@ -0,0 +1,77 @@
+// Copyright Epic Games, Inc. All Rights Reserved.
+
+#include "DataflowRendering/DataflowTextureRenderableType.h"
+
+#include "Dataflow/DataflowRenderingViewMode.h"
+
+#include "DataflowRendering/DataflowRenderableComponents.h"
+#include "DataflowRendering/DataflowRenderableType.h"
+#include "DataflowRendering/DataflowRenderableTypeInstance.h"
+#include "DataflowRendering/DataflowRenderableTypeRegistry.h"
+#include "DataflowRendering/DataflowRenderableTypeUtils.h"
+
+#include "Components/StaticMeshComponent.h"
+#include "Materials/MaterialInterface.h"
+#include "Materials/MaterialInstanceDynamic.h"
+
+#include "Engine/StaticMesh.h"
+#include "Engine/Texture2D.h"
+
+#include "UObject/ObjectPtr.h"
+
+namespace UE::Dataflow::Private
+{
+	///////////////////////////////////////////////////////////////////////////////////////////////////////////////
+
+	static UStaticMesh* LoadTexturePreviewMesh()
+	{
+		return LoadObject<UStaticMesh>(nullptr, TEXT("/Engine/BasicShapes/Plane.Plane"));
+	}
+
+
+	static UMaterialInterface* LoadTexturePreviewMaterial()
+	{
+		return LoadObject<UMaterialInterface>(nullptr, TEXT("/Engine/EditorMaterials/Dataflow/M_Dataflow_TexturePreview.M_Dataflow_TexturePreview"));
+	}
+
+	class FTexture2DRenderableType : public IRenderableType
+	{
+		UE_DATAFLOW_IRENDERABLE_PRIMARY_TYPE(TObjectPtr<UTexture2D>, Texture);
+		UE_DATAFLOW_IRENDERABLE_RENDER_GROUP(Texture);
+		UE_DATAFLOW_IRENDERABLE_VIEW_MODE(FDataflowConstruction2DViewMode);
+
+		virtual void GetPrimitiveComponents(const FRenderableTypeInstance& Instance, FRenderableComponents& OutComponents) const override
+		{
+			const TObjectPtr<UTexture2D> Texture = GetTexture(Instance, {});
+
+			static const FName ComponentName = TEXT("Texture");
+			if (UStaticMeshComponent* Component = OutComponents.AddNewComponent<UStaticMeshComponent>(ComponentName))
+			{
+				if (UMaterialInterface* PreviewMaterial = LoadTexturePreviewMaterial())
+				{
+					if (UMaterialInstanceDynamic* PreviewMaterialInstance = UMaterialInstanceDynamic::Create(PreviewMaterial, Component))
+					{
+						static FName TextureParameterName = TEXT("Texture");
+						PreviewMaterialInstance->SetTextureParameterValue(TextureParameterName, Texture.Get());
+						Component->SetMaterial(0, PreviewMaterialInstance);
+					}
+				}
+
+				Component->SetStaticMesh(LoadTexturePreviewMesh());
+
+				// texture is upside down by default , let's flip it around X
+				Component->SetWorldTransform(FTransform(FRotator(0, 0, 180)));
+			}
+		}
+	};
+
+	///////////////////////////////////////////////////////////////////////////////////////////////////////////////
+
+	void RegisterTextureRenderableTypes()
+	{
+		UE_DATAFLOW_REGISTER_RENDERABLE_TYPE(FTexture2DRenderableType);
+	}
+
+	///////////////////////////////////////////////////////////////////////////////////////////////////////////////
+
+}
\ No newline at end of file
diff --git a/Engine/Plugins/Experimental/Dataflow/Source/DataflowEditor/Private/DataflowRendering/DataflowTextureRenderableType.h b/Engine/Plugins/Experimental/Dataflow/Source/DataflowEditor/Private/DataflowRendering/DataflowTextureRenderableType.h
new file mode 100644
--- /dev/null
+++ b/Engine/Plugins/Experimental/Dataflow/Source/DataflowEditor/Private/DataflowRendering/DataflowTextureRenderableType.h
@@ -0,0 +1,8 @@
+// Copyright Epic Games, Inc. All Rights Reserved.
+
+#pragma once
+
+namespace UE::Dataflow::Private
+{
+	void RegisterTextureRenderableTypes();
+}
