Dataflow:
- Added RenderSettings to BoxArray and SphereArray render types

#virtualized
#Synced-CL 49907950

diff --git a/Engine/Plugins/Experimental/Dataflow/Source/DataflowEditor/Private/DataflowRendering/DataflowBoxRenderableType.cpp b/Engine/Plugins/Experimental/Dataflow/Source/DataflowEditor/Private/DataflowRendering/DataflowBoxRenderableType.cpp
--- a/Engine/Plugins/Experimental/Dataflow/Source/DataflowEditor/Private/DataflowRendering/DataflowBoxRenderableType.cpp
+++ b/Engine/Plugins/Experimental/Dataflow/Source/DataflowEditor/Private/DataflowRendering/DataflowBoxRenderableType.cpp
@@ -1,11 +1,13 @@
 // Copyright Epic Games, Inc. All Rights Reserved.
 
 #include "DataflowRendering/DataflowBoxRenderableType.h"
+#include "DataflowRendering/DataflowRenderableTypeUtils.h"
 
 #include "Components/BoxComponent.h"
 #include "Components/InstancedStaticMeshComponent.h"
 
 #include "Dataflow/DataflowRenderingViewMode.h"
+#include "Dataflow/DataflowEngineUtil.h"
 
 #include "DataflowRendering/DataflowRenderableComponents.h"
 #include "DataflowRendering/DataflowRenderableType.h"
@@ -17,6 +19,8 @@
 
 #include "UObject/ObjectPtr.h"
 
+#include "Materials/MaterialInstanceDynamic.h"
+
 namespace UE::Dataflow::Private
 {
 	///////////////////////////////////////////////////////////////////////////////////////////////////////////////
@@ -33,10 +37,25 @@ namespace UE::Dataflow::Private
 
 			static const FName BoxComponentName = TEXT("Box");
 
-			if (UBoxComponent* Component = OutComponents.AddNewComponent<UBoxComponent>(BoxComponentName))
+			UStaticMesh* BoxMesh = UE::Dataflow::RenderGeometry::GetDataflowBox();
+			if (BoxMesh)
 			{
-				Component->SetWorldLocation(Box.GetCenter());
-				Component->SetBoxExtent(Box.GetExtent());
+				if (UStaticMeshComponent* Component = OutComponents.AddNewComponent<UStaticMeshComponent>(BoxComponentName))
+				{
+					Component->SetStaticMesh(BoxMesh);
+
+					FTransform BoxTransform;
+					BoxTransform.SetTranslation(Box.GetCenter());
+					BoxTransform.SetScale3D(Box.GetExtent() / 50.f);
+
+					Component->SetWorldTransform(BoxTransform);
+
+					UMaterialInterface* Material = UE::Dataflow::RenderMaterial::GetDataflowColorMaterial();
+					if (Material)
+					{
+						Component->SetMaterial(0, Material);
+					}
+				}
 			}
 		}
 	};
@@ -48,52 +67,95 @@ namespace UE::Dataflow::Private
 		UE_DATAFLOW_IRENDERABLE_PRIMARY_TYPE(TArray<FBox>, BoxArray);
 		UE_DATAFLOW_IRENDERABLE_RENDER_GROUP(Surface);
 		UE_DATAFLOW_IRENDERABLE_VIEW_MODE(FDataflowConstruction3DViewMode);
+		UE_DATAFLOW_IRENDERABLE_SETTINGS(UDataflowBoxArrayRenderSettings);
 
 		virtual void GetPrimitiveComponents(const FRenderableTypeInstance& Instance, FRenderableComponents& OutComponents) const override
 		{
-			UStaticMesh* BoxMesh = LoadObject<UStaticMesh>(nullptr, TEXT("/Engine/BasicShapes/Cube.Cube"));
+			UStaticMesh* BoxMesh = UE::Dataflow::RenderGeometry::GetDataflowBox();
 			if (BoxMesh)
 			{
 				const TArray<FBox>& BoxArray = GetBoxArray(Instance, {});
 
-				static const FName BoxArrayComponentName = TEXT("BoxArray");
-
-				const int32 NumBoxes = BoxArray.Num();
-
-				UInstancedStaticMeshComponent* ISMComponent = OutComponents.AddNewComponent<UInstancedStaticMeshComponent>(BoxArrayComponentName);
-				ISMComponent->PreAllocateInstancesMemory(NumBoxes);
-				ISMComponent->SetStaticMesh(BoxMesh);
-
-				//for (int32 Idx = 0; Idx < NumBoxes; ++Idx)
-				//{
-				//	const FBox& Box = BoxArray[Idx];
-
-				//	FTransform BoxTransform;
-				//	BoxTransform.SetTranslation(Box.GetCenter());
-				//	BoxTransform.SetScale3D(Box.GetExtent() / 50.f);
-
-				//	ISMComponent->instance
-				//}
-				for (const FBox& Box : BoxArray)
+				if (BoxArray.Num())
 				{
-					FTransform BoxTransform;
-					BoxTransform.SetTranslation(Box.GetCenter());
-					BoxTransform.SetScale3D(Box.GetExtent() / 50.f);
-
-					ISMComponent->AddInstance(BoxTransform);
+					static const FName BoxArrayComponentName = TEXT("BoxArray");
+
+					const int32 NumBoxes = BoxArray.Num();
+
+					UInstancedStaticMeshComponent* ISMComponent = OutComponents.AddNewComponent<UInstancedStaticMeshComponent>(BoxArrayComponentName);
+					if (ISMComponent)
+					{
+						ISMComponent->PreAllocateInstancesMemory(NumBoxes);
+						ISMComponent->SetStaticMesh(BoxMesh);
+
+						FBox ComponentBoundingBox;
+						for (const FBox& Box : BoxArray)
+						{
+							ComponentBoundingBox += Box;
+
+							FTransform BoxTransform;
+							BoxTransform.SetTranslation(Box.GetCenter());
+							BoxTransform.SetScale3D(Box.GetExtent() / 50.f);
+
+							ISMComponent->AddInstance(BoxTransform);
+						}
+
+						ISMComponent->SetNumCustomDataFloats(3); // RGB
+
+						const UDataflowBoxArrayRenderSettings* Settings = Instance.GetTypedRenderSettings<UDataflowBoxArrayRenderSettings>();
+						if (Settings)
+						{
+							UMaterialInterface* Material = Settings->ColorSettings.bWireframe
+								? UE::Dataflow::RenderMaterial::GetDataflowColorWireframeMaterial()
+								: UE::Dataflow::RenderMaterial::GetDataflowColorMaterial();
+							if (Material)
+							{
+								UMaterialInstanceDynamic* MaterialInstance = UMaterialInstanceDynamic::Create(Material, ISMComponent);
+								if (MaterialInstance)
+								{
+									ISMComponent->SetMaterial(0, MaterialInstance);
+
+									MaterialInstance->SetScalarParameterValue(FName("Transparency"), Settings->ColorSettings.Transparency);
+
+									TArray<float> Colors; // NumInstances * 3
+									Colors.Init(1.0, BoxArray.Num() * 3);
+
+									int32 Idx = 0;
+									for (const FBox& Box : BoxArray)
+									{
+										FLinearColor Color = UE::Dataflow::Rendering::GetColor(Idx,
+											ComponentBoundingBox,
+											Box.GetCenter(),
+											2.0 * Box.GetExtent().GetMax(),
+											Settings->ColorSettings);
+
+										Colors[3 * Idx + 0] = Color.R;
+										Colors[3 * Idx + 1] = Color.G;
+										Colors[3 * Idx + 2] = Color.B;
+
+										Idx++;
+									}
+
+									ISMComponent->SetCustomData(0, BoxArray.Num() - 1, Colors);
+								}
+							}
+						}
+					}
 				}
 			}
 		}
 	};
+}
 
-	///////////////////////////////////////////////////////////////////////////////////////////////////////////////
+///////////////////////////////////////////////////////////////////////////////////////////////////////////////
 
+namespace UE::Dataflow::Private
+{
 	void RegisterBoxRenderableTypes()
 	{
 		UE_DATAFLOW_REGISTER_RENDERABLE_TYPE(FBoxSurfaceRenderableType);
 		UE_DATAFLOW_REGISTER_RENDERABLE_TYPE(FBoxArraySurfaceRenderableType);
 	}
+}
 
-	///////////////////////////////////////////////////////////////////////////////////////////////////////////////
-
-}
\ No newline at end of file
+///////////////////////////////////////////////////////////////////////////////////////////////////////////////
diff --git a/Engine/Plugins/Experimental/Dataflow/Source/DataflowEditor/Private/DataflowRendering/DataflowBoxRenderableType.h b/Engine/Plugins/Experimental/Dataflow/Source/DataflowEditor/Private/DataflowRendering/DataflowBoxRenderableType.h
--- a/Engine/Plugins/Experimental/Dataflow/Source/DataflowEditor/Private/DataflowRendering/DataflowBoxRenderableType.h
+++ b/Engine/Plugins/Experimental/Dataflow/Source/DataflowEditor/Private/DataflowRendering/DataflowBoxRenderableType.h
@@ -2,6 +2,21 @@
 
 #pragma once
 
+#include "DataflowRendering/DataflowRenderableTypeSettings.h"
+#include "DataflowRendering/DataflowRenderableTypeUtils.h"
+
+#include "DataflowBoxRenderableType.generated.h"
+
+UCLASS(MinimalAPI, AutoExpandCategories = "Surface|ConstantColor, Surface|RandomColor, Surface|ColorBySize")
+class UDataflowBoxArrayRenderSettings : public UDataflowRenderableTypeSettings
+{
+	GENERATED_BODY()
+
+public:
+	UPROPERTY(EditAnywhere, Category = "Surface", meta = (ShowOnlyInnerProperties))
+	FDataflowColorCommonRenderSettings ColorSettings;
+};
+
 namespace UE::Dataflow::Private
 {
 	void RegisterBoxRenderableTypes();
diff --git a/Engine/Plugins/Experimental/Dataflow/Source/DataflowEditor/Private/DataflowRendering/DataflowRenderableTypeUtils.cpp b/Engine/Plugins/Experimental/Dataflow/Source/DataflowEditor/Private/DataflowRendering/DataflowRenderableTypeUtils.cpp
--- a/Engine/Plugins/Experimental/Dataflow/Source/DataflowEditor/Private/DataflowRendering/DataflowRenderableTypeUtils.cpp
+++ b/Engine/Plugins/Experimental/Dataflow/Source/DataflowEditor/Private/DataflowRendering/DataflowRenderableTypeUtils.cpp
@@ -3,6 +3,7 @@
 #include "DataflowRendering/DataflowRenderableTypeUtils.h"
 #include "DataflowRendering/DataflowRenderableComponents.h"
 #include "DataflowRendering/DataflowRenderableTypeInstance.h"
+#include "Dataflow/DataflowEngineUtil.h"
 
 #include "Components/DynamicMeshComponent.h"
 #include "Engine/SkeletalMesh.h"
@@ -20,6 +21,16 @@
 
 #include "Dataflow/DataflowEditorStyle.h"
 
+FDataflowColorCommonRenderSettings::FDataflowColorCommonRenderSettings()
+{
+	constexpr bool bOnlyRGB = true;
+	ColorRamp.SetColorAtTime(0.0f, FLinearColor::Red, bOnlyRGB);
+	ColorRamp.SetColorAtTime(0.25f, FLinearColor::Yellow, bOnlyRGB);
+	ColorRamp.SetColorAtTime(0.5f, FLinearColor::Green, bOnlyRGB);
+	ColorRamp.SetColorAtTime(0.75f, FLinearColor::White, bOnlyRGB);
+	ColorRamp.SetColorAtTime(1.0f, FLinearColor::Blue, bOnlyRGB);
+}
+
 namespace UE::Dataflow::Rendering
 {
 	namespace Private
@@ -146,6 +157,7 @@ namespace UE::Dataflow::Rendering
 	}
 
 	///////////////////////////////////////////////////////////////////////////////////////////////////////////////
+	
 	UDynamicMeshComponent* AddDynamicMeshComponent(UE::Geometry::FDynamicMesh3&& DynamicMesh, FRenderableComponents& OutComponents)
 	{
 		if (UDynamicMeshComponent* Component = OutComponents.AddNewComponent<UDynamicMeshComponent>())
@@ -157,9 +169,60 @@ namespace UE::Dataflow::Rendering
 	}
 
 	///////////////////////////////////////////////////////////////////////////////////////////////////////////////
+	
 	int32 GetUVChannelFromInstance(const FRenderableTypeInstance& Instance, int32 DefaultValue)
 	{
 		const static FName DataflowUVChannelName(TEXT("DataflowUVChannel"));
 		return Instance.GetOutputValueByType<int32>(DefaultValue, DataflowUVChannelName);
 	}
+
+	///////////////////////////////////////////////////////////////////////////////////////////////////////////////
+
+	FLinearColor GetColor(const int32 InIdx, const FBox& InComponentBoundingBox, const FVector& InCenter, float InSize,
+		const FDataflowColorCommonRenderSettings& InSettings)
+	{
+		FLinearColor Color = FLinearColor::White;
+
+		if (InSettings.ColorMethod == EDataflowColorMethodType::Constant)
+		{
+			Color = InSettings.Color;
+		}
+		else if (InSettings.ColorMethod == EDataflowColorMethodType::Random)
+		{
+			Color = UE::Dataflow::Color::GetRandomColor(InSettings.RandomSeed, InIdx);
+		}
+		else if (InSettings.ColorMethod == EDataflowColorMethodType::BySize)
+		{
+			if (InSize < InSettings.SizeMin)
+			{
+				InSize = InSettings.SizeMin;
+			}
+			else if (InSize > InSettings.SizeMax)
+			{
+				InSize = InSettings.SizeMax;
+			}
+
+			if (InSettings.SizeMax - InSettings.SizeMin > UE_SMALL_NUMBER)
+			{
+				float NormSize = (InSize - InSettings.SizeMin) / (InSettings.SizeMax - InSettings.SizeMin);
+
+				Color = InSettings.ColorRamp.GetLinearColorValue(NormSize);
+			}
+		}
+		else if (InSettings.ColorMethod == EDataflowColorMethodType::BoundingBox)
+		{
+			if (InComponentBoundingBox.GetExtent().X > UE_SMALL_NUMBER &&
+				InComponentBoundingBox.GetExtent().Y > UE_SMALL_NUMBER &&
+				InComponentBoundingBox.GetExtent().Z > UE_SMALL_NUMBER)
+			{
+				float DeltaX = (InCenter.X - InComponentBoundingBox.Min.X) / (InComponentBoundingBox.Max.X - InComponentBoundingBox.Min.X);
+				float DeltaY = (InCenter.Y - InComponentBoundingBox.Min.Y) / (InComponentBoundingBox.Max.Y - InComponentBoundingBox.Min.Y);
+				float DeltaZ = (InCenter.Z - InComponentBoundingBox.Min.Z) / (InComponentBoundingBox.Max.Z - InComponentBoundingBox.Min.Z);
+
+				Color = FLinearColor(DeltaX, 1.f - DeltaY, DeltaZ);
+			}
+		}
+
+		return Color;
+	}
 }
diff --git a/Engine/Plugins/Experimental/Dataflow/Source/DataflowEditor/Private/DataflowRendering/DataflowSphereRenderableType.cpp b/Engine/Plugins/Experimental/Dataflow/Source/DataflowEditor/Private/DataflowRendering/DataflowSphereRenderableType.cpp
--- a/Engine/Plugins/Experimental/Dataflow/Source/DataflowEditor/Private/DataflowRendering/DataflowSphereRenderableType.cpp
+++ b/Engine/Plugins/Experimental/Dataflow/Source/DataflowEditor/Private/DataflowRendering/DataflowSphereRenderableType.cpp
@@ -2,10 +2,11 @@
 
 #include "DataflowRendering/DataflowSphereRenderableType.h"
 
-#include "Components/SphereComponent.h"
+#include "Components/StaticMeshComponent.h"
 #include "Components/InstancedStaticMeshComponent.h"
 
 #include "Dataflow/DataflowRenderingViewMode.h"
+#include "Dataflow/DataflowEngineUtil.h"
 
 #include "DataflowRendering/DataflowRenderableComponents.h"
 #include "DataflowRendering/DataflowRenderableType.h"
@@ -17,6 +18,8 @@
 
 #include "UObject/ObjectPtr.h"
 
+#include "Materials/MaterialInstanceDynamic.h"
+
 namespace UE::Dataflow::Private
 {
 	///////////////////////////////////////////////////////////////////////////////////////////////////////////////
@@ -33,10 +36,25 @@ namespace UE::Dataflow::Private
 
 			static const FName SphereComponentName = TEXT("Sphere");
 
-			if (USphereComponent* Component = OutComponents.AddNewComponent<USphereComponent>(SphereComponentName))
+			UStaticMesh* SphereMesh = UE::Dataflow::RenderGeometry::GetDataflowSphere();
+			if (SphereMesh)
 			{
-				Component->SetWorldLocation(Sphere.Center);
-				Component->SetSphereRadius(Sphere.W);
+				if (UStaticMeshComponent* Component = OutComponents.AddNewComponent<UStaticMeshComponent>(SphereComponentName))
+				{
+					Component->SetStaticMesh(SphereMesh);
+
+					FTransform SphereTransform;
+					SphereTransform.SetTranslation(Sphere.Center);
+					SphereTransform.SetScale3D(FVector(Sphere.W) / 50.f);
+
+					Component->SetWorldTransform(SphereTransform);
+
+					UMaterialInterface* Material = UE::Dataflow::RenderMaterial::GetDataflowColorMaterial();
+					if (Material)
+					{
+						Component->SetMaterial(0, Material);
+					}
+				}
 			}
 		}
 	};
@@ -48,42 +66,96 @@ namespace UE::Dataflow::Private
 		UE_DATAFLOW_IRENDERABLE_PRIMARY_TYPE(TArray<FSphere>, SphereArray);
 		UE_DATAFLOW_IRENDERABLE_RENDER_GROUP(Surface);
 		UE_DATAFLOW_IRENDERABLE_VIEW_MODE(FDataflowConstruction3DViewMode);
+		UE_DATAFLOW_IRENDERABLE_SETTINGS(UDataflowSphereArrayRenderSettings);
 
 		virtual void GetPrimitiveComponents(const FRenderableTypeInstance& Instance, FRenderableComponents& OutComponents) const override
 		{
-			UStaticMesh* SphereMesh = LoadObject<UStaticMesh>(nullptr, TEXT("/Engine/BasicShapes/Sphere.Sphere"));
+			UStaticMesh* SphereMesh = UE::Dataflow::RenderGeometry::GetDataflowSphere();
 			if (SphereMesh)
 			{
 				const TArray<FSphere>& SphereArray = GetSphereArray(Instance, {});
 
-				static const FName SphereArrayComponentName = TEXT("SphereArray");
-
-				const int32 NumSpheres = SphereArray.Num();
-
-				UInstancedStaticMeshComponent* ISMComponent = OutComponents.AddNewComponent<UInstancedStaticMeshComponent>(SphereArrayComponentName);
-				ISMComponent->PreAllocateInstancesMemory(NumSpheres);
-				ISMComponent->SetStaticMesh(SphereMesh);
-
-				for (const FSphere& Sphere : SphereArray)
+				if (SphereArray.Num())
 				{
-					FTransform SphereTransform;
-					SphereTransform.SetTranslation(Sphere.Center);
-					SphereTransform.SetScale3D(FVector(Sphere.W / 50.f));
-
-					ISMComponent->AddInstance(SphereTransform);
+					static const FName SphereArrayComponentName = TEXT("SphereArray");
+
+					const int32 NumSpheres = SphereArray.Num();
+
+					UInstancedStaticMeshComponent* ISMComponent = OutComponents.AddNewComponent<UInstancedStaticMeshComponent>(SphereArrayComponentName);
+					if (ISMComponent)
+					{
+						ISMComponent->PreAllocateInstancesMemory(NumSpheres);
+						ISMComponent->SetStaticMesh(SphereMesh);
+
+						FBox ComponentBoundingBox;
+						for (const FSphere& Sphere : SphereArray)
+						{
+							ComponentBoundingBox += FBox(Sphere.Center - FVector(Sphere.W), Sphere.Center + FVector(Sphere.W));
+
+							FTransform SphereTransform;
+							SphereTransform.SetTranslation(Sphere.Center);
+							SphereTransform.SetScale3D(FVector(Sphere.W) / 50.f);
+
+							ISMComponent->AddInstance(SphereTransform);
+						}
+
+						ISMComponent->SetNumCustomDataFloats(3); // RGB
+
+						const UDataflowSphereArrayRenderSettings* Settings = Instance.GetTypedRenderSettings<UDataflowSphereArrayRenderSettings>();
+						if (Settings)
+						{
+							UMaterialInterface* Material = Settings->ColorSettings.bWireframe
+								? UE::Dataflow::RenderMaterial::GetDataflowColorWireframeMaterial()
+								: UE::Dataflow::RenderMaterial::GetDataflowColorMaterial();
+							if (Material)
+							{
+								UMaterialInstanceDynamic* MaterialInstance = UMaterialInstanceDynamic::Create(Material, ISMComponent);
+								if (MaterialInstance)
+								{
+									ISMComponent->SetMaterial(0, MaterialInstance);
+
+									MaterialInstance->SetScalarParameterValue(FName("Transparency"), Settings->ColorSettings.Transparency);
+
+									TArray<float> Colors; // NumInstances * 3
+									Colors.Init(1.0, SphereArray.Num() * 3);
+
+									int32 Idx = 0;
+									for (const FSphere& Sphere : SphereArray)
+									{
+										FLinearColor Color = UE::Dataflow::Rendering::GetColor(Idx,
+											ComponentBoundingBox,
+											Sphere.Center,
+											2.0 * Sphere.W,
+											Settings->ColorSettings);
+
+										Colors[3 * Idx + 0] = Color.R;
+										Colors[3 * Idx + 1] = Color.G;
+										Colors[3 * Idx + 2] = Color.B;
+
+										Idx++;
+									}
+
+									ISMComponent->SetCustomData(0, SphereArray.Num() - 1, Colors);
+								}
+							}
+						}
+					}
 				}
 			}
 		}
 	};
+}
 
-	///////////////////////////////////////////////////////////////////////////////////////////////////////////////
+///////////////////////////////////////////////////////////////////////////////////////////////////////////////
 
+namespace UE::Dataflow::Private
+{
 	void RegisterSphereRenderableTypes()
 	{
 		UE_DATAFLOW_REGISTER_RENDERABLE_TYPE(FSphereSurfaceRenderableType);
 		UE_DATAFLOW_REGISTER_RENDERABLE_TYPE(FSphereArraySurfaceRenderableType);
 	}
+}
 
-	///////////////////////////////////////////////////////////////////////////////////////////////////////////////
+///////////////////////////////////////////////////////////////////////////////////////////////////////////////
 
-}
\ No newline at end of file
diff --git a/Engine/Plugins/Experimental/Dataflow/Source/DataflowEditor/Private/DataflowRendering/DataflowSphereRenderableType.h b/Engine/Plugins/Experimental/Dataflow/Source/DataflowEditor/Private/DataflowRendering/DataflowSphereRenderableType.h
--- a/Engine/Plugins/Experimental/Dataflow/Source/DataflowEditor/Private/DataflowRendering/DataflowSphereRenderableType.h
+++ b/Engine/Plugins/Experimental/Dataflow/Source/DataflowEditor/Private/DataflowRendering/DataflowSphereRenderableType.h
@@ -2,6 +2,21 @@
 
 #pragma once
 
+#include "DataflowRendering/DataflowRenderableTypeSettings.h"
+#include "DataflowRendering/DataflowRenderableTypeUtils.h"
+
+#include "DataflowSphereRenderableType.generated.h"
+
+UCLASS(MinimalAPI, AutoExpandCategories = "Surface|ConstantColor, Surface|RandomColor, Surface|ColorBySize")
+class UDataflowSphereArrayRenderSettings : public UDataflowRenderableTypeSettings
+{
+	GENERATED_BODY()
+
+public:
+	UPROPERTY(EditAnywhere, Category = "Surface", meta = (ShowOnlyInnerProperties))
+	FDataflowColorCommonRenderSettings ColorSettings;
+};
+
 namespace UE::Dataflow::Private
 {
 	void RegisterSphereRenderableTypes();
diff --git a/Engine/Plugins/Experimental/Dataflow/Source/DataflowEditor/Public/DataflowRendering/DataflowRenderableTypeUtils.h b/Engine/Plugins/Experimental/Dataflow/Source/DataflowEditor/Public/DataflowRendering/DataflowRenderableTypeUtils.h
--- a/Engine/Plugins/Experimental/Dataflow/Source/DataflowEditor/Public/DataflowRendering/DataflowRenderableTypeUtils.h
+++ b/Engine/Plugins/Experimental/Dataflow/Source/DataflowEditor/Public/DataflowRendering/DataflowRenderableTypeUtils.h
@@ -3,6 +3,9 @@
 #pragma once
 
 #include "Templates/SharedPointer.h"
+#include "Curves/LinearColorRamp.h"
+
+#include "DataflowRenderableTypeUtils.generated.h"
 
 struct FManagedArrayCollection;
 struct FMeshDescription;
@@ -11,9 +14,60 @@ class UDynamicMeshComponent;
 class USkeletalMesh;
 class UStaticMesh;
 class UMaterial;
+class UDataflowBoxArrayRenderSettings;
 
 #define UE_API DATAFLOWEDITOR_API
 
+UENUM(BlueprintType)
+enum class EDataflowColorMethodType : uint8
+{
+	Constant UMETA(DisplayName = "Constant"),
+	Random UMETA(DisplayName = "Random"),
+	BySize UMETA(DisplayName = "By Size"),
+	BoundingBox UMETA(DisplayName = "BoundinBox")
+};
+
+USTRUCT()
+struct FDataflowColorCommonRenderSettings
+{
+	GENERATED_BODY()
+
+public:
+	FDataflowColorCommonRenderSettings();
+
+	/** Display surface as wireframe */
+	UPROPERTY(EditAnywhere, Category = "Color")
+	bool bWireframe = false;
+
+	/** Method for coloring the boxes */
+	UPROPERTY(EditAnywhere, Category = "Color")
+	EDataflowColorMethodType ColorMethod = EDataflowColorMethodType::Constant;
+
+	/** Transparency */
+	UPROPERTY(EditAnywhere, Category = "Color", meta = (UIMin = 0, ClampMin = 0, UIMax = 1, ClampMax = 1))
+	float Transparency = 0.f;
+
+	/** Constant color */
+	UPROPERTY(EditAnywhere, Category = "ConstantColor", meta = (EditCondition = "ColorMethod == EDataflowColorMethodType::Constant"))
+	FLinearColor Color = FLinearColor::White;
+
+	/** Seed for the random color */
+	UPROPERTY(EditAnywhere, Category = "RandomColor", meta = (UIMin = 0, EditCondition = "ColorMethod == EDataflowColorMethodType::Random"))
+	int32 RandomSeed = 0;
+
+	/** Minimum size */
+	UPROPERTY(EditAnywhere, Category = "ColorBySize", meta = (UIMin = 0, ClampMin = 0, EditCondition = "ColorMethod == EDataflowColorMethodType::BySize"))
+	float SizeMin = 1.f;
+
+	/** Maximum size */
+	UPROPERTY(EditAnywhere, Category = "ColorBySize", meta = (UIMin = 0, ClampMin = 0, EditCondition = "ColorMethod == EDataflowColorMethodType::BySize"))
+	float SizeMax = 100.f;
+
+	/** Color ramp for colorin by scale */
+	UPROPERTY(EditAnywhere, Category = "ColorBySize", meta = (EditCondition = "ColorMethod == EDataflowColorMethodType::BySize"))
+	FLinearColorRamp ColorRamp;
+};
+
 namespace UE::Geometry
 {
 	class FDynamicMesh3;
@@ -64,6 +118,9 @@ namespace UE::Dataflow
 		}
 
 		int32 GetUVChannelFromInstance(const FRenderableTypeInstance& Instance, int32 DefaultValue = 0);
+
+		FLinearColor GetColor(const int32 InIdx, const FBox& InComponentBoundingBox, const FVector& InCenter, float InSize,
+			const FDataflowColorCommonRenderSettings& InSettings);
 	}
 }
 
diff --git a/Engine/Source/Runtime/Experimental/Dataflow/Engine/Private/Dataflow/DataflowEngineUtil.cpp b/Engine/Source/Runtime/Experimental/Dataflow/Engine/Private/Dataflow/DataflowEngineUtil.cpp
--- a/Engine/Source/Runtime/Experimental/Dataflow/Engine/Private/Dataflow/DataflowEngineUtil.cpp
+++ b/Engine/Source/Runtime/Experimental/Dataflow/Engine/Private/Dataflow/DataflowEngineUtil.cpp
@@ -198,5 +198,39 @@ namespace UE::Dataflow
 			return NewName;
 		}
 	}
+
+	namespace RenderGeometry
+	{
+		UStaticMesh* GetDataflowBox()
+		{
+			const FStringView Path = TEXT("/Engine/EditorMeshes/Dataflow/SM_Dataflow_Box.SM_Dataflow_Box");
+
+			return LoadObject<UStaticMesh>(nullptr, Path);
+		}
+
+		UStaticMesh* GetDataflowSphere()
+		{
+			const FStringView Path = TEXT("/Engine/EditorMeshes/Dataflow/SM_Dataflow_Sphere.SM_Dataflow_Sphere");
+		
+			return LoadObject<UStaticMesh>(nullptr, Path);
+		}
+	}
+
+	namespace RenderMaterial
+	{
+		UMaterialInterface* GetDataflowColorMaterial()
+		{
+			const FStringView Path = TEXT("/Engine/EditorMaterials/Dataflow/M_Dataflow_Color.M_Dataflow_Color");
+
+			return LoadObject<UMaterialInterface>(nullptr, Path);
+		}
+
+		UMaterialInterface* GetDataflowColorWireframeMaterial()
+		{
+			const FStringView Path = TEXT("/Engine/EditorMaterials/Dataflow/M_Dataflow_Color_Wireframe.M_Dataflow_Color_Wireframe");
+
+			return LoadObject<UMaterialInterface>(nullptr, Path);
+		}
+	}
 }
 
diff --git a/Engine/Source/Runtime/Experimental/Dataflow/Engine/Public/Dataflow/DataflowEngineUtil.h b/Engine/Source/Runtime/Experimental/Dataflow/Engine/Public/Dataflow/DataflowEngineUtil.h
--- a/Engine/Source/Runtime/Experimental/Dataflow/Engine/Public/Dataflow/DataflowEngineUtil.h
+++ b/Engine/Source/Runtime/Experimental/Dataflow/Engine/Public/Dataflow/DataflowEngineUtil.h
@@ -15,6 +15,9 @@
 #include "UObject/Object.h"
 #include "UObject/UnrealType.h"
 #include "UObject/UnrealTypePrivate.h"
+#include "Engine/StaticMesh.h"
+#include "Materials/MaterialInterface.h"
+
 
 #define UE_API DATAFLOWENGINE_API
 
@@ -92,6 +95,18 @@ namespace UE::Dataflow
 		UE_API FPropertyBagPropertyDesc GetPropertyBagPropertyDescFromDataflowConnection(const FDataflowConnection& Connection);
 		UE_API FName MakeUniqueNameForPropertyBag(FName OriginalName, const FInstancedPropertyBag& PropertyBag);
 	}
+
+	namespace RenderGeometry
+	{
+		UE_API UStaticMesh* GetDataflowBox();
+		UE_API UStaticMesh* GetDataflowSphere();
+	}
+
+	namespace RenderMaterial
+	{
+		UE_API UMaterialInterface* GetDataflowColorMaterial();
+		UE_API UMaterialInterface* GetDataflowColorWireframeMaterial();
+	}
 }
 
 #undef UE_API
