Dataflow:
- Added RenderSettings to Box/Sphere render types
- Added new features to RenderSettings

#Synced-CL 49912175

diff --git a/Engine/Plugins/Experimental/Dataflow/Source/DataflowEditor/Private/DataflowRendering/DataflowBoxRenderableType.cpp b/Engine/Plugins/Experimental/Dataflow/Source/DataflowEditor/Private/DataflowRendering/DataflowBoxRenderableType.cpp
--- a/Engine/Plugins/Experimental/Dataflow/Source/DataflowEditor/Private/DataflowRendering/DataflowBoxRenderableType.cpp
+++ b/Engine/Plugins/Experimental/Dataflow/Source/DataflowEditor/Private/DataflowRendering/DataflowBoxRenderableType.cpp
@@ -30,6 +30,7 @@ namespace UE::Dataflow::Private
 		UE_DATAFLOW_IRENDERABLE_PRIMARY_TYPE(FBox, Box);
 		UE_DATAFLOW_IRENDERABLE_RENDER_GROUP(Surface);
 		UE_DATAFLOW_IRENDERABLE_VIEW_MODE(FDataflowConstruction3DViewMode);
+		UE_DATAFLOW_IRENDERABLE_SETTINGS(UDataflowBoxRenderSettings);
 
 		virtual void GetPrimitiveComponents(const FRenderableTypeInstance& Instance, FRenderableComponents& OutComponents) const override
 		{
@@ -50,10 +51,23 @@ namespace UE::Dataflow::Private
 
 					Component->SetWorldTransform(BoxTransform);
 
-					UMaterialInterface* Material = UE::Dataflow::RenderMaterial::GetDataflowColorMaterial();
-					if (Material)
+					const UDataflowBoxRenderSettings* Settings = Instance.GetTypedRenderSettings<UDataflowBoxRenderSettings>();
+					if (Settings)
 					{
-						Component->SetMaterial(0, Material);
+						UMaterialInterface* Material = Settings->ColorSettings.bWireframe
+							? UE::Dataflow::RenderMaterial::GetDataflowColorWireframeMaterial()
+							: UE::Dataflow::RenderMaterial::GetDataflowColorMaterial();
+						if (Material)
+						{
+							UMaterialInstanceDynamic* MaterialInstance = UMaterialInstanceDynamic::Create(Material, Component);
+							if (MaterialInstance)
+							{
+								Component->SetMaterial(0, MaterialInstance);
+
+								MaterialInstance->SetVectorParameterValue(FName("Color"), Settings->ColorSettings.Color);
+								MaterialInstance->SetScalarParameterValue(FName("Transparency"), Settings->ColorSettings.Transparency);
+							}
+						}
 					}
 				}
 			}
diff --git a/Engine/Plugins/Experimental/Dataflow/Source/DataflowEditor/Private/DataflowRendering/DataflowBoxRenderableType.h b/Engine/Plugins/Experimental/Dataflow/Source/DataflowEditor/Private/DataflowRendering/DataflowBoxRenderableType.h
--- a/Engine/Plugins/Experimental/Dataflow/Source/DataflowEditor/Private/DataflowRendering/DataflowBoxRenderableType.h
+++ b/Engine/Plugins/Experimental/Dataflow/Source/DataflowEditor/Private/DataflowRendering/DataflowBoxRenderableType.h
@@ -7,6 +7,16 @@
 
 #include "DataflowBoxRenderableType.generated.h"
 
+UCLASS(MinimalAPI)
+class UDataflowBoxRenderSettings : public UDataflowRenderableTypeSettings
+{
+	GENERATED_BODY()
+
+public:
+	UPROPERTY(EditAnywhere, Category = "Surface", meta = (ShowOnlyInnerProperties))
+	FDataflowSimpleColorCommonRenderSettings ColorSettings;
+};
+
 UCLASS(MinimalAPI, AutoExpandCategories = "Surface|ConstantColor, Surface|RandomColor, Surface|ColorBySize")
 class UDataflowBoxArrayRenderSettings : public UDataflowRenderableTypeSettings
 {
diff --git a/Engine/Plugins/Experimental/Dataflow/Source/DataflowEditor/Private/DataflowRendering/DataflowRenderableTypeUtils.cpp b/Engine/Plugins/Experimental/Dataflow/Source/DataflowEditor/Private/DataflowRendering/DataflowRenderableTypeUtils.cpp
--- a/Engine/Plugins/Experimental/Dataflow/Source/DataflowEditor/Private/DataflowRendering/DataflowRenderableTypeUtils.cpp
+++ b/Engine/Plugins/Experimental/Dataflow/Source/DataflowEditor/Private/DataflowRendering/DataflowRenderableTypeUtils.cpp
@@ -27,7 +27,7 @@ FDataflowColorCommonRenderSettings::FDataflowColorCommonRenderSettings()
 	ColorRamp.SetColorAtTime(0.0f, FLinearColor::Red, bOnlyRGB);
 	ColorRamp.SetColorAtTime(0.25f, FLinearColor::Yellow, bOnlyRGB);
 	ColorRamp.SetColorAtTime(0.5f, FLinearColor::Green, bOnlyRGB);
-	ColorRamp.SetColorAtTime(0.75f, FLinearColor::White, bOnlyRGB);
+	ColorRamp.SetColorAtTime(0.75f, FLinearColor(FColor::Purple), bOnlyRGB);
 	ColorRamp.SetColorAtTime(1.0f, FLinearColor::Blue, bOnlyRGB);
 }
 
@@ -219,7 +219,84 @@ namespace UE::Dataflow::Rendering
 				float DeltaY = (InCenter.Y - InComponentBoundingBox.Min.Y) / (InComponentBoundingBox.Max.Y - InComponentBoundingBox.Min.Y);
 				float DeltaZ = (InCenter.Z - InComponentBoundingBox.Min.Z) / (InComponentBoundingBox.Max.Z - InComponentBoundingBox.Min.Z);
 
-				Color = FLinearColor(DeltaX, 1.f - DeltaY, DeltaZ);
+				switch (InSettings.PermutationIndex)
+				{
+				case 0:
+					Color = FLinearColor(DeltaX, DeltaY, DeltaZ);
+					break;
+				case 1:
+					Color = FLinearColor(DeltaX, 1.f - DeltaY, DeltaZ);
+					break;
+				case 2:
+					Color = FLinearColor(DeltaX, DeltaY, 1.f - DeltaZ);
+					break;
+				case 3:
+					Color = FLinearColor(DeltaX, 1.f - DeltaY, 1.f - DeltaZ);
+					break;
+				case 4:
+					Color = FLinearColor(DeltaX, DeltaZ, DeltaY);
+					break;
+				case 5:
+					Color = FLinearColor(DeltaX, 1.f - DeltaZ, DeltaY);
+					break;
+				case 6:
+					Color = FLinearColor(DeltaX, DeltaZ, 1.f - DeltaY);
+					break;
+				case 7:
+					Color = FLinearColor(DeltaX, 1.f - DeltaZ, 1.f - DeltaY);
+					break;
+				case 8:
+					Color = FLinearColor(DeltaY, DeltaX, DeltaZ);
+					break;
+				case 9:
+					Color = FLinearColor(DeltaY, 1.f - DeltaX, DeltaZ);
+					break;
+				case 10:
+					Color = FLinearColor(DeltaY, DeltaX, 1.f - DeltaZ);
+					break;
+				case 11:
+					Color = FLinearColor(DeltaY, 1.f - DeltaX, 1.f - DeltaZ);
+					break;
+				case 12:
+					Color = FLinearColor(DeltaY, DeltaZ, DeltaX);
+					break;
+				case 13:
+					Color = FLinearColor(DeltaY, 1.f - DeltaZ, DeltaX);
+					break;
+				case 14:
+					Color = FLinearColor(DeltaY, DeltaZ, 1.f - DeltaX);
+					break;
+				case 15:
+					Color = FLinearColor(DeltaY, 1.f - DeltaZ, 1.f - DeltaX);
+					break;
+				case 16:
+					Color = FLinearColor(DeltaZ, DeltaX, DeltaY);
+					break;
+				case 17:
+					Color = FLinearColor(DeltaZ, 1.f - DeltaX, DeltaY);
+					break;
+				case 18:
+					Color = FLinearColor(DeltaZ, DeltaX, 1.f - DeltaY);
+					break;
+				case 19:
+					Color = FLinearColor(DeltaZ, 1.f - DeltaX, 1.f - DeltaY);
+					break;
+				case 20:
+					Color = FLinearColor(DeltaZ, DeltaY, DeltaX);
+					break;
+				case 21:
+					Color = FLinearColor(DeltaZ, 1.f - DeltaY, DeltaX);
+					break;
+				case 22:
+					Color = FLinearColor(DeltaZ, DeltaY, 1.f - DeltaX);
+					break;
+				case 23:
+					Color = FLinearColor(DeltaZ, 1.f - DeltaY, 1.f - DeltaX);
+					break;
+				default:
+					Color = FLinearColor(DeltaX, DeltaY, DeltaZ);
+					break;
+				}
 			}
 		}
 
diff --git a/Engine/Plugins/Experimental/Dataflow/Source/DataflowEditor/Private/DataflowRendering/DataflowSphereRenderableType.cpp b/Engine/Plugins/Experimental/Dataflow/Source/DataflowEditor/Private/DataflowRendering/DataflowSphereRenderableType.cpp
--- a/Engine/Plugins/Experimental/Dataflow/Source/DataflowEditor/Private/DataflowRendering/DataflowSphereRenderableType.cpp
+++ b/Engine/Plugins/Experimental/Dataflow/Source/DataflowEditor/Private/DataflowRendering/DataflowSphereRenderableType.cpp
@@ -29,6 +29,7 @@ namespace UE::Dataflow::Private
 		UE_DATAFLOW_IRENDERABLE_PRIMARY_TYPE(FSphere, Sphere);
 		UE_DATAFLOW_IRENDERABLE_RENDER_GROUP(Surface);
 		UE_DATAFLOW_IRENDERABLE_VIEW_MODE(FDataflowConstruction3DViewMode);
+		UE_DATAFLOW_IRENDERABLE_SETTINGS(UDataflowSphereRenderSettings);
 
 		virtual void GetPrimitiveComponents(const FRenderableTypeInstance& Instance, FRenderableComponents& OutComponents) const override
 		{
@@ -49,10 +50,23 @@ namespace UE::Dataflow::Private
 
 					Component->SetWorldTransform(SphereTransform);
 
-					UMaterialInterface* Material = UE::Dataflow::RenderMaterial::GetDataflowColorMaterial();
-					if (Material)
+					const UDataflowSphereRenderSettings* Settings = Instance.GetTypedRenderSettings<UDataflowSphereRenderSettings>();
+					if (Settings)
 					{
-						Component->SetMaterial(0, Material);
+						UMaterialInterface* Material = Settings->ColorSettings.bWireframe
+							? UE::Dataflow::RenderMaterial::GetDataflowColorWireframeMaterial()
+							: UE::Dataflow::RenderMaterial::GetDataflowColorMaterial();
+						if (Material)
+						{
+							UMaterialInstanceDynamic* MaterialInstance = UMaterialInstanceDynamic::Create(Material, Component);
+							if (MaterialInstance)
+							{
+								Component->SetMaterial(0, MaterialInstance);
+
+								MaterialInstance->SetVectorParameterValue(FName("Color"), Settings->ColorSettings.Color);
+								MaterialInstance->SetScalarParameterValue(FName("Transparency"), Settings->ColorSettings.Transparency);
+							}
+						}
 					}
 				}
 			}
diff --git a/Engine/Plugins/Experimental/Dataflow/Source/DataflowEditor/Private/DataflowRendering/DataflowSphereRenderableType.h b/Engine/Plugins/Experimental/Dataflow/Source/DataflowEditor/Private/DataflowRendering/DataflowSphereRenderableType.h
--- a/Engine/Plugins/Experimental/Dataflow/Source/DataflowEditor/Private/DataflowRendering/DataflowSphereRenderableType.h
+++ b/Engine/Plugins/Experimental/Dataflow/Source/DataflowEditor/Private/DataflowRendering/DataflowSphereRenderableType.h
@@ -7,6 +7,16 @@
 
 #include "DataflowSphereRenderableType.generated.h"
 
+UCLASS(MinimalAPI)
+class UDataflowSphereRenderSettings : public UDataflowRenderableTypeSettings
+{
+	GENERATED_BODY()
+
+public:
+	UPROPERTY(EditAnywhere, Category = "Surface", meta = (ShowOnlyInnerProperties))
+	FDataflowSimpleColorCommonRenderSettings ColorSettings;
+};
+
 UCLASS(MinimalAPI, AutoExpandCategories = "Surface|ConstantColor, Surface|RandomColor, Surface|ColorBySize")
 class UDataflowSphereArrayRenderSettings : public UDataflowRenderableTypeSettings
 {
diff --git a/Engine/Plugins/Experimental/Dataflow/Source/DataflowEditor/Public/DataflowRendering/DataflowRenderableTypeUtils.h b/Engine/Plugins/Experimental/Dataflow/Source/DataflowEditor/Public/DataflowRendering/DataflowRenderableTypeUtils.h
--- a/Engine/Plugins/Experimental/Dataflow/Source/DataflowEditor/Public/DataflowRendering/DataflowRenderableTypeUtils.h
+++ b/Engine/Plugins/Experimental/Dataflow/Source/DataflowEditor/Public/DataflowRendering/DataflowRenderableTypeUtils.h
@@ -27,6 +27,25 @@ enum class EDataflowColorMethodType : uint8
 	BoundingBox UMETA(DisplayName = "BoundinBox")
 };
 
+USTRUCT()
+struct FDataflowSimpleColorCommonRenderSettings
+{
+	GENERATED_BODY()
+
+public:
+	/** Display surface as wireframe */
+	UPROPERTY(EditAnywhere, Category = "Color")
+	bool bWireframe = false;
+
+	/** Transparency */
+	UPROPERTY(EditAnywhere, Category = "Color", meta = (UIMin = 0, ClampMin = 0, UIMax = 1, ClampMax = 1))
+	float Transparency = 0.f;
+
+	/** Constant color */
+	UPROPERTY(EditAnywhere, Category = "ConstantColor")
+	FLinearColor Color = FLinearColor::White;
+};
+
 USTRUCT()
 struct FDataflowColorCommonRenderSettings
 {
@@ -39,14 +58,14 @@ public:
 	UPROPERTY(EditAnywhere, Category = "Color")
 	bool bWireframe = false;
 
-	/** Method for coloring the boxes */
-	UPROPERTY(EditAnywhere, Category = "Color")
-	EDataflowColorMethodType ColorMethod = EDataflowColorMethodType::Constant;
-
 	/** Transparency */
 	UPROPERTY(EditAnywhere, Category = "Color", meta = (UIMin = 0, ClampMin = 0, UIMax = 1, ClampMax = 1))
 	float Transparency = 0.f;
 
+	/** Method for coloring the boxes */
+	UPROPERTY(EditAnywhere, Category = "Color")
+	EDataflowColorMethodType ColorMethod = EDataflowColorMethodType::Constant;
+
 	/** Constant color */
 	UPROPERTY(EditAnywhere, Category = "ConstantColor", meta = (EditCondition = "ColorMethod == EDataflowColorMethodType::Constant"))
 	FLinearColor Color = FLinearColor::White;
@@ -66,6 +85,12 @@ public:
 	/** Color ramp for colorin by scale */
 	UPROPERTY(EditAnywhere, Category = "ColorBySize", meta = (EditCondition = "ColorMethod == EDataflowColorMethodType::BySize"))
 	FLinearColorRamp ColorRamp;
+
+	/** The normalized position of the center of the object in the component bounding box is computed.
+		Then the X or 1-X, Y or 1-Y, Z or 1-Z of the normalized position is set to the (R,G,B) values. 
+		There are 32 different permutation. The index specifies which one gets used. */
+	UPROPERTY(EditAnywhere, Category = "ColorByBoundingBox", meta = (UIMin = 0, UIMax = 23, EditCondition = "ColorMethod == EDataflowColorMethodType::BoundingBox"))
+	int32 PermutationIndex = 0;
 };
 
 namespace UE::Geometry
