Dataflow : fix potential compilation error by adding missing typename keywords

#Synced-CL 49892100

diff --git a/Engine/Plugins/Experimental/Dataflow/Source/DataflowEditor/Public/DataflowRendering/DataflowRenderableTypeInstance.h b/Engine/Plugins/Experimental/Dataflow/Source/DataflowEditor/Public/DataflowRendering/DataflowRenderableTypeInstance.h
--- a/Engine/Plugins/Experimental/Dataflow/Source/DataflowEditor/Public/DataflowRendering/DataflowRenderableTypeInstance.h
+++ b/Engine/Plugins/Experimental/Dataflow/Source/DataflowEditor/Public/DataflowRendering/DataflowRenderableTypeInstance.h
@@ -131,7 +131,7 @@ namespace UE::Dataflow
 
 		/** Find the first output of a specific match type and return its value - anytype version */
 		template <typename T UE_REQUIRES(std::is_base_of_v<FDataflowAnyType, T>)>
-		const T::FStorageType& GetOutputValueByType(const T::FStorageType& Default, FName Metadata = NAME_None) const
+		const typename T::FStorageType& GetOutputValueByType(const typename T::FStorageType& Default, FName Metadata = NAME_None) const
 		{
 			if (Node)
 			{
