Dataflow paint tool : add brush live mirroring
- Add a second brush indicator that mirrors the one controlled by the user ( projected onto the surface )

#Synced-CL 50024027

diff --git a/Engine/Plugins/Experimental/Dataflow/Source/DataflowEditor/Private/DataflowEditorTools/DataflowEditorVertexAttributePaintTool.cpp b/Engine/Plugins/Experimental/Dataflow/Source/DataflowEditor/Private/DataflowEditorTools/DataflowEditorVertexAttributePaintTool.cpp
--- a/Engine/Plugins/Experimental/Dataflow/Source/DataflowEditor/Private/DataflowEditorTools/DataflowEditorVertexAttributePaintTool.cpp
+++ b/Engine/Plugins/Experimental/Dataflow/Source/DataflowEditor/Private/DataflowEditorTools/DataflowEditorVertexAttributePaintTool.cpp
@@ -24,6 +24,7 @@
 #include "DynamicMeshToMeshDescription.h"
 #include "Engine/Engine.h"
 #include "InteractiveToolManager.h"
+#include "InteractiveGizmoManager.h"
 #include "Intersection/IntrLine2Line2.h"
 #include "Intersection/IntrSegment2Segment2.h"
 #include "Framework/Notifications/NotificationManager.h"
@@ -37,11 +38,11 @@
 #include "Selections/MeshConnectedComponents.h"
 #include "Selections/MeshVertexSelection.h"
 #include "Spatial/PointHashGrid3.h"
+#include "Spatial/SparseDynamicOctree3.h"
 #include "ToolSetupUtil.h"
 #include "Util/BufferUtil.h"
 #include "Widgets/Notifications/SNotificationList.h"
-
-
+#include "Generators/RectangleMeshGenerator.h"
 
 #include UE_INLINE_GENERATED_CPP_BY_NAME(DataflowEditorVertexAttributePaintTool)
 
@@ -63,6 +64,8 @@ namespace UE::DataflowEditorVertexAttributePaintTool::CVars
 
 namespace UE::DataflowEditorVertexAttributePaintTool::Private
 {
+	static const FString MirrorBrushGizmoType = TEXT("DataflowEditorVertexAttributePaintToolMirrorGizmo");
+
 	// probably should be something defined for the whole tool framework...
 #if WITH_EDITOR
 	static EAsyncExecution WeightPaintToolAsyncExecTarget = EAsyncExecution::LargeThreadPool;
@@ -664,34 +667,35 @@ void UDataflowEditorVertexAttributePaintTool::Setup()
 	GetMeshParams.bWantMeshTangents = true;
 	FDynamicMesh3 DynamicMeshCopy = UE::ToolTarget::GetDynamicMeshCopy(Target, GetMeshParams);
 	PreviewMesh->UpdatePreview(MoveTemp(DynamicMeshCopy));
-
+	
 	// Assign materials
 	FComponentMaterialSet MaterialSet = UE::ToolTarget::GetMaterialSet(Target);
 	PreviewMesh->SetMaterials(MaterialSet.Materials);
 
 	// make sure the dynamic mesh has all the necessary attributes
 	FDynamicMesh3* Mesh = GetSculptMesh();
+	check(Mesh);
 	Mesh->EnableVertexColors(FVector3f::One());
 	Mesh->Attributes()->EnablePrimaryColors();
 	Mesh->Attributes()->PrimaryColors()->CreatePerVertex(0.f);
-	UE::Geometry::FAxisAlignedBox3d Bounds = Mesh->GetBounds(true);
-
-	TFuture<void> PrecomputeFuture = Async(UE::DataflowEditorVertexAttributePaintTool::Private::WeightPaintToolAsyncExecTarget, [&]()
+	MeshLocalBounds = Mesh->GetBounds(true);
+	
+	TFuture<void> PrecomputeFuture = Async(UE::DataflowEditorVertexAttributePaintTool::Private::WeightPaintToolAsyncExecTarget, [this]()
 		{
 			PrecomputeFilterData();
 		});
 
-	TFuture<void> OctreeFuture = Async(UE::DataflowEditorVertexAttributePaintTool::Private::WeightPaintToolAsyncExecTarget, [&]()
+	TFuture<void> OctreeFuture = Async(UE::DataflowEditorVertexAttributePaintTool::Private::WeightPaintToolAsyncExecTarget, [this, &Mesh]()
 		{
 			// initialize dynamic octree
 			if (Mesh->TriangleCount() > 100000)
 			{
-				Octree.RootDimension = Bounds.MaxDim() / 10.0;
+				Octree.RootDimension = MeshLocalBounds.MaxDim() / 10.0;
 				Octree.SetMaxTreeDepth(4);
 			}
 			else
 			{
-				Octree.RootDimension = Bounds.MaxDim();
+				Octree.RootDimension = MeshLocalBounds.MaxDim();
 				Octree.SetMaxTreeDepth(8);
 			}
 			Octree.Initialize(Mesh);
@@ -709,7 +713,7 @@ void UDataflowEditorVertexAttributePaintTool::Setup()
 	}
 
 	// initialize brush radius range interval, brush properties
-	UMeshSculptToolBase::InitializeBrushSizeRange(Bounds);
+	UMeshSculptToolBase::InitializeBrushSizeRange(MeshLocalBounds);
 
 	// initialize properties
 	ToolProperties = NewObject<UDataflowEditorVertexAttributePaintToolProperties>(this);
@@ -723,6 +727,7 @@ void UDataflowEditorVertexAttributePaintTool::Setup()
 	MeshSelector->SetComponentSelectionMode(ToolProperties->SelectionProperties.ComponentSelectionMode);
 
 	InitializeIndicator();
+	InitializeMirrorIndicator();
 
 	// initialize our properties
 	AddToolPropertySource(UMeshSculptToolBase::BrushProperties);
@@ -794,6 +799,198 @@ void UDataflowEditorVertexAttributePaintTool::Setup()
 
 }
 
+void UDataflowEditorVertexAttributePaintTool::InitializeMirrorIndicator()
+{
+	// register and spawn brush indicator gizmo
+	GetToolManager()->GetPairedGizmoManager()->RegisterGizmoType(UE::DataflowEditorVertexAttributePaintTool::Private::MirrorBrushGizmoType, NewObject<UBrushStampIndicatorBuilder>());
+	MirrorBrushIndicator = GetToolManager()->GetPairedGizmoManager()->CreateGizmo<UBrushStampIndicator>(UE::DataflowEditorVertexAttributePaintTool::Private::MirrorBrushGizmoType, FString(), this);
+	MirrorBrushIndicatorMesh = MakeBrushIndicatorMesh(this, TargetWorld);
+	MirrorBrushIndicator->AttachedComponent = MirrorBrushIndicatorMesh->GetRootComponent();
+	MirrorBrushIndicator->LineThickness = 1.0;
+	MirrorBrushIndicator->bDrawIndicatorLines = true;
+	MirrorBrushIndicator->bDrawRadiusCircle = false;
+	MirrorBrushIndicator->LineColor = FLinearColor(0.9f, 0.4f, 0.4f);
+
+	MirrorPlaneMesh = NewObject<UPreviewMesh>(this);
+	MirrorPlaneMesh->CreateInWorld(TargetWorld, FTransform::Identity);
+	UE::Geometry::FRoundedRectangleMeshGenerator Gen;
+	Gen.Origin = { 0,0,0 };
+	Gen.Normal = { 0,0,1 };
+	Gen.Width = 100.0;
+	Gen.Height = 100.0;
+	Gen.WidthVertexCount = 10;
+	Gen.HeightVertexCount = 10;
+	Gen.bScaleUVByAspectRatio = false;
+	Gen.Radius = 10.0;
+	Gen.AngleSamples = 8;
+	Gen.SharpCorners = UE::Geometry::ERoundedRectangleCorner::None;
+	Gen.Generate();
+	FDynamicMesh3 PlaneMesh(&Gen);
+	MirrorPlaneMesh->UpdatePreview(&PlaneMesh);
+	BrushIndicatorMaterial = ToolSetupUtil::GetDefaultBrushVolumeMaterial(GetToolManager());
+	if (BrushIndicatorMaterial)
+	{
+		MirrorPlaneMesh->SetMaterial(BrushIndicatorMaterial);
+	}
+
+	// make sure raytracing is disabled on the brush indicator
+	Cast<UDynamicMeshComponent>(MirrorPlaneMesh->GetRootComponent())->SetEnableRaytracing(false);
+	MirrorPlaneMesh->SetShadowsEnabled(false);
+
+	UpdateMirrorPlaneMesh();
+}
+
+void UDataflowEditorVertexAttributePaintTool::UpdateMirrorPlaneMesh()
+{
+	if (MirrorPlaneMesh)
+	{
+		FTransform PlaneMeshTransform = FTransform::Identity;
+
+		MirrorPlaneMesh->SetVisible(IsBrushMirroringEnabled() && IsInBrushMode());
+
+		if (IsBrushMirroringEnabled())
+		{
+			switch(GetMirrorAxis())
+			{
+			case EAxis::X:
+				PlaneMeshTransform = FTransform(FRotator(90, 0, 0));
+				break;
+			case EAxis::Y:
+				PlaneMeshTransform = FTransform(FRotator(0, 0, 90));
+				break;
+			case EAxis::Z:
+				PlaneMeshTransform = FTransform::Identity;
+			}
+		}
+
+		if (const FDynamicMesh3* Mesh = GetSculptMesh())
+		{
+			const FVector MeshExtents = MeshLocalBounds.Extents();
+			
+			// scale of plane
+			FVector PlaneScale(1.0);
+			switch (GetMirrorAxis())
+			{
+			case EAxis::X:
+				PlaneScale.X = MeshExtents.Z;
+				PlaneScale.Y = MeshExtents.Y;
+				break;
+			case EAxis::Y:
+				PlaneScale.X = MeshExtents.X;
+				PlaneScale.Y = MeshExtents.Z;
+				break;
+			case EAxis::Z:
+				PlaneScale.X = MeshExtents.X;
+				PlaneScale.Y = MeshExtents.Y;
+			}
+
+			// the plane is originally 100 units and we want 10% over the size of the box
+			constexpr double ScaleAdjustment = 0.02 * 1.10;
+			PlaneMeshTransform.SetScale3D(PlaneScale * ScaleAdjustment);
+
+			// Position of the plane
+			FVector PlanePosition = MeshLocalBounds.Center();
+			if (!IsBrushMirroringInObjectSpace())
+			{
+				switch (GetMirrorAxis())
+				{
+				case EAxis::X:
+					PlanePosition.X = 0;
+					break;
+				case EAxis::Y:
+					PlanePosition.Y = 0;
+					break;
+				case EAxis::Z:
+					PlanePosition.Z = 0;
+				}
+			}
+			PlaneMeshTransform.SetTranslation(PlanePosition);
+		}
+		MirrorPlaneMesh->SetTransform(PlaneMeshTransform);
+	}
+}
+
+UE::Geometry::FFrame3d UDataflowEditorVertexAttributePaintTool::ComputeMirroredFrame(const UE::Geometry::FFrame3d& InLocalSpaceFrame) const
+{
+	UE::Geometry::FFrame3d OutFrame{ InLocalSpaceFrame };
+
+	const bool bIsObjectSpaceMirror = IsBrushMirroringInObjectSpace();
+	const EAxis::Type MirrorAxis = GetMirrorAxis();
+	FVector MirrorCenter = MeshLocalBounds.Center();
+
+	if (!bIsObjectSpaceMirror)
+	{
+		// Transform the frame to world space before mirroring
+		OutFrame.Transform(CurTargetTransform);
+		MirrorCenter = FVector::ZeroVector;
+	}
+
+	// Mirroring is in bounds space 
+	FTransform LocalTransform = OutFrame.ToFTransform();
+	LocalTransform.SetTranslation(LocalTransform.GetTranslation() - MirrorCenter);
+	LocalTransform.Mirror(MirrorAxis, EAxis::None);
+	LocalTransform.SetTranslation(LocalTransform.GetTranslation() + MirrorCenter);
+	OutFrame = UE::Geometry::FFrame3d(LocalTransform);
+	
+	if (!bIsObjectSpaceMirror)
+	{
+		// Transform back the frame to local space after mirroring
+		OutFrame.Transform(CurTargetTransform.InverseUnsafe());
+	}
+
+	return OutFrame;
+}
+
+void UDataflowEditorVertexAttributePaintTool::UpdateMirrorIndicator()
+{
+	if (BrushIndicatorMesh && BrushIndicator)
+	{
+		MirrorBrushIndicator->bVisible = BrushIndicator->bVisible && IsBrushMirroringEnabled();
+		MirrorBrushIndicatorMesh->SetVisible(BrushIndicatorMesh->IsVisible() && IsBrushMirroringEnabled());
+		MirrorBrushIndicator->bDrawRadiusCircle = BrushIndicator->bDrawRadiusCircle;
+
+		FTransform MirrorBrushTransform = ComputeMirroredFrame(HoverStamp.LocalFrame).ToFTransform();
+
+		// Project onto the surface
+		MirrorBrushTriangleID = IndexConstants::InvalidID;
+		if (const FDynamicMesh3* Mesh = GetSculptMesh())
+		{
+			const double RayMaxDistance = 10.0;
+			const FVector MirrorBrushNormal = MirrorBrushTransform.GetUnitAxis(EAxis::Z);
+			const FVector RayDirection = -MirrorBrushNormal;
+			const FVector RayOrigin = MirrorBrushTransform.GetTranslation() + MirrorBrushNormal * (RayMaxDistance * 0.5);
+
+			double OutRayParameter = TNumericLimits<double>::Max();
+			const int32 HitTID = Octree.FSparseDynamicOctree3::FindNearestHitObject(
+				FRay3d(RayOrigin, RayDirection),
+				[&](int tid) { return Mesh->GetTriBounds(tid); },
+				[&](int tid, const FRay3d& Ray) 
+				{
+					UE::Geometry::FIntrRay3Triangle3d Intr = 
+						UE::Geometry::TMeshQueries<FDynamicMesh3>::TriangleIntersection(*Mesh, tid, Ray);
+					if (Intr.IntersectionType == EIntersectionType::Point)
+					{
+						OutRayParameter = FMath::Min(OutRayParameter, Intr.RayParameter);
+						return Intr.RayParameter;
+					}
+					return TNumericLimits<double>::Max();
+
+				}, RayMaxDistance);
+			if (HitTID != IndexConstants::InvalidID && Mesh->IsTriangle(HitTID))
+			{
+				MirrorBrushTriangleID = HitTID;
+				MirrorBrushTransform.SetTranslation(RayOrigin + RayDirection * OutRayParameter);
+			}
+		}
+		MirrorBrushIndicator->Update(
+			(float)GetCurrentBrushRadius(),
+			MirrorBrushTransform,
+			1.0f - (float)GetCurrentBrushFalloff());
+
+		UpdateMirrorPlaneMesh();
+	}
+}
+
 void UDataflowEditorVertexAttributePaintTool::SetDataflowEditorContextObject(TObjectPtr<UDataflowContextObject> InDataflowEditorContextObject)
 {
 	DataflowEditorContextObject = InDataflowEditorContextObject;
@@ -845,7 +1042,22 @@ void UDataflowEditorVertexAttributePaintTool::Shutdown(EToolShutdownType Shutdow
 		PreviewMesh = nullptr;
 	}
 
+	if (MirrorBrushIndicatorMesh)
+	{
+		MirrorBrushIndicatorMesh->Disconnect();
+		MirrorBrushIndicatorMesh = nullptr;
+	}
+
+	if (MirrorPlaneMesh)
+	{
+		MirrorPlaneMesh->Disconnect();
+		MirrorPlaneMesh = nullptr;
+	}
+
 	UMeshSculptToolBase::Shutdown(ShutdownType);
+
+	BrushIndicatorMesh = nullptr;
+	GetToolManager()->GetPairedGizmoManager()->DeregisterGizmoType(UE::DataflowEditorVertexAttributePaintTool::Private::MirrorBrushGizmoType);
 }
 
 
@@ -866,11 +1078,32 @@ void UDataflowEditorVertexAttributePaintTool::EnableValuePicker(bool bEnable)
 	bPendingPickWeight = bEnable;
 }
 
-bool UDataflowEditorVertexAttributePaintTool::IsValuePickerEnabled()
+bool UDataflowEditorVertexAttributePaintTool::IsValuePickerEnabled() const
 {
 	return bPendingPickWeight;
 }
 
+void UDataflowEditorVertexAttributePaintTool::EnableBrushMirroring(bool bEnable)
+{
+	bBrushMirroringEnabled = bEnable;
+}
+
+bool UDataflowEditorVertexAttributePaintTool::IsBrushMirroringEnabled() const
+{
+	return bBrushMirroringEnabled;
+}
+
+bool UDataflowEditorVertexAttributePaintTool::IsBrushMirroringInObjectSpace() const
+{
+	return ToolProperties ? ToolProperties->MirrorProperties.bObjectSpace : false;
+}
+
+
+EAxis::Type UDataflowEditorVertexAttributePaintTool::GetMirrorAxis() const
+{
+	return ToolProperties ? ToolProperties->MirrorProperties.MirrorAxis.GetValue() : EAxis::X;
+}
+
 FInputChord UDataflowEditorVertexAttributePaintTool::GreyscaleDisplayInputChord = FInputChord(EModifierKey::Control, EKeys::NumPadSeven);
 FInputChord UDataflowEditorVertexAttributePaintTool::ColorMapDisplayInputChord = FInputChord(EModifierKey::Control, EKeys::NumPadEight);
 FInputChord UDataflowEditorVertexAttributePaintTool::MaterialColorDisplayInputChord = FInputChord(EModifierKey::Control, EKeys::NumPadSix);
@@ -1003,7 +1236,7 @@ void UDataflowEditorVertexAttributePaintTool::RegisterActions(FInteractiveToolAc
 		LOCTEXT("MirrorValues", "Mirror Values"),
 		LOCTEXT("MirrorValuesTooltip", "Mirror the current selection values"),
 		MirrorValuesInputChord.GetModifierKey(), MirrorValuesInputChord.Key,
-		[this]() { MirrorValues(); });
+		[this]() { OnMirrorCommand(); });
 	++ActionID;
 
 	//-------------------------------------------------------------------------------------------------------
@@ -1062,7 +1295,7 @@ FBox UDataflowEditorVertexAttributePaintTool::GetWorldSpaceFocusBox()
 				}
 				return static_cast<FBox>(MeshBounds);
 			}
-			return FBox(PreviewMesh->GetMesh()->GetBounds());
+			return FBox(MeshLocalBounds);
 		}
 
 	}
@@ -1267,7 +1500,7 @@ void UDataflowEditorVertexAttributePaintTool::OnCancelStroke()
 	CancelChange();
 }
 
-void UDataflowEditorVertexAttributePaintTool::UpdateROI(const FSculptBrushStamp& BrushStamp)
+void UDataflowEditorVertexAttributePaintTool::UpdateROI(const FSculptBrushStamp& BrushStamp, int32 BrushTriangleID)
 {
 	SCOPE_CYCLE_COUNTER(VertexAttributePaintTool_UpdateROI);
 
@@ -1280,7 +1513,7 @@ void UDataflowEditorVertexAttributePaintTool::UpdateROI(const FSculptBrushStamp&
 
 	TriangleROI.Reset();
 
-	int32 CenterTID = GetBrushTriangleID();
+	int32 CenterTID = BrushTriangleID;
 	if (Mesh->IsTriangle(CenterTID))
 	{
 		TriangleROI.Add(CenterTID);
@@ -1421,6 +1654,17 @@ bool UDataflowEditorVertexAttributePaintTool::UpdateStampPosition(const FRay& Wo
 	CurrentStamp.PrevLocalFrame = LastStamp.LocalFrame;
 	CurrentStamp.PrevWorldFrame = LastStamp.WorldFrame;
 
+	const EAxis::Type MirrorAxis = GetMirrorAxis();
+
+	// update the mirrored stamp
+	MirrorCurrentStamp = CurrentStamp;
+	MirrorCurrentStamp.LocalFrame = ComputeMirroredFrame(CurrentStamp.LocalFrame);
+	MirrorCurrentStamp.WorldFrame = MirrorCurrentStamp.LocalFrame;
+	MirrorCurrentStamp.WorldFrame.Transform(CurTargetTransform);
+	MirrorCurrentStamp.PrevLocalFrame = ComputeMirroredFrame(CurrentStamp.PrevLocalFrame);
+	MirrorCurrentStamp.PrevWorldFrame = MirrorCurrentStamp.PrevLocalFrame;
+	MirrorCurrentStamp.PrevWorldFrame.Transform(CurTargetTransform);
+
 	FVector3d MoveDelta = CurrentStamp.LocalFrame.Origin - CurrentStamp.PrevLocalFrame.Origin;
 
 	if (UseBrushOp->IgnoreZeroMovements() && MoveDelta.SquaredLength() < 0.1 * CurrentBrushRadius)
@@ -1431,7 +1675,7 @@ bool UDataflowEditorVertexAttributePaintTool::UpdateStampPosition(const FRay& Wo
 	return true;
 }
 
-bool UDataflowEditorVertexAttributePaintTool::ApplyStamp()
+bool UDataflowEditorVertexAttributePaintTool::ApplyStamp(const FSculptBrushStamp& InStamp)
 {
 	SCOPE_CYCLE_COUNTER(VertexAttributePaintTool_ApplyStamp);
 
@@ -1442,13 +1686,10 @@ bool UDataflowEditorVertexAttributePaintTool::ApplyStamp()
 	WeightBrushOp->bApplyRadiusLimit = IsInBrushMode();
 
 	FDynamicMesh3* Mesh = GetSculptMesh();
-	WeightBrushOp->ApplyStampByVertices(Mesh, CurrentStamp, VertexROI, ROIWeightValueBuffer);
+	WeightBrushOp->ApplyStampByVertices(Mesh, InStamp, VertexROI, ROIWeightValueBuffer);
 
 	bool bUpdated = SyncMeshWithWeightBuffer(Mesh);
 
-	LastStamp = CurrentStamp;
-	LastStamp.TimeStamp = FDateTime::Now();
-
 	return bUpdated;
 }
 
@@ -2419,6 +2660,7 @@ void UDataflowEditorVertexAttributePaintTool::OnTick(float DeltaTime)
 
 	ConfigureIndicator(IsVolumetricBrush());
 	SetIndicatorVisibility(!bIsLasso && !bIsMeshEditMode && bBrushIsHoverMesh && !bPendingPickWeight);
+	UpdateMirrorIndicator();
 
 	SCOPE_CYCLE_COUNTER(VertexAttributePaintTool_Tick);
 
@@ -2465,29 +2707,52 @@ void UDataflowEditorVertexAttributePaintTool::OnTick(float DeltaTime)
 		{
 			SCOPE_CYCLE_COUNTER(VertexAttributePaintTool_Tick_ApplyStampBlock);
 
+			FDynamicMesh3* Mesh = GetSculptMesh();
+
 			// update sculpt ROI
-			UpdateROI(CurrentStamp);
+			UpdateROI(CurrentStamp, GetBrushTriangleID());
 
 			// append updated ROI to modified region (async)
-			FDynamicMesh3* Mesh = GetSculptMesh();
 			TFuture<void> AccumulateROI = Async(UE::DataflowEditorVertexAttributePaintTool::Private::WeightPaintToolAsyncExecTarget, [&]()
 				{
 					UE::Geometry::VertexToTriangleOneRing(Mesh, VertexROI, AccumulatedTriangleROI);
 				});
 
 			// apply the stamp
-			bool bWeightsModified = ApplyStamp();
+			bool bWeightsModified = ApplyStamp(CurrentStamp);
 
 			if (bWeightsModified)
 			{
 				SCOPE_CYCLE_COUNTER(VertexAttributePaintTool_Tick_UpdateMeshBlock);
 				UpdateVertexColorOverlay(&TriangleROI);
-				//DynamicMeshComponent->FastNotifyTriangleVerticesUpdated(TriangleROI, EMeshRenderAttributeFlags::VertexColors);
-				GetToolManager()->PostInvalidation();
 			}
 
 			// we don't really need to wait for these to happen to end Tick()...
 			AccumulateROI.Wait();
+
+			LastStamp = CurrentStamp;
+			LastStamp.TimeStamp = FDateTime::Now();
+
+			// now mirror stamp 
+			if (bBrushMirroringEnabled)
+			{
+				// update sculpt ROI
+				UpdateROI(MirrorCurrentStamp, MirrorBrushTriangleID);
+
+				// apply the stamp
+				bWeightsModified |= ApplyStamp(MirrorCurrentStamp);
+
+				if (bWeightsModified)
+				{
+					SCOPE_CYCLE_COUNTER(VertexAttributePaintTool_Tick_UpdateMeshBlock);
+					UpdateVertexColorOverlay(&TriangleROI);
+				}
+			}
+
+			if (bWeightsModified)
+			{
+				GetToolManager()->PostInvalidation();
+			}
 		};
 
 	if (IsInBrushMode() && !bPendingPickWeight)
@@ -2821,6 +3086,18 @@ void UDataflowEditorVertexAttributePaintTool::UpdateVertexColorOverlay(const TSe
 	PreviewMesh->NotifyDeferredEditCompleted(UPreviewMesh::ERenderUpdateMode::FastUpdate, EMeshRenderAttributeFlags::VertexColors, false);
 }
 
+void UDataflowEditorVertexAttributePaintTool::OnMirrorCommand()
+{
+	if (IsInBrushMode())
+	{
+		EnableBrushMirroring(!IsBrushMirroringEnabled());
+	}
+	else
+	{
+		MirrorValues();
+	}
+}
+
 void UDataflowEditorVertexAttributePaintTool::MirrorValues()
 {
 	if (ToolProperties)
diff --git a/Engine/Plugins/Experimental/Dataflow/Source/DataflowEditor/Private/DataflowEditorTools/DataflowEditorVertexAttributePaintToolDetailCustomization.cpp b/Engine/Plugins/Experimental/Dataflow/Source/DataflowEditor/Private/DataflowEditorTools/DataflowEditorVertexAttributePaintToolDetailCustomization.cpp
--- a/Engine/Plugins/Experimental/Dataflow/Source/DataflowEditor/Private/DataflowEditorTools/DataflowEditorVertexAttributePaintToolDetailCustomization.cpp
+++ b/Engine/Plugins/Experimental/Dataflow/Source/DataflowEditor/Private/DataflowEditorTools/DataflowEditorVertexAttributePaintToolDetailCustomization.cpp
@@ -573,6 +573,7 @@ namespace UE::Dataflow::Editor
 			{
 			case EDataflowEditorToolEditMode::Brush:
 				AddBrushUI(DetailBuilder);
+				AddBrushMirrorUI(DetailBuilder);
 				break;
 			case EDataflowEditorToolEditMode::Mesh:
 				AddSelectionUI(DetailBuilder);
@@ -1051,7 +1052,94 @@ namespace UE::Dataflow::Editor
 
 			// Brush backface hit option
 			AddBrushHitBackFacesRow(BrushCategory);
+		}
+	}
+
+	///////////////////////////////////////////////////////////////////////////////////////////////////////////////
+
+	TSharedRef<SWidget> FVertexAttributePaintToolDetailCustomization::MakeBrushMirrorEnableWidget() const
+	{
+		return SNew(SCheckBox)
+			.Style(FAppStyle::Get(), "ToggleButtonCheckbox")
+			.IsChecked_Lambda([Tool = Tool]()-> ECheckBoxState
+				{
+					const bool bIsBrushMirroringEnabled = (Tool.IsValid() && Tool->IsBrushMirroringEnabled());
+					return bIsBrushMirroringEnabled ? ECheckBoxState::Checked : ECheckBoxState::Unchecked;
+				})
+			.OnCheckStateChanged_Lambda([Tool = Tool](const ECheckBoxState NewState)
+				{
+					if (Tool.IsValid())
+					{
+						Tool->EnableBrushMirroring(NewState == ECheckBoxState::Checked);
+					}
+				})
+			.Padding(2.0f)
+			.ToolTip(MakeTooltipWithShortcut(LOCTEXT("BrushMirrorEnabledTooltip", "Enable or disable brush mirrorring"), UDataflowEditorVertexAttributePaintTool::MirrorValuesInputChord))
+			.HAlign(HAlign_Center)
+			[
+				SNew(SImage)
+					.Image(FAppStyle::GetBrush("UMGEditor.Mirror"))
+					.ColorAndOpacity(FSlateColor::UseForeground())
+			];
+	}
+
+	void FVertexAttributePaintToolDetailCustomization::AddBrushMirrorEnableRow(IDetailCategoryBuilder& BrushCategory) const
+	{
+		BrushCategory.AddCustomRow(LOCTEXT("BrushMirrorEnableCategory", "Brush Mirror Enable"), false)
+			.NameContent()
+			[
+				SNew(STextBlock)
+					.Text(LOCTEXT("BrushMirrorEnableLabel", "Enable Mirroring"))
+					.Font(FAppStyle::GetFontStyle(TEXT("PropertyWindow.NormalFont")))
+					.ToolTipText(LOCTEXT("BrushMirrorEnableTooltip", "Enable/Disable brush mirroring"))
+			]
+			.ValueContent()
+			[
+				MakeBrushMirrorEnableWidget()
+			];
+	}
+
+	void FVertexAttributePaintToolDetailCustomization::AddBrushMirrorPlaneRow(IDetailCategoryBuilder& BrushCategory) const
+	{
+		BrushCategory.AddCustomRow(LOCTEXT("BrushMirrorPlaneCategory", "Brush Mirror Axis"), false)
+			.NameContent()
+			[
+				SNew(STextBlock)
+					.Text(LOCTEXT("BrushMirrorPlaneLabel", "Mirror Plane"))
+					.Font(FAppStyle::GetFontStyle(TEXT("PropertyWindow.NormalFont")))
+					.ToolTipText(LOCTEXT("BrushMirrorPlaneTooltip", "Plane axis to use for mirroring"))
+			]
+			.ValueContent()
+			[
+				MakeMirrorAxisWidget()
+			];
+	}
+
+	void FVertexAttributePaintToolDetailCustomization::AddBrushMirrorObjectSpaceRow(IDetailCategoryBuilder& BrushCategory) const
+	{
+		BrushCategory.AddCustomRow(LOCTEXT("BrushMirrorObjectSpaceCategory", "Brush Mirror Object Space"), false)
+			.NameContent()
+			[
+				SNew(STextBlock)
+					.Text(LOCTEXT("BrushMirrorObjectSpaceLabel", "Object Space"))
+					.Font(FAppStyle::GetFontStyle(TEXT("PropertyWindow.NormalFont")))
+					.ToolTipText(LOCTEXT("BrushMirrorObjectSpaceTooltip", "When checked, the planes will be centered on the object bounds."))
+			]
+			.ValueContent()
+			[
+				MakeMirrorObjectSpaceCheckboxWidget()
+			];
+	}
+
+	void FVertexAttributePaintToolDetailCustomization::AddBrushMirrorUI(IDetailLayoutBuilder& DetailBuilder) const
+	{
+		IDetailCategoryBuilder& BrushCategory = DetailBuilder.EditCategory("BrushMirror", FText::GetEmpty(), ECategoryPriority::Important);
+		{
+			AddBrushMirrorEnableRow(BrushCategory);
 
+			AddBrushMirrorPlaneRow(BrushCategory);
+
+			AddBrushMirrorObjectSpaceRow(BrushCategory);
 		}
 	}
 
@@ -1673,10 +1761,9 @@ namespace UE::Dataflow::Editor
 		];
 	}
 
-	void FVertexAttributePaintToolDetailCustomization::MakeSelectionMirrorOperationRow(IDetailCategoryBuilder& EditValuesCategory) const
+	TSharedRef<SWidget> FVertexAttributePaintToolDetailCustomization::MakeMirrorAxisWidget() const
 	{
-		TSharedRef<SSegmentedControl<EAxis::Type>> XYZButtonsWidget =
-			SNew(SSegmentedControl<EAxis::Type>)
+		return SNew(SSegmentedControl<EAxis::Type>)
 			.ToolTipText(LOCTEXT("MirrorAxisTooltip",
 				"X: copies weights across the YZ plane.\n"
 				"Y: copies weights across the XZ plane.\n"
@@ -1698,9 +1785,11 @@ namespace UE::Dataflow::Editor
 			.Text(LOCTEXT("MirrorYLabel", "Y"))
 			+ SSegmentedControl<EAxis::Type>::Slot(EAxis::Z)
 			.Text(LOCTEXT("MirrorZLabel", "Z"));
+	}
 
-		TSharedRef<SSegmentedControl<EDataflowEditorToolMirrorDirection>> MirrortSideButtonsWidget =
-			SNew(SSegmentedControl<EDataflowEditorToolMirrorDirection>)
+	TSharedRef<SWidget> FVertexAttributePaintToolDetailCustomization::MakeMirrorDirectionButtonsWidget() const
+	{
+		return SNew(SSegmentedControl<EDataflowEditorToolMirrorDirection>)
 			.ToolTipText(LOCTEXT("MirrorDirectionTooltip", "The direction that determines what side of the plane to copy weights from."))
 			.Value_Lambda([this]()
 				{
@@ -1720,23 +1809,45 @@ namespace UE::Dataflow::Editor
 			.Text(LOCTEXT("MirrorNegToPosLabel", "- to +"))
 			.ToolTipWidget(MakeTooltipWithShortcut(LOCTEXT("MirrorDirectionNegToPosTooltip", "Mirror the values from the negative side of the axis to the positive side."), UDataflowEditorVertexAttributePaintTool::InvertMirrorDirectionInputChord))
 			;
+	}
 
-		TSharedRef<SButton> MirrorButtonWidget = 
-			SNew(SButton)
-			.HAlign(HAlign_Center)
-			.Text(LOCTEXT("MirrorWeightsButtonLabel", "Mirror"))
-			.ToolTip(MakeTooltipWithShortcut(LOCTEXT("MirrorButtonTooltip", "Values are copied across the given plane in the given direction.\nThis command operates on the selected vertices."), UDataflowEditorVertexAttributePaintTool::MirrorValuesInputChord))
-			.OnClicked_Lambda([this]()
-				{
-					if (Tool.IsValid())
-					{
-						Tool->MirrorValues();
-					}
-					return FReply::Handled();
-				});
+	TSharedRef<SWidget> FVertexAttributePaintToolDetailCustomization::MakeMirrorPlaneWidget(bool bShowDirection) const
+	{
+		TSharedRef<SWidget> XYZButtonsWidget = MakeMirrorAxisWidget();
+		TSharedRef<SWidget> MirrorDirectionButtonsWidget =
+			bShowDirection
+			? MakeMirrorDirectionButtonsWidget()
+			: SNullWidget::NullWidget;
+
+		return SNew(SHorizontalBox)
+			+ SHorizontalBox::Slot()
+			.VAlign(VAlign_Center)
+			.FillWidth(WeightEditingLabelsPercent)
+			[
+				SNew(STextBlock)
+				.Text(LOCTEXT("MirrorPlaneLabel", "Mirror Plane"))
+				.Font(FAppStyle::GetFontStyle(TEXT("PropertyWindow.NormalFont")))
+				.ToolTipText(LOCTEXT("MirrorPlaneTooltip", "The plane to copy weights across."))
+			]
+			+ SHorizontalBox::Slot()
+			.FillWidth(1.f)
+			[
+				SNew(SHorizontalBox)
+				+ SHorizontalBox::Slot()
+				[
+					XYZButtonsWidget
+				]
+
+				+ SHorizontalBox::Slot()
+				[
+					MirrorDirectionButtonsWidget
+				]
+			];
+	}
 
-		TSharedRef<SCheckBox> ObjectSpaceCheckBox =
-			SNew(SCheckBox)
+	TSharedRef<SWidget> FVertexAttributePaintToolDetailCustomization::MakeMirrorObjectSpaceCheckboxWidget() const
+	{
+		return SNew(SCheckBox)
 			.IsChecked_Lambda([ToolProperties = ToolProperties]()
 				{
 					const bool bIsChecked = (ToolProperties.IsValid() && ToolProperties->MirrorProperties.bObjectSpace);
@@ -1749,6 +1860,48 @@ namespace UE::Dataflow::Editor
 						ToolProperties->MirrorProperties.bObjectSpace = (InCheckBoxState == ECheckBoxState::Checked);
 					}
 				});
+	}
+
+	TSharedRef<SWidget> FVertexAttributePaintToolDetailCustomization::MakeMirrorObjectSpaceWidget() const
+	{
+		return SNew(SHorizontalBox)
+			+ SHorizontalBox::Slot()
+			.VAlign(VAlign_Center)
+			.FillWidth(WeightEditingLabelsPercent)
+			[
+				SNew(STextBlock)
+					.Text(LOCTEXT("MirrorObjectSpaceLabel", "Object Space"))
+					.Font(FAppStyle::GetFontStyle(TEXT("PropertyWindow.NormalFont")))
+					.ToolTipText(LOCTEXT("MirrorObjectSpaceTooltip", "When checked, the planes will be centered on the object bounds."))
+			]
+			+ SHorizontalBox::Slot()
+			.FillWidth(1.f)
+			[
+				MakeMirrorObjectSpaceCheckboxWidget()
+			];
+	}
+
+	TSharedRef<SWidget> FVertexAttributePaintToolDetailCustomization::MakeMirrorButton() const
+	{
+		return SNew(SButton)
+			.HAlign(HAlign_Center)
+			.Text(LOCTEXT("MirrorWeightsButtonLabel", "Mirror"))
+			.ToolTip(MakeTooltipWithShortcut(LOCTEXT("MirrorButtonTooltip", "Values are copied across the given plane in the given direction.\nThis command operates on the selected vertices."), UDataflowEditorVertexAttributePaintTool::MirrorValuesInputChord))
+			.OnClicked_Lambda([this]()
+				{
+					if (Tool.IsValid())
+					{
+						Tool->MirrorValues();
+					}
+					return FReply::Handled();
+				});
+	}
+
+	void FVertexAttributePaintToolDetailCustomization::MakeSelectionMirrorOperationRow(IDetailCategoryBuilder& EditValuesCategory) const
+	{
+		TSharedRef<SWidget> MirrorPlaneWidget = MakeMirrorPlaneWidget(/*bShowDirection*/ true);
+		TSharedRef<SWidget> ObjectSpaceWidget = MakeMirrorObjectSpaceWidget();
+		TSharedRef<SWidget> MirrorButtonWidget = MakeMirrorButton();
 
 		constexpr float HorizontalPadding = 2.0f;
 		constexpr float VerticalPadding = 4.0f;
@@ -1765,49 +1918,12 @@ namespace UE::Dataflow::Editor
 					+ SVerticalBox::Slot()
 					.Padding(HorizontalPadding, VerticalPadding)
 					[
-						SNew(SHorizontalBox)
-						+ SHorizontalBox::Slot()
-						.VAlign(VAlign_Center)
-						.FillWidth(WeightEditingLabelsPercent)
-						[
-							SNew(STextBlock)
-								.Text(LOCTEXT("MirrorPlaneLabel", "Mirror Plane"))
-								.Font(FAppStyle::GetFontStyle(TEXT("PropertyWindow.NormalFont")))
-								.ToolTipText(LOCTEXT("MirrorPlaneTooltip", "The plane to copy weights across."))
-						]
-						+ SHorizontalBox::Slot()
-						.FillWidth(1.f)
-						[
-							SNew(SHorizontalBox)
-							+ SHorizontalBox::Slot()
-							[
-								XYZButtonsWidget
-							]
-
-							+ SHorizontalBox::Slot()
-							[
-								MirrortSideButtonsWidget
-							]
-						]
+						MirrorPlaneWidget
 					]
 					+ SVerticalBox::Slot()
 					.Padding(HorizontalPadding, VerticalPadding)
 					[
-						SNew(SHorizontalBox)
-							+ SHorizontalBox::Slot()
-							.VAlign(VAlign_Center)
-							.FillWidth(WeightEditingLabelsPercent)
-							[
-								SNew(STextBlock)
-									.Text(LOCTEXT("MirrorObjectSpaceLabel", "Object Space"))
-									.Font(FAppStyle::GetFontStyle(TEXT("PropertyWindow.NormalFont")))
-									.ToolTipText(LOCTEXT("MirrorObjectSpaceTooltip", "When checked, the planes will be cenetred on the object bounds."))
-							]
-							+ SHorizontalBox::Slot()
-							.FillWidth(1.f)
-							[
-								ObjectSpaceCheckBox
-							]
+						ObjectSpaceWidget
 					]
 					+ SVerticalBox::Slot()
 					.Padding(0.f, WeightEditVerticalPadding)
diff --git a/Engine/Plugins/Experimental/Dataflow/Source/DataflowEditor/Private/DataflowEditorTools/DataflowEditorVertexAttributePaintToolDetailCustomization.h b/Engine/Plugins/Experimental/Dataflow/Source/DataflowEditor/Private/DataflowEditorTools/DataflowEditorVertexAttributePaintToolDetailCustomization.h
--- a/Engine/Plugins/Experimental/Dataflow/Source/DataflowEditor/Private/DataflowEditorTools/DataflowEditorVertexAttributePaintToolDetailCustomization.h
+++ b/Engine/Plugins/Experimental/Dataflow/Source/DataflowEditor/Private/DataflowEditorTools/DataflowEditorVertexAttributePaintToolDetailCustomization.h
@@ -47,8 +47,14 @@ namespace UE::Dataflow::Editor
 		void AddBrushValueAtBrushRow(IDetailCategoryBuilder& BrushCategory) const;
 		void AddBrushValueQueryMethodRow(IDetailCategoryBuilder& BrushCategory) const;
 		void AddBrushHitBackFacesRow(IDetailCategoryBuilder& BrushCategory) const;
+
+		TSharedRef<SWidget> MakeBrushMirrorEnableWidget() const;
+		void AddBrushMirrorEnableRow(IDetailCategoryBuilder& BrushCategory) const;
+		void AddBrushMirrorPlaneRow(IDetailCategoryBuilder& BrushCategory) const;
+		void AddBrushMirrorObjectSpaceRow(IDetailCategoryBuilder& BrushCategory) const;
+		void AddBrushMirrorUI(IDetailLayoutBuilder& DetailBuilder) const;
+
 		void AddVisibilityFilterRow(IDetailCategoryBuilder& BrushCategory) const;
-		
 
 		void AddSelectionUI(IDetailLayoutBuilder& DetailBuilder) const;
 		TSharedRef<SWidget> MakeSelectionElementsToolbar() const;
@@ -66,8 +72,15 @@ namespace UE::Dataflow::Editor
 		void MakeSelectionPruneOperationRow(IDetailCategoryBuilder& EditValuesCategory) const;
 		void MakeSelectionCopyAndPasteRow(IDetailCategoryBuilder& EditValuesCategory) const;
 		
-		void HideProperty(IDetailLayoutBuilder& DetailBuilder, FName PropertyName) const;
+		TSharedRef<SWidget> MakeMirrorAxisWidget() const;
+		TSharedRef<SWidget> MakeMirrorDirectionButtonsWidget() const;
+		TSharedRef<SWidget> MakeMirrorPlaneWidget(bool bShowDirection) const;
+		TSharedRef<SWidget> MakeMirrorObjectSpaceCheckboxWidget() const;
+		TSharedRef<SWidget> MakeMirrorObjectSpaceWidget() const;
+		TSharedRef<SWidget> MakeMirrorButton() const;
 
+		void HideProperty(IDetailLayoutBuilder& DetailBuilder, FName PropertyName) const;
+		
 		IDetailLayoutBuilder* CurrentDetailBuilder = nullptr;
 		TWeakObjectPtr<UDataflowEditorVertexAttributePaintToolProperties> ToolProperties;
 		TWeakObjectPtr<UDataflowEditorVertexAttributePaintTool> Tool;
diff --git a/Engine/Plugins/Experimental/Dataflow/Source/DataflowEditor/Public/DataflowEditorTools/DataflowEditorVertexAttributePaintTool.h b/Engine/Plugins/Experimental/Dataflow/Source/DataflowEditor/Public/DataflowEditorTools/DataflowEditorVertexAttributePaintTool.h
--- a/Engine/Plugins/Experimental/Dataflow/Source/DataflowEditor/Public/DataflowEditorTools/DataflowEditorVertexAttributePaintTool.h
+++ b/Engine/Plugins/Experimental/Dataflow/Source/DataflowEditor/Public/DataflowEditorTools/DataflowEditorVertexAttributePaintTool.h
@@ -481,7 +481,20 @@ protected:
 	// end UMeshSculptToolBase API
 
 	void EnableValuePicker(bool bEnable);
-	bool IsValuePickerEnabled();
+	bool IsValuePickerEnabled() const;
+
+	void EnableBrushMirroring(bool bEnable);
+	bool IsBrushMirroringEnabled() const;
+
+	bool IsBrushMirroringInObjectSpace() const;
+
+	EAxis::Type GetMirrorAxis() const;
+
+	void InitializeMirrorIndicator();
+	void UpdateMirrorIndicator();
+	void UpdateMirrorPlaneMesh();
+	UE::Geometry::FFrame3d ComputeMirroredFrame(const UE::Geometry::FFrame3d& InLocalSpaceFrame) const;
+	void OnMirrorCommand();
 
 protected:
 	//
@@ -521,9 +534,22 @@ protected:
 	UPROPERTY()
 	TObjectPtr<UPreviewMesh> PreviewMesh = nullptr;
 	FMeshDescription PreviewMeshDescription;
+	UE::Geometry::FAxisAlignedBox3d MeshLocalBounds;
+
+	UPROPERTY()
+	TObjectPtr<UToolMeshSelector> MeshSelector = nullptr;
 
 	UPROPERTY()
-	TObjectPtr<UToolMeshSelector> MeshSelector;
+	TObjectPtr<UBrushStampIndicator> MirrorBrushIndicator = nullptr;
+
+	UPROPERTY()
+	TObjectPtr<UPreviewMesh> MirrorBrushIndicatorMesh = nullptr;
+
+	UPROPERTY()
+	TObjectPtr<UPreviewMesh> MirrorPlaneMesh = nullptr;
+
+	FSculptBrushStamp MirrorCurrentStamp;
+	int32 MirrorBrushTriangleID = INDEX_NONE;
 
 	// realtime visualization
 	UE_API void OnDynamicMeshComponentChanged(UDynamicMeshComponent* Component, const FMeshVertexChange* Change, bool bRevert);
@@ -544,10 +570,10 @@ protected:
 	TArray<bool> VisibilityFilterBuffer;
 	TSet<int> VertexSetBuffer;
 	TSet<int> TriangleROI;
-	UE_API void UpdateROI(const FSculptBrushStamp& CurrentStamp);
+	UE_API void UpdateROI(const FSculptBrushStamp& CurrentStamp, int32 BrushTriangleID);
 
 	UE_API bool UpdateStampPosition(const FRay& WorldRay);
-	UE_API bool ApplyStamp();
+	UE_API bool ApplyStamp(const FSculptBrushStamp& InStamp);
 
 	UE::Geometry::FDynamicMeshOctree3 Octree;
 
@@ -556,6 +582,7 @@ protected:
 
 	bool bPendingPickWeight = false;
 	bool bBrushIsHoverMesh = false;
+	bool bBrushMirroringEnabled = false;
 
 	TArray<int32> ROITriangleBuffer;
 	TArray<double> ROIWeightValueBuffer;
