Dataflow : paint tool : add invert selection option

#Synced-CL 50024027

diff --git a/Engine/Plugins/Experimental/Dataflow/Source/DataflowEditor/Private/DataflowEditorTools/DataflowEditorVertexAttributePaintTool.cpp b/Engine/Plugins/Experimental/Dataflow/Source/DataflowEditor/Private/DataflowEditorTools/DataflowEditorVertexAttributePaintTool.cpp
--- a/Engine/Plugins/Experimental/Dataflow/Source/DataflowEditor/Private/DataflowEditorTools/DataflowEditorVertexAttributePaintTool.cpp
+++ b/Engine/Plugins/Experimental/Dataflow/Source/DataflowEditor/Private/DataflowEditorTools/DataflowEditorVertexAttributePaintTool.cpp
@@ -2062,6 +2062,14 @@ void UDataflowEditorVertexAttributePaintTool::ShrinkSelection() const
 	}
 }
 
+void UDataflowEditorVertexAttributePaintTool::InvertSelection() const
+{
+	if (MeshSelector)
+	{
+		MeshSelector->InvertSelection();
+	}
+}
+
 void UDataflowEditorVertexAttributePaintTool::FloodSelection() const
 {
 	if (MeshSelector)
diff --git a/Engine/Plugins/Experimental/Dataflow/Source/DataflowEditor/Private/DataflowEditorTools/DataflowEditorVertexAttributePaintToolDetailCustomization.cpp b/Engine/Plugins/Experimental/Dataflow/Source/DataflowEditor/Private/DataflowEditorTools/DataflowEditorVertexAttributePaintToolDetailCustomization.cpp
--- a/Engine/Plugins/Experimental/Dataflow/Source/DataflowEditor/Private/DataflowEditorTools/DataflowEditorVertexAttributePaintToolDetailCustomization.cpp
+++ b/Engine/Plugins/Experimental/Dataflow/Source/DataflowEditor/Private/DataflowEditorTools/DataflowEditorVertexAttributePaintToolDetailCustomization.cpp
@@ -1266,6 +1266,24 @@ namespace UE::Dataflow::Editor
 
 	TSharedRef<SWidget> FVertexAttributePaintToolDetailCustomization::MakeSelectionEditActionsToolbar() const
 	{
+		auto MakeButton = [this](const FText& InText, const FText& InTooltip, const FInputChord& InChord, TFunction<void()> OnClicked)
+			{
+				return SNew(SButton)
+					.IsEnabled_Lambda([this]
+						{
+							return Tool.IsValid() ? Tool->HasSelection() : false;
+						})
+					.HAlign(HAlign_Center)
+					.VAlign(VAlign_Center)
+					.Text(InText)
+					.ToolTip(MakeTooltipWithShortcut(InTooltip, InChord))
+					.OnClicked_Lambda([OnClicked]()
+						{
+							OnClicked();
+							return FReply::Handled();
+						});
+			};
+
 		// [ GROW / SHRINK / FLOOD ]
 		return SNew(SVerticalBox)
 			+ SVerticalBox::Slot()
@@ -1275,51 +1293,33 @@ namespace UE::Dataflow::Editor
 				+ SHorizontalBox::Slot()
 				.Padding(2.f, WeightEditVerticalPadding)
 				[
-					SNew(SButton)
-						.IsEnabled_Lambda([this]
-							{
-								return Tool.IsValid() ? Tool->HasSelection() : false;
-							})
-						.HAlign(HAlign_Center)
-						.VAlign(VAlign_Center)
-						.Text(LOCTEXT("GrowSelectionButtonLabel", "Grow"))
-						.ToolTip(MakeTooltipWithShortcut(
-							LOCTEXT("GrowSelectionTooltip", "Grow the current selection by adding connected neighbors to current selection."), 
-							UDataflowEditorVertexAttributePaintTool::GrowSelectionInputChord
-							))
-					.OnClicked_Lambda([this]()
-						{
-							if (Tool.IsValid())
-							{
-								Tool->GrowSelection();
-							}
-							return FReply::Handled();
-						})
+					MakeButton(
+						LOCTEXT("GrowSelectionButtonLabel", "Grow"),
+						LOCTEXT("GrowSelectionTooltip", "Grow the current selection by adding connected neighbors to current selection."),
+						UDataflowEditorVertexAttributePaintTool::GrowSelectionInputChord,
+						[this]() { if (Tool.IsValid()) { Tool->GrowSelection(); } }
+					)
 				]
 				+ SHorizontalBox::Slot()
 				.Padding(2.f, WeightEditVerticalPadding)
 				[
-					SNew(SButton)
-					.IsEnabled_Lambda([this]
-						{
-							return Tool.IsValid() ? Tool->HasSelection() : false;
-						})
-					.HAlign(HAlign_Center)
-					.VAlign(VAlign_Center)
-					.Text(LOCTEXT("ShrinkSelectionButtonLabel", "Shrink"))
-					.ToolTip(MakeTooltipWithShortcut(
+					MakeButton(
+						LOCTEXT("ShrinkSelectionButtonLabel", "Shrink"),
 						LOCTEXT("ShrinkSelectionTooltip", "Shrink the current selection by removing components on the border of the current selection."),
-						UDataflowEditorVertexAttributePaintTool::ShrinkSelectionInputChord
-						))
-					.OnClicked_Lambda([this]()
-						{
-							if (Tool.IsValid())
-							{
-								Tool->ShrinkSelection();
-							}
-							return FReply::Handled();
-						})
+						UDataflowEditorVertexAttributePaintTool::ShrinkSelectionInputChord,
+						[this]() { if (Tool.IsValid()) { Tool->ShrinkSelection(); } }
+					)
 				]
+				+ SHorizontalBox::Slot()
+					.Padding(2.f, WeightEditVerticalPadding)
+					[
+						MakeButton(
+							LOCTEXT("InvertSelectionButtonLabel", "Invert"),
+							LOCTEXT("InvertSelectionTooltip", "Invert the current selection."),
+							FInputChord(),
+							[this]() { if (Tool.IsValid()) { Tool->InvertSelection(); } }
+						)
+					]
 			]
 			+ SVerticalBox::Slot()
 			.AutoHeight()
@@ -1328,47 +1328,22 @@ namespace UE::Dataflow::Editor
 				+ SHorizontalBox::Slot()
 				.Padding(2.f, WeightEditVerticalPadding)
 				[
-					SNew(SButton)
-					.IsEnabled_Lambda([this]
-						{
-							return Tool.IsValid() ? Tool->HasSelection() : false;
-						})
-					.HAlign(HAlign_Center)
-					.VAlign(VAlign_Center)
-					.Text(LOCTEXT("FloodSelectionButtonLabel", "Flood"))
-						.ToolTip(MakeTooltipWithShortcut(
-							LOCTEXT("FloodSelectionTooltip", "Flood the current selection by adding all connected components to the current selection."),
-							UDataflowEditorVertexAttributePaintTool::FloodSelectionInputChord
-						))
-					.OnClicked_Lambda([this]()
-						{
-							if (Tool.IsValid())
-							{
-								Tool->FloodSelection();
-							}
-							return FReply::Handled();
-						})
+					MakeButton(
+						LOCTEXT("FloodSelectionButtonLabel", "Flood"),
+						LOCTEXT("FloodSelectionTooltip", "Flood the current selection by adding all connected components to the current selection."),
+						UDataflowEditorVertexAttributePaintTool::FloodSelectionInputChord,
+						[this]() { if (Tool.IsValid()) { Tool->FloodSelection(); } }
+					)
 				]
 				+ SHorizontalBox::Slot()
 				.Padding(2.f, WeightEditVerticalPadding)
 				[
-					SNew(SButton)
-					.IsEnabled_Lambda([this]
-						{
-							return Tool.IsValid() ? Tool->HasSelection() : false;
-						})
-					.HAlign(HAlign_Center)
-					.VAlign(VAlign_Center)
-					.Text(LOCTEXT("BorderSelectionButtonLabel", "Convert to Border"))
-					.ToolTipText(LOCTEXT("BorderSelectionTooltip", "Select vertices on the border of the current selection."))
-					.OnClicked_Lambda([this]()
-						{
-							if (Tool.IsValid())
-							{
-								Tool->SelectBorder();
-							}
-							return FReply::Handled();
-						})
+					MakeButton(
+						LOCTEXT("BorderSelectionButtonLabel", "Convert to Border"),
+						LOCTEXT("BorderSelectionTooltip", "Select vertices on the border of the current selection."),
+						FInputChord(),
+						[this]() { if (Tool.IsValid()) { Tool->SelectBorder(); } }
+					)
 				]
 			];
 	}
diff --git a/Engine/Plugins/Experimental/Dataflow/Source/DataflowEditor/Public/DataflowEditorTools/DataflowEditorVertexAttributePaintTool.h b/Engine/Plugins/Experimental/Dataflow/Source/DataflowEditor/Public/DataflowEditorTools/DataflowEditorVertexAttributePaintTool.h
--- a/Engine/Plugins/Experimental/Dataflow/Source/DataflowEditor/Public/DataflowEditorTools/DataflowEditorVertexAttributePaintTool.h
+++ b/Engine/Plugins/Experimental/Dataflow/Source/DataflowEditor/Public/DataflowEditorTools/DataflowEditorVertexAttributePaintTool.h
@@ -423,6 +423,7 @@ public:
 	UE_API bool HasSelection() const;
 	UE_API void GrowSelection() const;
 	UE_API void ShrinkSelection() const;
+	UE_API void InvertSelection() const;
 	UE_API void FloodSelection() const;
 	UE_API void SelectBorder() const;
 	UE_API void SetComponentSelectionMode(EComponentSelectionMode NewMode);
diff --git a/Engine/Plugins/Runtime/MeshModelingToolset/Source/MeshModelingToolsEditorOnly/Private/ToolMeshSelector.cpp b/Engine/Plugins/Runtime/MeshModelingToolset/Source/MeshModelingToolsEditorOnly/Private/ToolMeshSelector.cpp
--- a/Engine/Plugins/Runtime/MeshModelingToolset/Source/MeshModelingToolsEditorOnly/Private/ToolMeshSelector.cpp
+++ b/Engine/Plugins/Runtime/MeshModelingToolset/Source/MeshModelingToolsEditorOnly/Private/ToolMeshSelector.cpp
@@ -303,6 +303,16 @@ void UToolMeshSelector::ShrinkSelection() const
 	PolygonSelectionMechanic->ShrinkSelection(/*bAsTriangleTopology*/ true);
 }
 
+void UToolMeshSelector::InvertSelection() const
+{
+	if (!ensure(PolygonSelectionMechanic))
+	{
+		return;
+	}
+
+	PolygonSelectionMechanic->InvertSelection(/*bAsTriangleTopology*/ true);
+}
+
 void UToolMeshSelector::FloodSelection() const
 {
 	if (!ensure(PolygonSelectionMechanic))
diff --git a/Engine/Plugins/Runtime/MeshModelingToolset/Source/MeshModelingToolsEditorOnly/Public/ToolMeshSelector.h b/Engine/Plugins/Runtime/MeshModelingToolset/Source/MeshModelingToolsEditorOnly/Public/ToolMeshSelector.h
--- a/Engine/Plugins/Runtime/MeshModelingToolset/Source/MeshModelingToolsEditorOnly/Public/ToolMeshSelector.h
+++ b/Engine/Plugins/Runtime/MeshModelingToolset/Source/MeshModelingToolsEditorOnly/Public/ToolMeshSelector.h
@@ -77,6 +77,7 @@ public:
 	// edit selection
 	UE_API void GrowSelection() const;
 	UE_API void ShrinkSelection() const;
+	UE_API void InvertSelection() const;
 	UE_API void FloodSelection() const;
 	UE_API void SelectBorder() const;
 
diff --git a/Engine/Plugins/Runtime/MeshModelingToolset/Source/ModelingComponents/Private/Selection/PolygonSelectionMechanic.cpp b/Engine/Plugins/Runtime/MeshModelingToolset/Source/ModelingComponents/Private/Selection/PolygonSelectionMechanic.cpp
--- a/Engine/Plugins/Runtime/MeshModelingToolset/Source/ModelingComponents/Private/Selection/PolygonSelectionMechanic.cpp
+++ b/Engine/Plugins/Runtime/MeshModelingToolset/Source/ModelingComponents/Private/Selection/PolygonSelectionMechanic.cpp
@@ -367,6 +367,30 @@ void UPolygonSelectionMechanic::ShrinkSelection(bool bAsTriangleTopology)
 	});
 }
 
+void UPolygonSelectionMechanic::InvertSelection(bool bAsTriangleTopology)
+{
+	using namespace UE::Geometry;
+
+	ExecuteActionThroughGeometrySelection(bAsTriangleTopology, LOCTEXT("InvertSelectionChange", "Invert Selection"),
+		[this](FGeometrySelection& GeometrySelection)
+		{
+			FGeometrySelection NewSelection;
+			NewSelection.InitializeTypes(GeometrySelection);
+
+			bool bSuccess = MakeSelectAllSelection(
+				*Topology->GetMesh(), Topology,
+				[](FGeoSelectionID) { return true; },
+				NewSelection);
+
+			bSuccess = bSuccess && CombineSelectionInPlace(NewSelection, GeometrySelection,
+				UE::Geometry::EGeometrySelectionCombineModes::Subtract);
+
+			GeometrySelection = MoveTemp(NewSelection);
+
+			return bSuccess;
+		});
+}
+
 void UPolygonSelectionMechanic::ConvertSelectionToBorderVertices(bool bAsTriangleTopology)
 {
 	using namespace UE::Geometry;
diff --git a/Engine/Plugins/Runtime/MeshModelingToolset/Source/ModelingComponents/Public/Selection/PolygonSelectionMechanic.h b/Engine/Plugins/Runtime/MeshModelingToolset/Source/ModelingComponents/Public/Selection/PolygonSelectionMechanic.h
--- a/Engine/Plugins/Runtime/MeshModelingToolset/Source/ModelingComponents/Public/Selection/PolygonSelectionMechanic.h
+++ b/Engine/Plugins/Runtime/MeshModelingToolset/Source/ModelingComponents/Public/Selection/PolygonSelectionMechanic.h
@@ -61,6 +61,14 @@ public:
 	 */
 	UE_API void ShrinkSelection(bool bAsTriangleTopology);
 
+	/*
+	* Invert selection
+	*
+	* @param bAsTriangleTopology Can be set true if the topology type is FTriangleGroupTopology, to
+	*  perform the operation a bit more efficiently by using the mesh topology directly.
+	*/
+	UE_API void InvertSelection(bool bAsTriangleTopology);
+
 	/**
 	 * Converts selection to a vertex/corner selection of just the boundary vertices/corners
 	 * 
