Dataflow paint tool : make the mesh element visualizer settings persistent

#Synced-CL 50024027

diff --git a/Engine/Plugins/Experimental/Dataflow/Source/DataflowEditor/Private/DataflowEditorTools/DataflowEditorVertexAttributePaintTool.cpp b/Engine/Plugins/Experimental/Dataflow/Source/DataflowEditor/Private/DataflowEditorTools/DataflowEditorVertexAttributePaintTool.cpp
--- a/Engine/Plugins/Experimental/Dataflow/Source/DataflowEditor/Private/DataflowEditorTools/DataflowEditorVertexAttributePaintTool.cpp
+++ b/Engine/Plugins/Experimental/Dataflow/Source/DataflowEditor/Private/DataflowEditorTools/DataflowEditorVertexAttributePaintTool.cpp
@@ -763,13 +763,23 @@ void UDataflowEditorVertexAttributePaintTool::Setup()
 
 	// mesh element display is used to show seams and other features ofthe mesh
 	MeshElementsDisplay = NewObject<UMeshElementsVisualizer>(this);
+
 	MeshElementsDisplay->CreateInWorld(TargetWorld, FTransform::Identity);
 	if (ensure(MeshElementsDisplay->Settings))
 	{
-		MeshElementsDisplay->Settings->bShowNormalSeams = false;
-		MeshElementsDisplay->Settings->ThicknessScale = 0.5f;
-		MeshElementsDisplay->Settings->BoundaryEdgeColor = FColor(164, 73, 164);
-		MeshElementsDisplay->Settings->RestoreProperties(this, TEXT("DataflowEditorVertexAttributePaintTool"));
+		MeshElementsDisplay->Settings->bShowBorders = ToolProperties->DisplayProperties.bShowBorders;
+		MeshElementsDisplay->Settings->bShowUVSeams = ToolProperties->DisplayProperties.bShowUVSeams;
+		MeshElementsDisplay->Settings->bShowNormalSeams = ToolProperties->DisplayProperties.bShowNormalSeams;
+		MeshElementsDisplay->Settings->bShowTangentSeams = ToolProperties->DisplayProperties.bShowTangentSeams;
+		MeshElementsDisplay->Settings->bShowColorSeams = ToolProperties->DisplayProperties.bShowColorSeams;
+		MeshElementsDisplay->Settings->ThicknessScale = ToolProperties->DisplayProperties.ThicknessScale;
+		MeshElementsDisplay->Settings->WireframeColor = ToolProperties->DisplayProperties.WireframeColor;
+		MeshElementsDisplay->Settings->BoundaryEdgeColor = ToolProperties->DisplayProperties.BoundaryEdgeColor;
+		MeshElementsDisplay->Settings->UVSeamColor = ToolProperties->DisplayProperties.UVSeamColor;
+		MeshElementsDisplay->Settings->NormalSeamColor = ToolProperties->DisplayProperties.NormalSeamColor;
+		MeshElementsDisplay->Settings->TangentSeamColor = ToolProperties->DisplayProperties.TangentSeamColor;
+		MeshElementsDisplay->Settings->ColorSeamColor = ToolProperties->DisplayProperties.ColorSeamColor;
+		MeshElementsDisplay->Settings->DepthBias = ToolProperties->DisplayProperties.DepthBias;
 		AddToolPropertySource(MeshElementsDisplay->Settings);
 	}
 	MeshElementsDisplay->SetMeshAccessFunction([this](UMeshElementsVisualizer::ProcessDynamicMeshFunc ProcessFunc) {
@@ -1005,11 +1015,25 @@ void UDataflowEditorVertexAttributePaintTool::Shutdown(EToolShutdownType Shutdow
 
 	if (MeshElementsDisplay)
 	{
-		if (ensure(MeshElementsDisplay->Settings))
-		{
-			MeshElementsDisplay->Settings->SaveProperties(this, TEXT("DataflowEditorWeightMapPaintTool2"));
+		// make sure we move back all settings to the display properties so they can be saved in the config
+		if (ToolProperties && MeshElementsDisplay->Settings)
+		{
+			ToolProperties->DisplayProperties.bShowBorders = MeshElementsDisplay->Settings->bShowBorders;
+			ToolProperties->DisplayProperties.bShowUVSeams = MeshElementsDisplay->Settings->bShowUVSeams;
+			ToolProperties->DisplayProperties.bShowNormalSeams = MeshElementsDisplay->Settings->bShowNormalSeams;
+			ToolProperties->DisplayProperties.bShowTangentSeams = MeshElementsDisplay->Settings->bShowTangentSeams;
+			ToolProperties->DisplayProperties.bShowColorSeams = MeshElementsDisplay->Settings->bShowColorSeams;
+			ToolProperties->DisplayProperties.ThicknessScale = MeshElementsDisplay->Settings->ThicknessScale;
+			ToolProperties->DisplayProperties.WireframeColor = MeshElementsDisplay->Settings->WireframeColor;
+			ToolProperties->DisplayProperties.BoundaryEdgeColor = MeshElementsDisplay->Settings->BoundaryEdgeColor;
+			ToolProperties->DisplayProperties.UVSeamColor = MeshElementsDisplay->Settings->UVSeamColor;
+			ToolProperties->DisplayProperties.NormalSeamColor = MeshElementsDisplay->Settings->NormalSeamColor;
+			ToolProperties->DisplayProperties.TangentSeamColor = MeshElementsDisplay->Settings->TangentSeamColor;
+			ToolProperties->DisplayProperties.ColorSeamColor = MeshElementsDisplay->Settings->ColorSeamColor;
+			ToolProperties->DisplayProperties.DepthBias = MeshElementsDisplay->Settings->DepthBias;
 		}
 		MeshElementsDisplay->Disconnect();
+		MeshElementsDisplay = nullptr;
 	}
 
 	// make sure the tool properties are saved for the next session
@@ -1019,6 +1043,7 @@ void UDataflowEditorVertexAttributePaintTool::Shutdown(EToolShutdownType Shutdow
 
 		ToolProperties->DisplayProperties.ColorRamp.OnColorCurveChangedDelegate.RemoveAll(this);
 		ToolProperties->SaveProperties(this);
+		ToolProperties->SaveConfig();
 	}
 
 	if (MeshSelector)
diff --git a/Engine/Plugins/Experimental/Dataflow/Source/DataflowEditor/Public/DataflowEditorTools/DataflowEditorVertexAttributePaintTool.h b/Engine/Plugins/Experimental/Dataflow/Source/DataflowEditor/Public/DataflowEditorTools/DataflowEditorVertexAttributePaintTool.h
--- a/Engine/Plugins/Experimental/Dataflow/Source/DataflowEditor/Public/DataflowEditorTools/DataflowEditorVertexAttributePaintTool.h
+++ b/Engine/Plugins/Experimental/Dataflow/Source/DataflowEditor/Public/DataflowEditorTools/DataflowEditorVertexAttributePaintTool.h
@@ -210,8 +210,62 @@ struct FDataflowEditorVertexAttributePaintToolDisplayProperties
 
 	UPROPERTY()
 	float MaxValue = 100;
+
+	/** Should mesh boundary edges be shown */
+	UPROPERTY()
+	bool bShowBorders = true;
+
+	/** Should mesh uv seam edges be shown */
+	UPROPERTY()
+	bool bShowUVSeams = true;
+
+	/** Should mesh normal seam edges be shown */
+	UPROPERTY()
+	bool bShowNormalSeams = false;
+
+	/** Should mesh tangent seam edges be shown */
+	UPROPERTY()
+	bool bShowTangentSeams = true;
+
+	/** Should mesh color seam edges be shown */
+	UPROPERTY()
+	bool bShowColorSeams = true;
+
+	/** Multiplier on edge thicknesses */
+	UPROPERTY()
+	float ThicknessScale = 0.5f;
+
+	/** Color of mesh wireframe */
+	UPROPERTY()
+	FColor WireframeColor = FColor(128, 128, 128);
+
+	/** Color of mesh boundary edges */
+	UPROPERTY()
+	FColor BoundaryEdgeColor = FColor(164, 73, 164);
+
+	/** Color of mesh UV seam edges */
+	UPROPERTY()
+	FColor UVSeamColor = FColor(240, 160, 15);
+
+	/** Color of mesh normal seam edges */
+	UPROPERTY()
+	FColor NormalSeamColor = FColor(128, 128, 240);
+
+	/** Color of mesh tangent seam edges */
+	UPROPERTY()
+	FColor TangentSeamColor = FColor(64, 240, 240);
+
+	/** Color of mesh color seam edges */
+	UPROPERTY()
+	FColor ColorSeamColor = FColor(46, 204, 113);
+
+	/** Depth bias used to slightly shift depth of lines */
+	UPROPERTY()
+	float DepthBias = 0.2;
 };
 
+///////////////////////////////////////////////////////////////////////////////////////////////////////////////////////////////////////////////////
+
 USTRUCT()
 struct FDataflowEditorVertexAttributePaintToolMirrorProperties
 {
