Dataflow : add rendering for geometry collection asset output

#Synced-CL 50024027

diff --git a/Engine/Plugins/Experimental/Dataflow/Source/DataflowEditor/Private/DataflowRendering/DataflowGeometryCollectionRenderableType.cpp b/Engine/Plugins/Experimental/Dataflow/Source/DataflowEditor/Private/DataflowRendering/DataflowGeometryCollectionRenderableType.cpp
--- a/Engine/Plugins/Experimental/Dataflow/Source/DataflowEditor/Private/DataflowRendering/DataflowGeometryCollectionRenderableType.cpp
+++ b/Engine/Plugins/Experimental/Dataflow/Source/DataflowEditor/Private/DataflowRendering/DataflowGeometryCollectionRenderableType.cpp
@@ -20,9 +20,12 @@
 #include "GeometryCollection/Facades/CollectionCurveFacade.h"
 #include "GeometryCollection/Facades/CollectionTransformFacade.h"
 #include "GeometryCollection/Facades/CollectionTransformSelectionFacade.h"
+#include "GeometryCollection/GeometryCollectionComponent.h"
+#include "GeometryCollection/GeometryCollectionObject.h"
 
 #include "UObject/ObjectPtr.h"
 
+
 #include "PlanarCut.h" // for ConvertGeometryCollectionToDynamicMesh
 #include "Dataflow/DataflowEditorStyle.h"
 
@@ -477,11 +480,41 @@ namespace UE::Dataflow::Private
 		}
 	};
 
+	////////////////////////////////////////////////////////////////////////////////////////////////////////////////////////////
+
+	class FGeometryCollectionAssetRenderableType : public IRenderableType
+	{
+		UE_DATAFLOW_IRENDERABLE_PRIMARY_TYPE(TObjectPtr<UGeometryCollection>, GeometryCollectionAsset);
+		UE_DATAFLOW_IRENDERABLE_RENDER_GROUP(Asset);
+		UE_DATAFLOW_IRENDERABLE_VIEW_MODE(FDataflowConstruction3DViewMode);
+
+		virtual bool CanRender(const FRenderableTypeInstance& Instance) const
+		{
+			const TObjectPtr<UGeometryCollection> AssetPtr = GetGeometryCollectionAsset(Instance);
+			return AssetPtr && !AssetPtr->IsEmpty();
+		}
+
+		virtual void GetPrimitiveComponents(const FRenderableTypeInstance& Instance, FRenderableComponents& OutComponents) const override
+		{
+			if (const TObjectPtr<UGeometryCollection> AssetPtr = GetGeometryCollectionAsset(Instance))
+			{
+				static FName ComponentName(TEXT("Geometry Collection Asset"));
+				if (UGeometryCollectionComponent* Component = OutComponents.AddNewComponent<UGeometryCollectionComponent>(ComponentName))
+				{
+					Component->SetCastShadow(false);
+					Component->SetSimulatePhysics(false);
+					Component->SetRestCollection(AssetPtr.Get());
+				}
+			}
+		}
+	};
+
 	///////////////////////////////////////////////////////////////////////////////////////////////////////////////
 
 	void RegisterGeometryCollectionRenderableTypes()
 	{
 		UE_DATAFLOW_REGISTER_RENDERABLE_TYPE(FGeometryCollectionSurfaceRenderableType);
 		UE_DATAFLOW_REGISTER_RENDERABLE_TYPE(FGeometryCollectionUvsRenderableType);
+		UE_DATAFLOW_REGISTER_RENDERABLE_TYPE(FGeometryCollectionAssetRenderableType);
 	}
 }
