Dataflow:
- Set collision disabled on Box/BoxArray and Sphere/SphereArray rendering (to avoid generating Physics bodies)

#Synced-CL 50146089

diff --git a/Engine/Plugins/Experimental/Dataflow/Source/DataflowEditor/Private/DataflowRendering/DataflowBoxRenderableType.cpp b/Engine/Plugins/Experimental/Dataflow/Source/DataflowEditor/Private/DataflowRendering/DataflowBoxRenderableType.cpp
--- a/Engine/Plugins/Experimental/Dataflow/Source/DataflowEditor/Private/DataflowRendering/DataflowBoxRenderableType.cpp
+++ b/Engine/Plugins/Experimental/Dataflow/Source/DataflowEditor/Private/DataflowRendering/DataflowBoxRenderableType.cpp
@@ -50,6 +50,7 @@ namespace UE::Dataflow::Private
 					BoxTransform.SetScale3D(Box.GetExtent() / 50.f);
 
 					Component->SetWorldTransform(BoxTransform);
+					Component->SetCollisionEnabled(ECollisionEnabled::NoCollision);
 
 					const UDataflowBoxRenderSettings* Settings = Instance.GetTypedRenderSettings<UDataflowBoxRenderSettings>();
 					if (Settings)
@@ -102,6 +103,8 @@ namespace UE::Dataflow::Private
 						ISMComponent->PreAllocateInstancesMemory(NumBoxes);
 						ISMComponent->SetStaticMesh(BoxMesh);
 
+						ISMComponent->SetCollisionEnabled(ECollisionEnabled::NoCollision);
+
 						FBox ComponentBoundingBox;
 						for (const FBox& Box : BoxArray)
 						{
diff --git a/Engine/Plugins/Experimental/Dataflow/Source/DataflowEditor/Private/DataflowRendering/DataflowSphereRenderableType.cpp b/Engine/Plugins/Experimental/Dataflow/Source/DataflowEditor/Private/DataflowRendering/DataflowSphereRenderableType.cpp
--- a/Engine/Plugins/Experimental/Dataflow/Source/DataflowEditor/Private/DataflowRendering/DataflowSphereRenderableType.cpp
+++ b/Engine/Plugins/Experimental/Dataflow/Source/DataflowEditor/Private/DataflowRendering/DataflowSphereRenderableType.cpp
@@ -49,6 +49,7 @@ namespace UE::Dataflow::Private
 					SphereTransform.SetScale3D(FVector(Sphere.W) / 50.f);
 
 					Component->SetWorldTransform(SphereTransform);
+					Component->SetCollisionEnabled(ECollisionEnabled::NoCollision);
 
 					const UDataflowSphereRenderSettings* Settings = Instance.GetTypedRenderSettings<UDataflowSphereRenderSettings>();
 					if (Settings)
@@ -101,6 +102,8 @@ namespace UE::Dataflow::Private
 						ISMComponent->PreAllocateInstancesMemory(NumSpheres);
 						ISMComponent->SetStaticMesh(SphereMesh);
 
+						ISMComponent->SetCollisionEnabled(ECollisionEnabled::NoCollision);
+
 						FBox ComponentBoundingBox;
 						for (const FSphere& Sphere : SphereArray)
 						{
