[Backout] - CL50148102
#rnx
Original CL Desc
-----------------------------------------------------------------
Dataflow:
- Added rendering support for FDataflowVolume (it renders the voxel centers as points and the outer layer of voxels as boxes) to the new rendering system

#Synced-CL 50147825

#Synced-CL 50148102

diff --git a/Engine/Plugins/Experimental/Dataflow/Source/DataflowEditor/Public/DataflowRendering/DataflowRenderableTypeUtils.h b/Engine/Plugins/Experimental/Dataflow/Source/DataflowEditor/Public/DataflowRendering/DataflowRenderableTypeUtils.h
--- a/Engine/Plugins/Experimental/Dataflow/Source/DataflowEditor/Public/DataflowRendering/DataflowRenderableTypeUtils.h
+++ b/Engine/Plugins/Experimental/Dataflow/Source/DataflowEditor/Public/DataflowRendering/DataflowRenderableTypeUtils.h
@@ -52,7 +52,7 @@ struct FDataflowColorCommonRenderSettings
 	GENERATED_BODY()
 
 public:
-	UE_API FDataflowColorCommonRenderSettings();
+	FDataflowColorCommonRenderSettings();
 
 	/** Display surface as wireframe */
 	UPROPERTY(EditAnywhere, Category = "Color")
@@ -144,7 +144,7 @@ namespace UE::Dataflow
 
 		int32 GetUVChannelFromInstance(const FRenderableTypeInstance& Instance, int32 DefaultValue = 0);
 
-		UE_API FLinearColor GetColor(const int32 InIdx, const FBox& InComponentBoundingBox, const FVector& InCenter, float InSize,
+		FLinearColor GetColor(const int32 InIdx, const FBox& InComponentBoundingBox, const FVector& InCenter, float InSize,
 			const FDataflowColorCommonRenderSettings& InSettings);
 	}
 }
diff --git a/Engine/Plugins/Experimental/DataflowVolumetricPlugin/Source/DataflowVolumeCore/Internal/Dataflow/DataflowVolumeRenderableType.cpp b/Engine/Plugins/Experimental/DataflowVolumetricPlugin/Source/DataflowVolumeCore/Internal/Dataflow/DataflowVolumeRenderableType.cpp
deleted file mode 100644
diff --git a/Engine/Plugins/Experimental/DataflowVolumetricPlugin/Source/DataflowVolumeCore/Internal/Dataflow/DataflowVolumeRenderableType.h b/Engine/Plugins/Experimental/DataflowVolumetricPlugin/Source/DataflowVolumeCore/Internal/Dataflow/DataflowVolumeRenderableType.h
deleted file mode 100644
diff --git a/Engine/Plugins/Experimental/DataflowVolumetricPlugin/Source/DataflowVolumeCore/Private/Dataflow/DataflowVolumeCoreModule.cpp b/Engine/Plugins/Experimental/DataflowVolumetricPlugin/Source/DataflowVolumeCore/Private/Dataflow/DataflowVolumeCoreModule.cpp
--- a/Engine/Plugins/Experimental/DataflowVolumetricPlugin/Source/DataflowVolumeCore/Private/Dataflow/DataflowVolumeCoreModule.cpp
+++ b/Engine/Plugins/Experimental/DataflowVolumetricPlugin/Source/DataflowVolumeCore/Private/Dataflow/DataflowVolumeCoreModule.cpp
@@ -7,7 +7,6 @@
 #include "Modules/ModuleInterface.h"
 #include "Modules/ModuleManager.h"
 #include "Dataflow/OpenVDB.h"
-#include "Dataflow/DataflowVolumeRenderableType.h"
 
 #define LOCTEXT_NAMESPACE "DataflowCore"
 
@@ -15,12 +14,8 @@ void IDataflowVolumeCoreModule::StartupModule()
 {
 	// Initialize OpenVDB
 	openvdb::initialize();
-
-	// Register Volume render type
-	UE::Dataflow::Private::RegisterVolumeRenderableTypes();
 }
 
-
 void IDataflowVolumeCoreModule::ShutdownModule()
 {}
 
