Dataflow paint tool : Make volumetric and plane meshes stay solid in LitWireframe mode

#Synced-CL 50212807

diff --git a/Engine/Plugins/Experimental/Dataflow/Source/DataflowEditor/Private/DataflowEditorTools/DataflowEditorVertexAttributePaintTool.cpp b/Engine/Plugins/Experimental/Dataflow/Source/DataflowEditor/Private/DataflowEditorTools/DataflowEditorVertexAttributePaintTool.cpp
--- a/Engine/Plugins/Experimental/Dataflow/Source/DataflowEditor/Private/DataflowEditorTools/DataflowEditorVertexAttributePaintTool.cpp
+++ b/Engine/Plugins/Experimental/Dataflow/Source/DataflowEditor/Private/DataflowEditorTools/DataflowEditorVertexAttributePaintTool.cpp
@@ -809,12 +809,22 @@ void UDataflowEditorVertexAttributePaintTool::Setup()
 
 }
 
+void UDataflowEditorVertexAttributePaintTool::InitializeIndicator()
+{
+	UMeshSculptToolBase::InitializeIndicator();
+	if (BrushIndicatorMesh)
+	{
+		BrushIndicatorMesh->SetOverrideWireframeRenderMaterial(BrushIndicatorMaterial);
+	}
+}
+
 void UDataflowEditorVertexAttributePaintTool::InitializeMirrorIndicator()
 {
 	// register and spawn brush indicator gizmo
 	GetToolManager()->GetPairedGizmoManager()->RegisterGizmoType(UE::DataflowEditorVertexAttributePaintTool::Private::MirrorBrushGizmoType, NewObject<UBrushStampIndicatorBuilder>());
 	MirrorBrushIndicator = GetToolManager()->GetPairedGizmoManager()->CreateGizmo<UBrushStampIndicator>(UE::DataflowEditorVertexAttributePaintTool::Private::MirrorBrushGizmoType, FString(), this);
 	MirrorBrushIndicatorMesh = MakeBrushIndicatorMesh(this, TargetWorld);
+	MirrorBrushIndicatorMesh->SetOverrideWireframeRenderMaterial(BrushIndicatorMaterial);
 	MirrorBrushIndicator->AttachedComponent = MirrorBrushIndicatorMesh->GetRootComponent();
 	MirrorBrushIndicator->LineThickness = 1.0;
 	MirrorBrushIndicator->bDrawIndicatorLines = true;
@@ -837,7 +847,7 @@ void UDataflowEditorVertexAttributePaintTool::InitializeMirrorIndicator()
 	Gen.Generate();
 	FDynamicMesh3 PlaneMesh(&Gen);
 	MirrorPlaneMesh->UpdatePreview(&PlaneMesh);
-	BrushIndicatorMaterial = ToolSetupUtil::GetDefaultBrushVolumeMaterial(GetToolManager());
+	//BrushIndicatorMaterial = ToolSetupUtil::GetDefaultBrushVolumeMaterial(GetToolManager());
 	if (BrushIndicatorMaterial)
 	{
 		MirrorPlaneMesh->SetMaterial(BrushIndicatorMaterial);
@@ -846,6 +856,7 @@ void UDataflowEditorVertexAttributePaintTool::InitializeMirrorIndicator()
 	// make sure raytracing is disabled on the brush indicator
 	Cast<UDynamicMeshComponent>(MirrorPlaneMesh->GetRootComponent())->SetEnableRaytracing(false);
 	MirrorPlaneMesh->SetShadowsEnabled(false);
+	MirrorPlaneMesh->SetOverrideWireframeRenderMaterial(BrushIndicatorMaterial);
 
 	UpdateMirrorPlaneMesh();
 }
diff --git a/Engine/Plugins/Experimental/Dataflow/Source/DataflowEditor/Public/DataflowEditorTools/DataflowEditorVertexAttributePaintTool.h b/Engine/Plugins/Experimental/Dataflow/Source/DataflowEditor/Public/DataflowEditorTools/DataflowEditorVertexAttributePaintTool.h
--- a/Engine/Plugins/Experimental/Dataflow/Source/DataflowEditor/Public/DataflowEditorTools/DataflowEditorVertexAttributePaintTool.h
+++ b/Engine/Plugins/Experimental/Dataflow/Source/DataflowEditor/Public/DataflowEditorTools/DataflowEditorVertexAttributePaintTool.h
@@ -533,8 +533,10 @@ protected:
 	void LoadBrushConfig(EDataflowEditorToolEditOperation BrushMode);
 
 	virtual bool SharesBrushPropertiesChanges() const override { return false; }
+	virtual void InitializeIndicator() override;
 	// end UMeshSculptToolBase API
 
+
 	void EnableValuePicker(bool bEnable);
 	bool IsValuePickerEnabled() const;
 
