Dataflow paint tool : improve mirror UI
- Allow plane to be hidden in both brush and mesh mode
- Add the option to hide the plane when painting

#Synced-CL 50239131

diff --git a/Engine/Plugins/Experimental/Dataflow/Source/DataflowEditor/Private/DataflowEditorTools/DataflowEditorVertexAttributePaintTool.cpp b/Engine/Plugins/Experimental/Dataflow/Source/DataflowEditor/Private/DataflowEditorTools/DataflowEditorVertexAttributePaintTool.cpp
--- a/Engine/Plugins/Experimental/Dataflow/Source/DataflowEditor/Private/DataflowEditorTools/DataflowEditorVertexAttributePaintTool.cpp
+++ b/Engine/Plugins/Experimental/Dataflow/Source/DataflowEditor/Private/DataflowEditorTools/DataflowEditorVertexAttributePaintTool.cpp
@@ -867,7 +867,9 @@ void UDataflowEditorVertexAttributePaintTool::UpdateMirrorPlaneMesh()
 	{
 		FTransform PlaneMeshTransform = FTransform::Identity;
 
-		MirrorPlaneMesh->SetVisible(IsBrushMirroringEnabled() && IsInBrushMode());
+		const bool bIsBrushMirrorVisible = IsInBrushMode() && IsBrushMirroringEnabled() && IsMirrorPlaneWidgetVisible() && !bIsMirrorHiddenBecauseOfBrushStroke;
+		const bool bIsMeshMirrorVisible = !IsInBrushMode() && IsMirrorPlaneWidgetVisible();
+		MirrorPlaneMesh->SetVisible(bIsBrushMirrorVisible || bIsMeshMirrorVisible);
 
 		if (IsBrushMirroringEnabled())
 		{
@@ -1129,11 +1131,30 @@ bool UDataflowEditorVertexAttributePaintTool::IsBrushMirroringEnabled() const
 	return bBrushMirroringEnabled;
 }
 
+void UDataflowEditorVertexAttributePaintTool::SetMirrorPlaneWidgetVisible(bool bVisible)
+{
+	bMirrorPlaneWidgetVisible = bVisible;
+	UpdateMirrorPlaneMesh();
+}
+
+bool UDataflowEditorVertexAttributePaintTool::IsMirrorPlaneWidgetVisible()
+{
+	return bMirrorPlaneWidgetVisible;
+}
+
 bool UDataflowEditorVertexAttributePaintTool::IsBrushMirroringInObjectSpace() const
 {
 	return ToolProperties ? ToolProperties->MirrorProperties.bObjectSpace : false;
 }
 
+void UDataflowEditorVertexAttributePaintTool::SetMirrorAxis(EAxis::Type InAxis)
+{
+	if (ToolProperties)
+	{
+		ToolProperties->MirrorProperties.MirrorAxis = InAxis;
+		UpdateMirrorPlaneMesh();
+	}
+}
 
 EAxis::Type UDataflowEditorVertexAttributePaintTool::GetMirrorAxis() const
 {
@@ -1465,6 +1486,8 @@ void UDataflowEditorVertexAttributePaintTool::OnBeginStroke(const FRay& WorldRay
 		return;
 	}
 
+	bIsMirrorHiddenBecauseOfBrushStroke = (ToolProperties && ToolProperties->MirrorProperties.bHideOnBrushStroke);
+
 	UpdateBrushPosition(WorldRay);
 
 	if (PaintBrushOpProperties)
@@ -1509,6 +1532,8 @@ void UDataflowEditorVertexAttributePaintTool::OnEndStroke()
 		return;
 	}
 
+	bIsMirrorHiddenBecauseOfBrushStroke = false;
+
 	GetActiveBrushOp()->EndStroke(GetSculptMesh(), LastStamp, VertexROI);
 
 	UpdateVertexColorOverlay(&TriangleROI);
diff --git a/Engine/Plugins/Experimental/Dataflow/Source/DataflowEditor/Private/DataflowEditorTools/DataflowEditorVertexAttributePaintToolDetailCustomization.cpp b/Engine/Plugins/Experimental/Dataflow/Source/DataflowEditor/Private/DataflowEditorTools/DataflowEditorVertexAttributePaintToolDetailCustomization.cpp
--- a/Engine/Plugins/Experimental/Dataflow/Source/DataflowEditor/Private/DataflowEditorTools/DataflowEditorVertexAttributePaintToolDetailCustomization.cpp
+++ b/Engine/Plugins/Experimental/Dataflow/Source/DataflowEditor/Private/DataflowEditorTools/DataflowEditorVertexAttributePaintToolDetailCustomization.cpp
@@ -142,10 +142,13 @@ namespace UE::Dataflow::Editor
 					})
 				.OnValueChanged_Lambda([this](EDataflowEditorToolEditMode Mode)
 					{
-						if (ToolProperties.IsValid())
+						if (Tool.IsValid() && ToolProperties.IsValid())
 						{
 							ToolProperties->EditingMode = Mode;
-							// TODO(ccaillaud) ==>> ToolProperties->OwnerTool->ToggleEditingMode();
+							if (Mode == EDataflowEditorToolEditMode::Mesh)
+							{
+								Tool->SetMirrorPlaneWidgetVisible(false);
+							}
 							if (CurrentDetailBuilder)
 							{
 								CurrentDetailBuilder->ForceRefreshDetails();
@@ -1111,7 +1114,35 @@ namespace UE::Dataflow::Editor
 			]
 			.ValueContent()
 			[
-				MakeMirrorAxisWidget()
+				MakeMirrorShowAndAxisWidget()
+			];
+	}
+
+	void FVertexAttributePaintToolDetailCustomization::AddBrushMirrorHideOnBrushStrokeRow(IDetailCategoryBuilder& BrushCategory) const
+	{
+		BrushCategory.AddCustomRow(LOCTEXT("BrushMirrorHideOnBrushStrokeCategory", "Brush Mirror Hide On Brush Stroke"), false)
+			.NameContent()
+			[
+				SNew(STextBlock)
+					.Text(LOCTEXT("BrushMirrorHideOnBrushStrokeLabel", "Hide Plane When Painting"))
+					.Font(FAppStyle::GetFontStyle(TEXT("PropertyWindow.NormalFont")))
+					.ToolTipText(LOCTEXT("BrushMirrorHideOnBrushStrokeTooltip", "When painting, hide the mirror plane widget in the viewport."))
+			]
+			.ValueContent()
+			[
+				SNew(SCheckBox)
+				.IsChecked_Lambda([ToolProperties = ToolProperties]()
+					{
+						const bool bIsChecked = (ToolProperties.IsValid() && ToolProperties->MirrorProperties.bHideOnBrushStroke);
+						return bIsChecked ? ECheckBoxState::Checked : ECheckBoxState::Unchecked;
+					})
+				.OnCheckStateChanged_Lambda([ToolProperties = ToolProperties](ECheckBoxState InCheckBoxState)
+					{
+						if (ToolProperties.IsValid())
+						{
+							ToolProperties->MirrorProperties.bHideOnBrushStroke = (InCheckBoxState == ECheckBoxState::Checked);
+						}
+					})
 			];
 	}
 
@@ -1123,7 +1154,7 @@ namespace UE::Dataflow::Editor
 				SNew(STextBlock)
 					.Text(LOCTEXT("BrushMirrorObjectSpaceLabel", "Object Space"))
 					.Font(FAppStyle::GetFontStyle(TEXT("PropertyWindow.NormalFont")))
-					.ToolTipText(LOCTEXT("BrushMirrorObjectSpaceTooltip", "When checked, the planes will be centered on the object bounds."))
+					.ToolTipText(LOCTEXT("BrushMirrorObjectSpaceTooltip", "When checked, the mirror plane widget will be centered on the object bounds."))
 			]
 			.ValueContent()
 			[
@@ -1139,6 +1170,8 @@ namespace UE::Dataflow::Editor
 
 			AddBrushMirrorPlaneRow(BrushCategory);
 
+			AddBrushMirrorHideOnBrushStrokeRow(BrushCategory);
+
 			AddBrushMirrorObjectSpaceRow(BrushCategory);
 		}
 	}
@@ -1745,13 +1778,13 @@ namespace UE::Dataflow::Editor
 				"Z: copies weights across the XY plane."))
 			.Value_Lambda([this]()
 				{
-					return ToolProperties.IsValid() ? ToolProperties->MirrorProperties.MirrorAxis : TEnumAsByte<EAxis::Type>(EAxis::X);
+					return Tool.IsValid()? Tool->GetMirrorAxis(): EAxis::X;
 				})
 			.OnValueChanged_Lambda([this](EAxis::Type Mode)
 				{
-					if (ToolProperties.IsValid())
+					if (Tool.IsValid())
 					{
-						ToolProperties->MirrorProperties.MirrorAxis = Mode;
+						Tool->SetMirrorAxis(Mode);
 					}
 				})
 			+ SSegmentedControl<EAxis::Type>::Slot(EAxis::X)
@@ -1762,6 +1795,49 @@ namespace UE::Dataflow::Editor
 			.Text(LOCTEXT("MirrorZLabel", "Z"));
 	}
 
+	TSharedRef<SWidget> FVertexAttributePaintToolDetailCustomization::MakeMirrorShowAndAxisWidget() const
+	{
+		TSharedRef<SWidget> ShowButton = SNew(SCheckBox)
+			.Style(FAppStyle::Get(), "ToggleButtonCheckbox")
+			.ToolTipText(LOCTEXT("MirrorPlaneShowHide", "Toggle the mirror plane visibility in the viewport."))
+			.HAlign(HAlign_Center)
+			.Padding(2)
+			.IsChecked_Lambda([this]()-> ECheckBoxState
+			{
+				return (Tool.IsValid() && Tool->IsMirrorPlaneWidgetVisible())? ECheckBoxState::Checked : ECheckBoxState::Unchecked;
+			})
+			.OnCheckStateChanged_Lambda([this](const ECheckBoxState NewState)
+			{
+				if (Tool.IsValid())
+				{
+					Tool->SetMirrorPlaneWidgetVisible(NewState == ECheckBoxState::Checked);
+				}
+			})
+			[
+				SNew(SImage)
+					.ColorAndOpacity(FSlateColor::UseForeground())
+					.Image(FAppStyle::GetBrush("Level.VisibleIcon16x"))
+			];
+
+		auto AxisWidgetIsEnabled = [this]()
+			{
+				return (Tool.IsValid() && ((Tool->IsInBrushMode() && Tool->IsBrushMirroringEnabled()) || !Tool->IsInBrushMode()));
+			};
+
+		return SNew(SHorizontalBox)
+			.IsEnabled_Lambda(AxisWidgetIsEnabled)
+			+ SHorizontalBox::Slot()
+			.FillWidth(.25f)
+			[
+				ShowButton
+			]
+			+ SHorizontalBox::Slot()
+			.FillWidth(.75f)
+			[
+				MakeMirrorAxisWidget()
+			];
+	}
+
 	TSharedRef<SWidget> FVertexAttributePaintToolDetailCustomization::MakeMirrorDirectionButtonsWidget() const
 	{
 		return SNew(SSegmentedControl<EDataflowEditorToolMirrorDirection>)
@@ -1788,7 +1864,7 @@ namespace UE::Dataflow::Editor
 
 	TSharedRef<SWidget> FVertexAttributePaintToolDetailCustomization::MakeMirrorPlaneWidget(bool bShowDirection) const
 	{
-		TSharedRef<SWidget> XYZButtonsWidget = MakeMirrorAxisWidget();
+		TSharedRef<SWidget> XYZButtonsWidget = MakeMirrorShowAndAxisWidget();
 		TSharedRef<SWidget> MirrorDirectionButtonsWidget =
 			bShowDirection
 			? MakeMirrorDirectionButtonsWidget()
@@ -1809,11 +1885,13 @@ namespace UE::Dataflow::Editor
 			[
 				SNew(SHorizontalBox)
 				+ SHorizontalBox::Slot()
+				.FillWidth(0.5f)
 				[
 					XYZButtonsWidget
 				]
 
 				+ SHorizontalBox::Slot()
+				.FillWidth(0.5f)
 				[
 					MirrorDirectionButtonsWidget
 				]
diff --git a/Engine/Plugins/Experimental/Dataflow/Source/DataflowEditor/Private/DataflowEditorTools/DataflowEditorVertexAttributePaintToolDetailCustomization.h b/Engine/Plugins/Experimental/Dataflow/Source/DataflowEditor/Private/DataflowEditorTools/DataflowEditorVertexAttributePaintToolDetailCustomization.h
--- a/Engine/Plugins/Experimental/Dataflow/Source/DataflowEditor/Private/DataflowEditorTools/DataflowEditorVertexAttributePaintToolDetailCustomization.h
+++ b/Engine/Plugins/Experimental/Dataflow/Source/DataflowEditor/Private/DataflowEditorTools/DataflowEditorVertexAttributePaintToolDetailCustomization.h
@@ -51,6 +51,7 @@ namespace UE::Dataflow::Editor
 		TSharedRef<SWidget> MakeBrushMirrorEnableWidget() const;
 		void AddBrushMirrorEnableRow(IDetailCategoryBuilder& BrushCategory) const;
 		void AddBrushMirrorPlaneRow(IDetailCategoryBuilder& BrushCategory) const;
+		void AddBrushMirrorHideOnBrushStrokeRow(IDetailCategoryBuilder& BrushCategory) const;
 		void AddBrushMirrorObjectSpaceRow(IDetailCategoryBuilder& BrushCategory) const;
 		void AddBrushMirrorUI(IDetailLayoutBuilder& DetailBuilder) const;
 
@@ -73,6 +74,7 @@ namespace UE::Dataflow::Editor
 		void MakeSelectionCopyAndPasteRow(IDetailCategoryBuilder& EditValuesCategory) const;
 		
 		TSharedRef<SWidget> MakeMirrorAxisWidget() const;
+		TSharedRef<SWidget> MakeMirrorShowAndAxisWidget() const;
 		TSharedRef<SWidget> MakeMirrorDirectionButtonsWidget() const;
 		TSharedRef<SWidget> MakeMirrorPlaneWidget(bool bShowDirection) const;
 		TSharedRef<SWidget> MakeMirrorObjectSpaceCheckboxWidget() const;
diff --git a/Engine/Plugins/Experimental/Dataflow/Source/DataflowEditor/Public/DataflowEditorTools/DataflowEditorVertexAttributePaintTool.h b/Engine/Plugins/Experimental/Dataflow/Source/DataflowEditor/Public/DataflowEditorTools/DataflowEditorVertexAttributePaintTool.h
--- a/Engine/Plugins/Experimental/Dataflow/Source/DataflowEditor/Public/DataflowEditorTools/DataflowEditorVertexAttributePaintTool.h
+++ b/Engine/Plugins/Experimental/Dataflow/Source/DataflowEditor/Public/DataflowEditorTools/DataflowEditorVertexAttributePaintTool.h
@@ -275,6 +275,10 @@ struct FDataflowEditorVertexAttributePaintToolMirrorProperties
 	UPROPERTY()
 	bool bObjectSpace = true;
 
+	/* if true, the mirror planes will be hidden when actively painting */
+	UPROPERTY()
+	bool bHideOnBrushStroke = false;
+
 	UPROPERTY()
 	TEnumAsByte<EAxis::Type> MirrorAxis = EAxis::X;
 
@@ -543,8 +547,12 @@ protected:
 	void EnableBrushMirroring(bool bEnable);
 	bool IsBrushMirroringEnabled() const;
 
+	void SetMirrorPlaneWidgetVisible(bool bVisible);
+	bool IsMirrorPlaneWidgetVisible();
+
 	bool IsBrushMirroringInObjectSpace() const;
 
+	void SetMirrorAxis(EAxis::Type InAxis);
 	EAxis::Type GetMirrorAxis() const;
 
 	void InitializeMirrorIndicator();
@@ -640,6 +648,8 @@ protected:
 	bool bPendingPickWeight = false;
 	bool bBrushIsHoverMesh = false;
 	bool bBrushMirroringEnabled = false;
+	bool bMirrorPlaneWidgetVisible = true;
+	bool bIsMirrorHiddenBecauseOfBrushStroke = false;
 
 	TArray<int32> ROITriangleBuffer;
 	TArray<double> ROIWeightValueBuffer;
