Dataflow : attempt to fix crashes related to rednering texture / image in the construction view

#Synced-CL 50239131

diff --git a/Engine/Plugins/Experimental/Dataflow/Source/DataflowEditor/Private/DataflowRendering/DataflowTextureRenderableType.cpp b/Engine/Plugins/Experimental/Dataflow/Source/DataflowEditor/Private/DataflowRendering/DataflowTextureRenderableType.cpp
--- a/Engine/Plugins/Experimental/Dataflow/Source/DataflowEditor/Private/DataflowRendering/DataflowTextureRenderableType.cpp
+++ b/Engine/Plugins/Experimental/Dataflow/Source/DataflowEditor/Private/DataflowRendering/DataflowTextureRenderableType.cpp
@@ -62,9 +62,9 @@ namespace UE::Dataflow::Private
 		return PreviewMaterialInstance;
 	}
 
-	static UTexture2D* CreateTexture2DFromImage(const FImage& InImage)
+	static UTexture2D* CreateTexture2DFromImage(const FImage& InImage, FRenderableComponents& Components)
 	{
-		UTexture2D* NewTexture = NewObject<UTexture2D>();
+		UTexture2D* NewTexture = Components.MakeNewObject<UTexture2D>();
 		if (NewTexture)
 		{
 			// convert to the BGRA8
@@ -144,7 +144,7 @@ namespace UE::Dataflow::Private
 			}
 			else
 			{
-				TextureObject = CreateTexture2DFromImage(Image.GetImage());
+				TextureObject = CreateTexture2DFromImage(Image.GetImage(), OutComponents);
 			}
 			Instance.CacheValue(TextureObject);
 
diff --git a/Engine/Plugins/Experimental/Dataflow/Source/DataflowEditor/Public/DataflowRendering/DataflowRenderableComponents.h b/Engine/Plugins/Experimental/Dataflow/Source/DataflowEditor/Public/DataflowRendering/DataflowRenderableComponents.h
--- a/Engine/Plugins/Experimental/Dataflow/Source/DataflowEditor/Public/DataflowRendering/DataflowRenderableComponents.h
+++ b/Engine/Plugins/Experimental/Dataflow/Source/DataflowEditor/Public/DataflowRendering/DataflowRenderableComponents.h
@@ -42,6 +42,12 @@ namespace UE::Dataflow
 			return NewComponent;
 		}
 
+		template<typename T UE_REQUIRES(TIsDerivedFrom<T, UObject>::Value)>
+		T* MakeNewObject(FName ObjectName = NAME_None)
+		{
+			return NewObject<T>(RootActor);
+		}
+
 		const TArray<UPrimitiveComponent*>& GetComponents() const { return Components; }
 
 	private:
