Dataflow : make sure to disable render state for actors spawned to cache simulation results

#Synced-CL 50239131

diff --git a/Engine/Plugins/Experimental/Dataflow/Source/DataflowEditor/Private/Dataflow/DataflowSimulationGenerator.cpp b/Engine/Plugins/Experimental/Dataflow/Source/DataflowEditor/Private/Dataflow/DataflowSimulationGenerator.cpp
--- a/Engine/Plugins/Experimental/Dataflow/Source/DataflowEditor/Private/Dataflow/DataflowSimulationGenerator.cpp
+++ b/Engine/Plugins/Experimental/Dataflow/Source/DataflowEditor/Private/Dataflow/DataflowSimulationGenerator.cpp
@@ -1,6 +1,8 @@
 // Copyright Epic Games, Inc. All Rights Reserved.
 
 #include "Dataflow/DataflowSimulationGenerator.h"
+
+#include "Components/PrimitiveComponent.h"
 #include "Dataflow/DataflowSimulationManager.h"
 #include "Dataflow/DataflowSimulationControls.h"
 #include "Dataflow/DataflowSimulationUtils.h"
@@ -122,6 +124,15 @@ bool FDataflowTaskManager::AllocateSimulationResource(const FVector2f& TimeRange
 		DataflowManager->SetSimulationEnabled(!SimulationTask->GetTask().bAsyncCaching);
 	}
 
+	if (PreviewActor)
+	{
+		// Prevent the component from having a rendering proxy that can cause unecessary render state invalidation 
+		// and potential issue when calling those from a worker thread
+		PreviewActor->ForEachComponent<UPrimitiveComponent>(true, [](UPrimitiveComponent* Component)
+			{
+				Component->DestroyRenderState_Concurrent();
+			});
+	}
 	return true;
 }
 
