Dataflow : fix icon covering small list item thumbnails

#Synced-CL 50358557

diff --git a/Engine/Plugins/Experimental/Dataflow/Source/DataflowEditor/Private/Dataflow/DataflowEditorModule.cpp b/Engine/Plugins/Experimental/Dataflow/Source/DataflowEditor/Private/Dataflow/DataflowEditorModule.cpp
--- a/Engine/Plugins/Experimental/Dataflow/Source/DataflowEditor/Private/Dataflow/DataflowEditorModule.cpp
+++ b/Engine/Plugins/Experimental/Dataflow/Source/DataflowEditor/Private/Dataflow/DataflowEditorModule.cpp
@@ -51,6 +51,7 @@
 #include "Modules/ModuleManager.h"
 #include "Styling/SlateIconFinder.h"
 #include "Widgets/Images/SImage.h"
+#include "AssetViewWidgets.h"
 
 #define LOCTEXT_NAMESPACE "DataflowEditor"
 
@@ -197,29 +198,51 @@ namespace UE::Dataflow::Private
 			const FSlateBrush* DataflowAssetIconBrush = FSlateIcon(FAppStyle::GetAppStyleSetName(), "GraphEditor.EventGraph_16x").GetIcon();
 
 			// Wrap the icon in an outer background layer to improve visibility on light thumbnails
-			const float Padding = 4.f;
-			const float OuterBoxSize = 16.f;
-			return SNew(SBox)
-				.Padding(Padding, Padding, Padding, Padding)
+			constexpr float Padding = 4.f;
+			constexpr float OuterBoxSize = 16.f;
+			TSharedRef<SWidget> BorderWidget = SNew(SBorder)
+				.BorderImage(FAppStyle::GetBrush("AssetThumbnail", ".AssetThumbnailStatusBar"))
 				[
-					SNew(SBorder)
-					.BorderImage(FAppStyle::GetBrush("AssetThumbnail", ".AssetThumbnailStatusBar"))
+					SNew(SBox)
+					.WidthOverride(OuterBoxSize)
+					.HeightOverride(OuterBoxSize)
+					.HAlign(HAlign_Center)
+					.VAlign(VAlign_Center)
 					[
-						SNew(SBox)
-						.WidthOverride(OuterBoxSize)
-						.HeightOverride(OuterBoxSize)
-						.HAlign(HAlign_Center)
-						.VAlign(VAlign_Center)
-						[
-							SNew(SImage)
-							.Image(DataflowAssetIconBrush)
-							.ColorAndOpacity(bIsDataflowEmbedded 
-								? FSlateColor(FLinearColor::White)
-								: FSlateColor(FLinearColor::Yellow)
-							)
-						]
+						SNew(SImage)
+						.Image(DataflowAssetIconBrush)
+						.ColorAndOpacity(bIsDataflowEmbedded
+							? FSlateColor(FLinearColor::White)
+							: FSlateColor(FLinearColor::Yellow)
+						)
 					]
 				];
+
+			// make sure we do not show the dataflow icon if the thumbnail is a list item icon
+			// to avoid convering the actual thumbnail 
+			auto  VisibilityBasedOnSizeLambda = [BorderWidget]() -> EVisibility
+				{
+					TSharedPtr<SWidget> Parent = BorderWidget->GetParentWidget();
+					while (Parent)
+					{
+						if (Parent->GetType() == SAssetListViewRow::StaticWidgetClass().GetWidgetType())
+						{
+							if (Parent->GetDesiredSize().Y < OuterBoxSize * 3.f)
+							{
+								return EVisibility::Collapsed;
+							}
+						}
+						Parent = Parent->GetParentWidget();
+					}
+					return EVisibility::Visible;
+				};
+
+			return SNew(SBox)
+				.Padding(Padding, Padding, Padding, Padding)
+				.Visibility_Lambda(MoveTemp(VisibilityBasedOnSizeLambda))
+				[
+					BorderWidget
+				];
 		}
 		return SNullWidget::NullWidget;
 	}
