Dataflow : add support for anytype when registering a renderableType

#Synced-CL 50431047

diff --git a/Engine/Plugins/Experimental/Dataflow/Source/DataflowEditor/Private/DataflowRendering/DataflowRenderableTypeRegistry.cpp b/Engine/Plugins/Experimental/Dataflow/Source/DataflowEditor/Private/DataflowRendering/DataflowRenderableTypeRegistry.cpp
--- a/Engine/Plugins/Experimental/Dataflow/Source/DataflowEditor/Private/DataflowRendering/DataflowRenderableTypeRegistry.cpp
+++ b/Engine/Plugins/Experimental/Dataflow/Source/DataflowEditor/Private/DataflowRendering/DataflowRenderableTypeRegistry.cpp
@@ -2,6 +2,7 @@
 
 #include "DataflowRendering/DataflowRenderableTypeRegistry.h"
 
+#include "Dataflow/DataflowAnyTypeRegistry.h"
 #include "DataflowRendering/DataflowRenderableType.h"
 
 namespace UE::Dataflow
@@ -16,8 +17,13 @@ namespace UE::Dataflow
 	{
 		if (RenderableType)
 		{
+			const FName OutputType = RenderableType->GetOutputType();
+			if (FAnyTypesRegistry::IsAnyTypeStatic(OutputType))
+			{
+				RegisteredAnyTypes.Add(OutputType);
+			}
 			RenderableTypesByPrimaryType
-				.FindOrAdd(RenderableType->GetOutputType())
+				.FindOrAdd(OutputType)
 				.Add(RenderableType);
 		}
 	}
@@ -30,6 +36,17 @@ namespace UE::Dataflow
 		{
 			return *RenderableTypes;
 		}
+		// no specific type  regsitered, let's check if there's a anytype compatible
+		for (FName RegisteredAnyType : RegisteredAnyTypes)
+		{
+			if (FAnyTypesRegistry::AreTypesCompatibleStatic(PrimaryType, RegisteredAnyType))
+			{
+				if (const FRenderableTypes* RenderableTypes = RenderableTypesByPrimaryType.Find(RegisteredAnyType))
+				{
+					return *RenderableTypes;
+				}
+			}
+		}
 		return EmptyTypes;
 	}
 }
diff --git a/Engine/Plugins/Experimental/Dataflow/Source/DataflowEditor/Public/DataflowRendering/DataflowRenderableTypeRegistry.h b/Engine/Plugins/Experimental/Dataflow/Source/DataflowEditor/Public/DataflowRendering/DataflowRenderableTypeRegistry.h
--- a/Engine/Plugins/Experimental/Dataflow/Source/DataflowEditor/Public/DataflowRendering/DataflowRenderableTypeRegistry.h
+++ b/Engine/Plugins/Experimental/Dataflow/Source/DataflowEditor/Public/DataflowRendering/DataflowRenderableTypeRegistry.h
@@ -32,6 +32,7 @@ namespace UE::Dataflow
 
 	private:
 		TMap<FName, FRenderableTypes> RenderableTypesByPrimaryType;
+		TSet<FName> RegisteredAnyTypes;
 	};
 }
 
