Dataflow: add Transform Sampler node and better support for struct anytype
- Added the ability to get a instanced struct object (typed or not ) from FContextValue

#Synced-CL 50529101

diff --git a/Engine/Plugins/Experimental/Dataflow/Source/DataflowNodes/Private/Dataflow/DataflowNodesPlugin.cpp b/Engine/Plugins/Experimental/Dataflow/Source/DataflowNodes/Private/Dataflow/DataflowNodesPlugin.cpp
--- a/Engine/Plugins/Experimental/Dataflow/Source/DataflowNodes/Private/Dataflow/DataflowNodesPlugin.cpp
+++ b/Engine/Plugins/Experimental/Dataflow/Source/DataflowNodes/Private/Dataflow/DataflowNodesPlugin.cpp
@@ -21,6 +21,7 @@
 #include "Dataflow/DataflowGetCollectionBoundsNode.h"
 #include "Dataflow/SamplerNodes/DataflowSamplerTypes.h"
 #include "Dataflow/SamplerNodes/DataflowGradientSamplerNode.h"
+#include "Dataflow/SamplerNodes/DataflowTransformSamplerNode.h"
 
 #define LOCTEXT_NAMESPACE "DataflowNodes"
 
@@ -66,6 +67,7 @@ void IDataflowNodesPlugin::StartupModule()
 
 	UE::Dataflow::RegisterValueSamplerToAttributeNodes();
 	UE::Dataflow::RegisterGradientValueSamplerNodes();
+	UE::Dataflow::RegisterTransformValueSamplerNodes();
 
 	UE::Dataflow::DataflowCollectionAttributeKeyNodes();
 	DATAFLOW_NODE_REGISTER_CREATION_FACTORY(FDataflowCollectionAddScalarVertexPropertyNode);
diff --git a/Engine/Plugins/Experimental/Dataflow/Source/DataflowNodes/Private/Dataflow/DataflowValueSamplerToAttributeNode.cpp b/Engine/Plugins/Experimental/Dataflow/Source/DataflowNodes/Private/Dataflow/DataflowValueSamplerToAttributeNode.cpp
--- a/Engine/Plugins/Experimental/Dataflow/Source/DataflowNodes/Private/Dataflow/DataflowValueSamplerToAttributeNode.cpp
+++ b/Engine/Plugins/Experimental/Dataflow/Source/DataflowNodes/Private/Dataflow/DataflowValueSamplerToAttributeNode.cpp
@@ -70,7 +70,7 @@ void FDataflowValueSamplerToAttributeNode::Evaluate(UE::Dataflow::FContext& Cont
 		}
 		else
 		{
-			if (const FDataflowValueSampler* InSampler = GetValue(Context, &Sampler).GetConstStructView().GetPtr<FDataflowValueSampler>())
+			if (const FDataflowValueSampler* InSampler = GetValue(Context, &Sampler).GetTypedStructPtr<FDataflowValueSampler>())
 			{
 				TArrayView<float> Values = MakeArrayView<float>(TargetAttribute.GetData(), TargetAttribute.Num());
 				InSampler->Sample(VertexAttribute->GetConstArray(), Values);
diff --git a/Engine/Plugins/Experimental/Dataflow/Source/DataflowNodes/Private/Dataflow/SamplerNodes/DataflowSamplerTypes.h b/Engine/Plugins/Experimental/Dataflow/Source/DataflowNodes/Private/Dataflow/SamplerNodes/DataflowSamplerTypes.h
--- a/Engine/Plugins/Experimental/Dataflow/Source/DataflowNodes/Private/Dataflow/SamplerNodes/DataflowSamplerTypes.h
+++ b/Engine/Plugins/Experimental/Dataflow/Source/DataflowNodes/Private/Dataflow/SamplerNodes/DataflowSamplerTypes.h
@@ -3,6 +3,7 @@
 #pragma once 
 
 #include "Dataflow/DataflowAnyType.h"
+#include "StructUtils/InstancedStruct.h"
 
 #include "DataflowSamplerTypes.generated.h"
 
diff --git a/Engine/Plugins/Experimental/Dataflow/Source/DataflowNodes/Private/Dataflow/SamplerNodes/DataflowTransformSamplerNode.cpp b/Engine/Plugins/Experimental/Dataflow/Source/DataflowNodes/Private/Dataflow/SamplerNodes/DataflowTransformSamplerNode.cpp
new file mode 100644
--- /dev/null
+++ b/Engine/Plugins/Experimental/Dataflow/Source/DataflowNodes/Private/Dataflow/SamplerNodes/DataflowTransformSamplerNode.cpp
@@ -0,0 +1,47 @@
+// Copyright Epic Games, Inc. All Rights Reserved.
+
+#include "Dataflow/SamplerNodes/DataflowTransformSamplerNode.h"
+#include "Dataflow/DataflowNodeFactory.h"
+
+namespace UE::Dataflow
+{
+	void RegisterTransformValueSamplerNodes()
+	{
+		DATAFLOW_NODE_REGISTER_CREATION_FACTORY(FDataflowTransformValueSamplerNode);
+	}
+}
+
+////////////////////////////////////////////////////////////////////////////////////////////////////////////////////////
+
+void FDataflowTransformValueSampler::Sample(TArrayView<const FVector3f> Positions, TArrayView<float> OutValues) const
+{
+	if (const FDataflowValueSampler* SamplerPtr = Sampler.GetPtr())
+	{
+		TArray<FVector3f> TransformedPositions;
+		TransformedPositions.SetNumUninitialized(Positions.Num());
+		for (int32 Index = 0; Index < Positions.Num(); ++Index)
+		{
+			TransformedPositions[Index] = FVector3f(Transform.InverseTransformPosition(FVector(Positions[Index])));
+		}
+		SamplerPtr->Sample(TransformedPositions, OutValues);
+	}
+}
+
+FDataflowTransformValueSamplerNode::FDataflowTransformValueSamplerNode(const UE::Dataflow::FNodeParameters& InParam, FGuid InGuid)
+	: FDataflowNode(InParam, InGuid)
+{
+	RegisterInputConnection(&TransformSampler.Transform, GET_MEMBER_NAME_CHECKED(FDataflowTransformValueSampler, Transform));
+	RegisterInputConnection(&Sampler);
+	RegisterOutputConnection(&TransformSampler);
+}
+
+void FDataflowTransformValueSamplerNode::Evaluate(UE::Dataflow::FContext& Context, const FDataflowOutput* Out) const
+{
+	if (Out->IsA(&TransformSampler))
+	{
+		FDataflowTransformValueSampler OutTransformSampler;
+		OutTransformSampler.Transform = GetValue(Context, &TransformSampler.Transform);
+		OutTransformSampler.Sampler = GetValue(Context, &Sampler).GetTypedInstancedStruct<FDataflowValueSampler>();
+		SetValue(Context, OutTransformSampler, &TransformSampler);
+	}
+}
diff --git a/Engine/Plugins/Experimental/Dataflow/Source/DataflowNodes/Private/Dataflow/SamplerNodes/DataflowTransformSamplerNode.h b/Engine/Plugins/Experimental/Dataflow/Source/DataflowNodes/Private/Dataflow/SamplerNodes/DataflowTransformSamplerNode.h
new file mode 100644
--- /dev/null
+++ b/Engine/Plugins/Experimental/Dataflow/Source/DataflowNodes/Private/Dataflow/SamplerNodes/DataflowTransformSamplerNode.h
@@ -0,0 +1,55 @@
+// Copyright Epic Games, Inc. All Rights Reserved.
+
+#pragma once 
+
+#include "Dataflow/DataflowNode.h"
+#include "Dataflow/SamplerNodes/DataflowSamplerTypes.h"
+
+#include "DataflowTransformSamplerNode.generated.h"
+
+////////////////////////////////////////////////////////////////////////////////////////////////////////////////////////
+
+USTRUCT()
+struct FDataflowTransformValueSampler : public FDataflowValueSampler
+{
+	GENERATED_BODY();
+
+	virtual ~FDataflowTransformValueSampler() override = default;
+	virtual void Sample(TArrayView<const FVector3f> Positions, TArrayView<float> OutValues) const override;
+
+public:
+	/** Transform to transform the points into a different space */
+	UPROPERTY(EditAnywhere, Category = Transform, meta = (DataflowInput))
+	FTransform Transform;
+
+	/** Sampler to send the transformed point to */
+	UPROPERTY()
+	TInstancedStruct<FDataflowValueSampler> Sampler;
+};
+
+/**
+ * Output a transform sampler 
+ */
+USTRUCT()
+struct FDataflowTransformValueSamplerNode : public FDataflowNode
+{
+	GENERATED_BODY()
+	DATAFLOW_NODE_DEFINE_INTERNAL(FDataflowTransformValueSamplerNode, "Transform Sampler", "Samplers", "Offset Translation Rotation Scale Space points")
+
+public:
+	FDataflowTransformValueSamplerNode(const UE::Dataflow::FNodeParameters& InParam, FGuid InGuid = FGuid::NewGuid());
+
+private:
+	UPROPERTY(meta = (DataflowInput))
+	FDataflowValueSamplerTypes Sampler;
+
+	UPROPERTY(EditAnywhere, Category = Sampler, meta = (DataflowOutput, ShowOnlyInnerProperties))
+	FDataflowTransformValueSampler TransformSampler;
+
+	void Evaluate(UE::Dataflow::FContext& Context, const FDataflowOutput* Out) const override;
+};
+
+namespace UE::Dataflow
+{
+	void RegisterTransformValueSamplerNodes();
+}
\ No newline at end of file
diff --git a/Engine/Source/Runtime/Experimental/Dataflow/Core/Private/Dataflow/DataflowContextCache.cpp b/Engine/Source/Runtime/Experimental/Dataflow/Core/Private/Dataflow/DataflowContextCache.cpp
--- a/Engine/Source/Runtime/Experimental/Dataflow/Core/Private/Dataflow/DataflowContextCache.cpp
+++ b/Engine/Source/Runtime/Experimental/Dataflow/Core/Private/Dataflow/DataflowContextCache.cpp
@@ -351,6 +351,18 @@ namespace UE::Dataflow
 		return {};
 	}
 
+	FInstancedStruct FContextValue::GetInstancedStruct() const
+	{
+		FInstancedStruct OutInstancedStruct;
+		FConstStructView StructView = GetConstStructView();
+		if (StructView.IsValid())
+		{
+			OutInstancedStruct.InitializeAs(StructView.GetScriptStruct(), StructView.GetMemory());
+		}
+		return OutInstancedStruct;
+	}
+
+
 	const FContextCacheElementBase* FContextValue::GetCacheEntry() const
 	{
 		return GetCacheEntryPrivate();
diff --git a/Engine/Source/Runtime/Experimental/Dataflow/Core/Public/Dataflow/DataflowContextCache.h b/Engine/Source/Runtime/Experimental/Dataflow/Core/Public/Dataflow/DataflowContextCache.h
--- a/Engine/Source/Runtime/Experimental/Dataflow/Core/Public/Dataflow/DataflowContextCache.h
+++ b/Engine/Source/Runtime/Experimental/Dataflow/Core/Public/Dataflow/DataflowContextCache.h
@@ -655,6 +655,25 @@ namespace UE::Dataflow
 		DATAFLOWCORE_API FContextValue ToArray() const;
 
 		DATAFLOWCORE_API FConstStructView GetConstStructView() const;
+		DATAFLOWCORE_API FInstancedStruct GetInstancedStruct() const;
+
+		template<typename T>
+		const T* GetTypedStructPtr() const
+		{
+			return GetConstStructView().GetPtr<T>();
+		}
+
+		template<typename T>
+		const TInstancedStruct<T> GetTypedInstancedStruct() const
+		{
+			TInstancedStruct<T> OutInstancedStruct;
+			FConstStructView StructView = GetConstStructView();
+			if (StructView.IsValid())
+			{
+				OutInstancedStruct.InitializeAsScriptStruct(StructView.GetScriptStruct(), StructView.GetMemory());
+			}
+			return OutInstancedStruct;
+		}
 
 		UE_DEPRECATED(5.8, "From release 5.10, this method is going to be private, direct access to the cache element is not authorized and the rest of the API of this class should be used instead to get data from the cache")
 		DATAFLOWCORE_API const FContextCacheElementBase* GetCacheEntry() const;
