Dataflow:
- Added watch function to FDataflowPoints output type

#Synced-CL 50535463

diff --git a/Engine/Plugins/Experimental/Dataflow/Source/DataflowEditor/Private/Dataflow/DataflowSNode.cpp b/Engine/Plugins/Experimental/Dataflow/Source/DataflowEditor/Private/Dataflow/DataflowSNode.cpp
--- a/Engine/Plugins/Experimental/Dataflow/Source/DataflowEditor/Private/Dataflow/DataflowSNode.cpp
+++ b/Engine/Plugins/Experimental/Dataflow/Source/DataflowEditor/Private/Dataflow/DataflowSNode.cpp
@@ -25,6 +25,8 @@
 #include "GraphEditorSettings.h"
 #include "Dataflow/DataflowSchema.h"
 #include "Dataflow/DataflowSPin.h"
+#include "Dataflow/DataflowPoints.h"
+#include "GeometryCollection/Facades/PointsFacade.h"
 
 #define LOCTEXT_NAMESPACE "SDataflowEdNode"
 
@@ -50,7 +52,20 @@ namespace UE::Dataflow::Private
 			{
 				OutString += FString::Printf(TEXT("%s : FBox not set\n"), *Output.GetName().ToString());
 			}
-		
+
+		}
+		else if (Output.GetType() == "FDataflowPoints")
+		{
+			const FDataflowPoints DefaultValue;
+			const FDataflowPoints& Value = Output.ReadValue(Context, DefaultValue);
+
+			GeometryCollection::Facades::FPointsFacade PointFacade = Value.GetPointsFacade();
+
+			const int32 NumPoints = PointFacade.GetNumPoints();
+			const int32 NumAttrs = PointFacade.GetNumAttrs() - 1;
+			FString Attrs = PointFacade.GetAttrsAsString(true);
+
+			OutString += FString::Printf(TEXT("%s : %d point(s) | %d attribute(s): %s\n"), *Output.GetName().ToString(), NumPoints, NumAttrs, *Attrs);
 		}
 		else if (FDataflowSelectionTypePolicy::SupportsTypeStatic(Output.GetType()))
 		{
diff --git a/Engine/Plugins/Experimental/Dataflow/Source/DataflowEditor/Private/DataflowRendering/DataflowRenderableTypeInstance.cpp b/Engine/Plugins/Experimental/Dataflow/Source/DataflowEditor/Private/DataflowRendering/DataflowRenderableTypeInstance.cpp
--- a/Engine/Plugins/Experimental/Dataflow/Source/DataflowEditor/Private/DataflowRendering/DataflowRenderableTypeInstance.cpp
+++ b/Engine/Plugins/Experimental/Dataflow/Source/DataflowEditor/Private/DataflowRendering/DataflowRenderableTypeInstance.cpp
@@ -85,6 +85,29 @@ namespace UE::Dataflow
 		return { TEXT("-- Unknown --") };
 	}
 
+	//
+	// Schema: ComponentName = CustomName + OutputName + NodeName
+	//
+	FName FRenderableTypeInstance::GetComponentName(const FName& InComponentName) const
+	{
+		if (RenderableType)
+		{
+			if (Node)
+			{
+				if (const FDataflowOutput* Output = Node->FindOutput(OutputName))
+				{
+					const FString NodeNamestr = Node->GetName().ToString();
+					const FString OutputNameStr = OutputName.ToString();
+					const FString ComponentNameStr = InComponentName.ToString();
+
+					return FName(FString::Format(TEXT("{0}_{1}_{2}"), { ComponentNameStr, OutputName.ToString(), NodeNamestr }));
+				}
+			}
+		}
+
+		return FName(TEXT("-- Unknown --"));
+	}
+
 	void FRenderableTypeInstance::GetPrimitiveComponents(FRenderableComponents& OutComponents) const
 	{
 		if (RenderableType)
diff --git a/Engine/Plugins/Experimental/Dataflow/Source/DataflowEditor/Public/DataflowRendering/DataflowRenderableTypeInstance.h b/Engine/Plugins/Experimental/Dataflow/Source/DataflowEditor/Public/DataflowRendering/DataflowRenderableTypeInstance.h
--- a/Engine/Plugins/Experimental/Dataflow/Source/DataflowEditor/Public/DataflowRendering/DataflowRenderableTypeInstance.h
+++ b/Engine/Plugins/Experimental/Dataflow/Source/DataflowEditor/Public/DataflowRendering/DataflowRenderableTypeInstance.h
@@ -37,6 +37,9 @@ namespace UE::Dataflow
 		/** Returns the display name of this instance */
 		FString GetDisplayName() const;
 
+		/** Returns the display name of this instance */
+		FName GetComponentName(const FName& InComponentName) const;
+
 		/** Get the corresponding  */
 		void GetPrimitiveComponents(FRenderableComponents& OutComponents) const;
 
