Dataflow:
- Renamed current point rendering to PointArrayRenderable
- Added point rendering to use the new FDataflowPoints

#Synced-CL 50535463

diff --git a/Engine/Plugins/Experimental/Dataflow/Source/DataflowEditor/Private/Dataflow/DataflowEditorModule.cpp b/Engine/Plugins/Experimental/Dataflow/Source/DataflowEditor/Private/Dataflow/DataflowEditorModule.cpp
--- a/Engine/Plugins/Experimental/Dataflow/Source/DataflowEditor/Private/Dataflow/DataflowEditorModule.cpp
+++ b/Engine/Plugins/Experimental/Dataflow/Source/DataflowEditor/Private/Dataflow/DataflowEditorModule.cpp
@@ -42,6 +42,7 @@
 #include "DataflowRendering/DataflowGeometryCollectionRenderableType.h"
 #include "DataflowRendering/DataflowGeometryCollectionProximityRenderableType.h"
 #include "DataflowRendering/DataflowPointRenderableType.h"
+#include "DataflowRendering/DataflowPointArrayRenderableType.h"
 #include "DataflowRendering/DataflowBoxRenderableType.h"
 #include "DataflowRendering/DataflowSphereRenderableType.h"
 #include "DataflowRendering/DataflowMaterialRenderableType.h"
@@ -338,6 +339,7 @@ void FDataflowEditorModule::StartupModule()
 	UE::Dataflow::Private::RegisterGeometryCollectionRenderableTypes();
 	UE::Dataflow::Private::RegisterGeometryCollectionProximityRenderableTypes();
 	UE::Dataflow::Private::RegisterPointRenderableTypes();
+	UE::Dataflow::Private::RegisterPointArrayRenderableTypes();
 	UE::Dataflow::Private::RegisterBoxRenderableTypes();
 	UE::Dataflow::Private::RegisterSphereRenderableTypes();
 	UE::Dataflow::Private::RegisterMaterialRenderableTypes();
diff --git a/Engine/Plugins/Experimental/Dataflow/Source/DataflowEditor/Private/DataflowRendering/DataflowPointArrayRenderableType.cpp b/Engine/Plugins/Experimental/Dataflow/Source/DataflowEditor/Private/DataflowRendering/DataflowPointArrayRenderableType.cpp
new file mode 100644
--- /dev/null
+++ b/Engine/Plugins/Experimental/Dataflow/Source/DataflowEditor/Private/DataflowRendering/DataflowPointArrayRenderableType.cpp
@@ -0,0 +1,76 @@
+// Copyright Epic Games, Inc. All Rights Reserved.
+
+#include "DataflowRendering/DataflowPointArrayRenderableType.h"
+
+#include "Drawing/PointSetComponent.h"
+#include "Components/InstancedStaticMeshComponent.h"
+
+#include "Dataflow/DataflowRenderingViewMode.h"
+
+#include "DataflowRendering/DataflowRenderableComponents.h"
+#include "DataflowRendering/DataflowRenderableType.h"
+#include "DataflowRendering/DataflowRenderableTypeInstance.h"
+#include "DataflowRendering/DataflowRenderableTypeRegistry.h"
+#include "DataflowRendering/DataflowRenderableTypeUtils.h"
+
+#include "Materials/MaterialInterface.h"
+#include "Materials/Material.h"
+
+#include "UObject/ObjectPtr.h"
+#include "GeometryCollection/Facades/PointsFacade.h"
+#include "Dataflow/DataflowPoints.h"
+
+namespace UE::Dataflow::Private
+{
+	///////////////////////////////////////////////////////////////////////////////////////////////////////////////
+
+	class FPointArraySurfaceRenderableType : public IRenderableType
+	{
+		UE_DATAFLOW_IRENDERABLE_PRIMARY_TYPE(TArray<FVector>, PointArray);
+		UE_DATAFLOW_IRENDERABLE_RENDER_GROUP(Points);
+		UE_DATAFLOW_IRENDERABLE_VIEW_MODE(FDataflowConstruction3DViewMode);
+		UE_DATAFLOW_IRENDERABLE_SETTINGS(UDataflowPointArrayRenderSettings);
+
+	public:
+		FPointArraySurfaceRenderableType()
+		{
+			constexpr bool bDepthTested = false;
+			Material = (bDepthTested)
+				? LoadObject<UMaterial>(nullptr, TEXT("/MeshModelingToolsetExp/Materials/PointSetComponentMaterial"))
+				: LoadObject<UMaterial>(nullptr, TEXT("/MeshModelingToolsetExp/Materials/PointSetOverlaidComponentMaterial"));
+		}
+
+		virtual void GetPrimitiveComponents(const FRenderableTypeInstance& Instance, FRenderableComponents& OutComponents) const override
+		{
+			const UDataflowPointArrayRenderSettings* Settings = Instance.GetTypedRenderSettings<UDataflowPointArrayRenderSettings>();
+			if (Settings)
+			{
+				const TArray<FVector>& PointArray = GetPointArray(Instance, {});
+
+				const FName PointComponentName = Instance.GetComponentName(TEXT("Points"));
+
+				const int32 NumPoints = PointArray.Num();
+
+				if (UPointSetComponent* PointComponent = OutComponents.AddNewComponent<UPointSetComponent>(PointComponentName))
+				{
+					PointComponent->ReservePoints(NumPoints);
+					PointComponent->AddPoints(PointArray, Settings->Color, Settings->Size);
+
+					PointComponent->SetPointMaterial(Material);
+				}
+			}
+		}
+
+		UMaterialInterface* Material = nullptr;
+	};
+
+	///////////////////////////////////////////////////////////////////////////////////////////////////////////////
+
+	void RegisterPointArrayRenderableTypes()
+	{
+		UE_DATAFLOW_REGISTER_RENDERABLE_TYPE(FPointArraySurfaceRenderableType);
+	}
+
+	///////////////////////////////////////////////////////////////////////////////////////////////////////////////
+
+}
\ No newline at end of file
diff --git a/Engine/Plugins/Experimental/Dataflow/Source/DataflowEditor/Private/DataflowRendering/DataflowPointArrayRenderableType.h b/Engine/Plugins/Experimental/Dataflow/Source/DataflowEditor/Private/DataflowRendering/DataflowPointArrayRenderableType.h
new file mode 100644
--- /dev/null
+++ b/Engine/Plugins/Experimental/Dataflow/Source/DataflowEditor/Private/DataflowRendering/DataflowPointArrayRenderableType.h
@@ -0,0 +1,26 @@
+// Copyright Epic Games, Inc. All Rights Reserved.
+
+#pragma once
+
+#include "DataflowRendering/DataflowRenderableTypeSettings.h"
+#include "DataflowRendering/DataflowRenderableTypeUtils.h"
+
+#include "DataflowPointArrayRenderableType.generated.h"
+
+UCLASS(MinimalAPI, AutoExpandCategories = "")
+class UDataflowPointArrayRenderSettings : public UDataflowRenderableTypeSettings
+{
+	GENERATED_BODY()
+public:
+	UPROPERTY(EditAnywhere, Category = "Point")
+	float Size = 5.0f;
+
+	UPROPERTY(EditAnywhere, Category = "Point")
+	FColor Color = FColor::Blue;
+};
+
+namespace UE::Dataflow::Private
+{
+	void RegisterPointArrayRenderableTypes();
+}
+
diff --git a/Engine/Plugins/Experimental/Dataflow/Source/DataflowEditor/Private/DataflowRendering/DataflowPointRenderableType.cpp b/Engine/Plugins/Experimental/Dataflow/Source/DataflowEditor/Private/DataflowRendering/DataflowPointRenderableType.cpp
--- a/Engine/Plugins/Experimental/Dataflow/Source/DataflowEditor/Private/DataflowRendering/DataflowPointRenderableType.cpp
+++ b/Engine/Plugins/Experimental/Dataflow/Source/DataflowEditor/Private/DataflowRendering/DataflowPointRenderableType.cpp
@@ -4,6 +4,7 @@
 
 #include "Drawing/PointSetComponent.h"
 #include "Components/InstancedStaticMeshComponent.h"
+#include "Drawing/LineSetComponent.h"
 
 #include "Dataflow/DataflowRenderingViewMode.h"
 
@@ -16,20 +17,35 @@
 #include "Materials/MaterialInterface.h"
 #include "Materials/Material.h"
 
+#include "GeometryCollection/Facades/PointsFacade.h"
+
 #include "UObject/ObjectPtr.h"
+#include "GeometryCollection/Facades/PointsFacade.h"
+#include "Dataflow/DataflowPoints.h"
+
+UDataflowPointRenderSettings::UDataflowPointRenderSettings(const FObjectInitializer& ObjectInitializer) : Super(ObjectInitializer)
+{
+	constexpr bool bOnlyRGB = true;
+	ColorRamp.SetColorAtTime(0.0f, FLinearColor::Blue, bOnlyRGB);
+	ColorRamp.SetColorAtTime(1.0f, FLinearColor::Red, bOnlyRGB);
+
+	LengthColorRamp.SetColorAtTime(0.0f, FLinearColor::Green, bOnlyRGB);
+	LengthColorRamp.SetColorAtTime(1.0f, FLinearColor::Red, bOnlyRGB);
+}
 
 namespace UE::Dataflow::Private
 {
 	///////////////////////////////////////////////////////////////////////////////////////////////////////////////
 
-	class FPointArraySurfaceRenderableType : public IRenderableType
+	class FPointSurfaceRenderableType : public IRenderableType
 	{
-		UE_DATAFLOW_IRENDERABLE_PRIMARY_TYPE(TArray<FVector>, PointArray);
-		UE_DATAFLOW_IRENDERABLE_RENDER_GROUP(Surface);
+		UE_DATAFLOW_IRENDERABLE_PRIMARY_TYPE(FDataflowPoints, PointArray);
+		UE_DATAFLOW_IRENDERABLE_RENDER_GROUP(Points);
 		UE_DATAFLOW_IRENDERABLE_VIEW_MODE(FDataflowConstruction3DViewMode);
+		UE_DATAFLOW_IRENDERABLE_SETTINGS(UDataflowPointRenderSettings);
 
 	public:
-		FPointArraySurfaceRenderableType()
+		FPointSurfaceRenderableType()
 		{
 			constexpr bool bDepthTested = false;
 			Material = (bDepthTested)
@@ -37,22 +53,218 @@ namespace UE::Dataflow::Private
 				: LoadObject<UMaterial>(nullptr, TEXT("/MeshModelingToolsetExp/Materials/PointSetOverlaidComponentMaterial"));
 		}
 
-		virtual void GetPrimitiveComponents(const FRenderableTypeInstance& Instance, FRenderableComponents& OutComponents) const override
+		virtual void GetPrimitiveComponents(const UE::Dataflow::FRenderableTypeInstance& Instance, FRenderableComponents& OutComponents) const override
 		{
-			const TArray<FVector>& PointArray = GetPointArray(Instance, {});
+			const UDataflowPointRenderSettings* Settings = Instance.GetTypedRenderSettings<UDataflowPointRenderSettings>();
+			if (Settings)
+			{
+				const FDataflowPoints& Points = GetPointArray(Instance, {});
 
-			static const FName PointComponentName = TEXT("PointArray");
-			static const FColor PointColor = FColor::Blue;
-			static const float PointSize = 15.f;
+				const FName PointComponentName = Instance.GetComponentName(TEXT("Points"));
 
-			const int32 NumPoints = PointArray.Num();
+				GeometryCollection::Facades::FPointsFacade PointFacadeOutPoints = Points.GetPointsFacade();
 
-			if (UPointSetComponent* PointComponent = OutComponents.AddNewComponent<UPointSetComponent>(PointComponentName))
-			{
-				PointComponent->ReservePoints(NumPoints);
-				PointComponent->AddPoints(PointArray, PointColor, PointSize);
+				const TArray<FVector>& PointArray = PointFacadeOutPoints.GetPointsAsArray();
+
+				const int32 NumPoints = PointArray.Num();
+
+				TArray<FRenderablePoint> RenderablePoints;
+				RenderablePoints.SetNumUninitialized(NumPoints);
+
+				GeometryCollection::Facades::FPointsFacade PointFacade = Points.GetPointsFacade();
+
+				const FName AttributeName = FName(*Settings->Attribute);
+
+				if (PointFacade.HasAttribute(AttributeName))
+				{
+					if (PointFacade.GetAttributeType(AttributeName) == EManagedArrayType::FFloatType)
+					{
+						TArray<float> AttrValues = PointFacade.GetFloatAttributeValues(AttributeName);
+						{
+							for (int32 Idx = 0; Idx < NumPoints; ++Idx)
+							{
+								float AttrValue = AttrValues[Idx];
+
+								if (AttrValue < Settings->Min)
+								{
+									AttrValue = Settings->Min;
+								}
+								else if (AttrValue > Settings->Max)
+								{
+									AttrValue = Settings->Max;
+								}
+
+								FLinearColor LinearColor = Settings->Color;
+								FColor Color = LinearColor.ToFColor(true).WithAlpha(255);
+
+								if (Settings->Max - Settings->Min > UE_SMALL_NUMBER)
+								{
+									float NormAttrValue = (AttrValue - Settings->Min) / (Settings->Max - Settings->Min);
+
+									LinearColor = Settings->ColorRamp.GetLinearColorValue(NormAttrValue);
+									Color = LinearColor.ToFColor(true).WithAlpha(255);
+								}
+
+								RenderablePoints[Idx] = FRenderablePoint(PointArray[Idx], Color, Settings->Size);
+							}
+						}
+					}
+					else if (PointFacade.GetAttributeType(AttributeName) == EManagedArrayType::FInt32Type)
+					{
+						TArray<int32> AttrValues = PointFacade.GetIntAttributeValues(AttributeName);
+						{
+							for (int32 Idx = 0; Idx < NumPoints; ++Idx)
+							{
+								float AttrValue = float(AttrValues[Idx]);
+
+								if (AttrValue < Settings->Min)
+								{
+									AttrValue = Settings->Min;
+								}
+								else if (AttrValue > Settings->Max)
+								{
+									AttrValue = Settings->Max;
+								}
+
+								FLinearColor LinearColor = Settings->Color;
+								FColor Color = LinearColor.ToFColor(true).WithAlpha(255);
+
+								if (Settings->Max - Settings->Min > UE_SMALL_NUMBER)
+								{
+									float NormAttrValue = (AttrValue - Settings->Min) / (Settings->Max - Settings->Min);
+
+									LinearColor = Settings->ColorRamp.GetLinearColorValue(NormAttrValue);
+									Color = LinearColor.ToFColor(true).WithAlpha(255);
+								}
+
+								RenderablePoints[Idx] = FRenderablePoint(PointArray[Idx], Color, Settings->Size);
+							}
+						}
+					}
+					else if (PointFacade.GetAttributeType(AttributeName) == EManagedArrayType::FBoolType)
+					{
+						TBitArray<> AttrValues = PointFacade.GetBoolAttributeValues(AttributeName);
+						{
+							for (int32 Idx = 0; Idx < NumPoints; ++Idx)
+							{
+								float AttrValue = AttrValues[Idx] ? 1.f : 0.f;
+
+								if (AttrValue < Settings->Min)
+								{
+									AttrValue = Settings->Min;
+								}
+								else if (AttrValue > Settings->Max)
+								{
+									AttrValue = Settings->Max;
+								}
+
+								FLinearColor LinearColor = Settings->Color;
+								FColor Color = LinearColor.ToFColor(true).WithAlpha(255);
+
+								if (Settings->Max - Settings->Min > UE_SMALL_NUMBER)
+								{
+									float NormAttrValue = (AttrValue - Settings->Min) / (Settings->Max - Settings->Min);
+
+									LinearColor = Settings->ColorRamp.GetLinearColorValue(NormAttrValue);
+									Color = LinearColor.ToFColor(true).WithAlpha(255);
+								}
+
+								RenderablePoints[Idx] = FRenderablePoint(PointArray[Idx], Color, Settings->Size);
+							}
+						}
+					}
+					else if (PointFacade.GetAttributeType(AttributeName) == EManagedArrayType::FVectorType)
+					{
+						TArray<FVector3f> AttrValues = PointFacade.GetVector3AttributeValues(AttributeName);
+						{
+							for (int32 Idx = 0; Idx < NumPoints; ++Idx)
+							{
+								float AttrValue = AttrValues[Idx].Length();
+
+								if (AttrValue < Settings->Min)
+								{
+									AttrValue = Settings->Min;
+								}
+								else if (AttrValue > Settings->Max)
+								{
+									AttrValue = Settings->Max;
+								}
+
+								FLinearColor LinearColor = Settings->Color;
+								FColor Color = LinearColor.ToFColor(true).WithAlpha(255);
+
+								if (Settings->Max - Settings->Min > UE_SMALL_NUMBER)
+								{
+									float NormAttrValue = (AttrValue - Settings->Min) / (Settings->Max - Settings->Min);
+
+									LinearColor = Settings->ColorRamp.GetLinearColorValue(NormAttrValue);
+									Color = LinearColor.ToFColor(true).WithAlpha(255);
+								}
+
+								RenderablePoints[Idx] = FRenderablePoint(PointArray[Idx], Color, Settings->Size);
+							}
+
+							// Render vectors
+							static const FName LineComponentName = TEXT("Vectors");
+
+							if (ULineSetComponent* LineComponent = OutComponents.AddNewComponent<ULineSetComponent>(LineComponentName))
+							{
+								TArray<FRenderableLine> RenderableLines;
+								RenderableLines.SetNumUninitialized(NumPoints);
+
+								for (int32 Idx = 0; Idx < NumPoints; ++Idx)
+								{
+									float Length = AttrValues[Idx].Length();
+
+									if (Length < Settings->LengthMin)
+									{
+										Length = Settings->LengthMin;
+									}
+									else if (Length > Settings->LengthMax)
+									{
+										Length = Settings->LengthMax;
+									}
+
+									FLinearColor LinearColor = Settings->Color;
+									FColor Color = LinearColor.ToFColor(true).WithAlpha(255);
+
+									if (Settings->LengthMax - Settings->LengthMin > UE_SMALL_NUMBER)
+									{
+										float NormAttrValue = (Length - Settings->LengthMin) / (Settings->LengthMax - Settings->LengthMin);
+
+										LinearColor = Settings->LengthColorRamp.GetLinearColorValue(NormAttrValue);
+										Color = LinearColor.ToFColor(true).WithAlpha(255);
+									}
+
+									RenderableLines[Idx] = FRenderableLine(PointArray[Idx], PointArray[Idx] + Settings->LengthScalar * FVector(AttrValues[Idx]), Color, Settings->LineThickness);
+								}
+
+								LineComponent->AddLines(RenderableLines);
+								LineComponent->SetLineMaterial(Material);
+							}
+						}
+					}
+					else
+					{
+						RenderablePoints.SetNumUninitialized(0);
+					}
+				}
+				else
+				{
+					for (int32 Idx = 0; Idx < NumPoints; ++Idx)
+					{
+						RenderablePoints[Idx] = FRenderablePoint(PointArray[Idx], Settings->Color, Settings->Size);
+					}
+				}
+
+				// Render points
+				if (UPointSetComponent* PointComponent = OutComponents.AddNewComponent<UPointSetComponent>(PointComponentName))
+				{
+					PointComponent->ReservePoints(NumPoints);
+					PointComponent->AddPoints(RenderablePoints);
 
-				PointComponent->SetPointMaterial(Material);
+					PointComponent->SetPointMaterial(Material);
+				}
 			}
 		}
 
@@ -63,7 +275,7 @@ namespace UE::Dataflow::Private
 
 	void RegisterPointRenderableTypes()
 	{
-		UE_DATAFLOW_REGISTER_RENDERABLE_TYPE(FPointArraySurfaceRenderableType);
+		UE_DATAFLOW_REGISTER_RENDERABLE_TYPE(FPointSurfaceRenderableType);
 	}
 
 	///////////////////////////////////////////////////////////////////////////////////////////////////////////////
diff --git a/Engine/Plugins/Experimental/Dataflow/Source/DataflowEditor/Private/DataflowRendering/DataflowPointRenderableType.h b/Engine/Plugins/Experimental/Dataflow/Source/DataflowEditor/Private/DataflowRendering/DataflowPointRenderableType.h
--- a/Engine/Plugins/Experimental/Dataflow/Source/DataflowEditor/Private/DataflowRendering/DataflowPointRenderableType.h
+++ b/Engine/Plugins/Experimental/Dataflow/Source/DataflowEditor/Private/DataflowRendering/DataflowPointRenderableType.h
@@ -2,7 +2,66 @@
 
 #pragma once
 
+#include "DataflowRendering/DataflowRenderableTypeSettings.h"
+#include "DataflowRendering/DataflowRenderableTypeUtils.h"
+#include "Curves/LinearColorRamp.h"
+#include "UObject/ObjectPtr.h"
+
+#include "DataflowPointRenderableType.generated.h"
+
+UCLASS(MinimalAPI, AutoExpandCategories = "")
+class UDataflowPointRenderSettings : public UDataflowRenderableTypeSettings
+{
+	GENERATED_BODY()
+
+	UDataflowPointRenderSettings(const FObjectInitializer& ObjectInitializer = FObjectInitializer::Get());
+
+public:
+	/** Size of displayed points */
+	UPROPERTY(EditAnywhere, Category = "Point", meta = (UIMin = 0.01f, ClampMin = 0.01f))
+	float Size = 5.0f;
+
+	/** Color of displayed points */
+	UPROPERTY(EditAnywhere, Category = "Point")
+	FColor Color = FColor::Blue;
+
+	/** Attribute on points to color the displayed points (Supported types: float, bool, int32, Vector3f */
+	UPROPERTY(EditAnywhere, Category = "Point")
+	FString Attribute = FString("");
+
+	/** Attribute minimum value for coloring (This value represents the minimum color on the color ramp */
+	UPROPERTY(EditAnywhere, Category = "Point")
+	float Min = 0.f;
+
+	/** Attribute maximum value for coloring (This value represents the maximum color on the color ramp */
+	UPROPERTY(EditAnywhere, Category = "Point")
+	float Max = 1.f;
+
+	/** Color ramp to use to display a single value attribute - ignored for vector attributes */
+	UPROPERTY(EditAnywhere, Category = "Point")
+	FLinearColorRamp ColorRamp;
+
+	UPROPERTY(EditAnywhere, Category = "Vector", meta = (UIMin = 0.01f, ClampMin = 0.01f))
+	float LineThickness = 2.f;
+
+	UPROPERTY(EditAnywhere, Category = "Vector", meta = (UIMin = 0.01f, ClampMin = 0.01f))
+	float LengthScalar = 1.f;
+
+	/** Attribute minimum value for coloring (This value represents the minimum color on the color ramp */
+	UPROPERTY(EditAnywhere, Category = "Vector", meta = (UIMin = 0.01f, ClampMin = 0.01f))
+	float LengthMin = 0.f;
+
+	/** Attribute maximum value for coloring (This value represents the maximum color on the color ramp */
+	UPROPERTY(EditAnywhere, Category = "Vector")
+	float LengthMax = 1.f;
+
+	/** Color ramp to use to display a vector value attribute length - ignored for scalar attributes */
+	UPROPERTY(EditAnywhere, Category = "Vector")
+	FLinearColorRamp LengthColorRamp;
+};
+
 namespace UE::Dataflow::Private
 {
 	void RegisterPointRenderableTypes();
 }
+
