Dataflow:
- Fixed missing Super() call

#Synced-CL 50577292

diff --git a/Engine/Plugins/Experimental/Dataflow/Source/DataflowEditor/Private/DataflowRendering/DataflowGeometryCollectionRenderableType.cpp b/Engine/Plugins/Experimental/Dataflow/Source/DataflowEditor/Private/DataflowRendering/DataflowGeometryCollectionRenderableType.cpp
--- a/Engine/Plugins/Experimental/Dataflow/Source/DataflowEditor/Private/DataflowRendering/DataflowGeometryCollectionRenderableType.cpp
+++ b/Engine/Plugins/Experimental/Dataflow/Source/DataflowEditor/Private/DataflowRendering/DataflowGeometryCollectionRenderableType.cpp
@@ -67,7 +67,7 @@ namespace GeometryCollectionSurfaceRenderableType::Private
 }
 
 ////////////////////////////////////////////////////////////////////////////////////////////////////////////////////////////////
-UDataflowGeometryCollectionSurfaceRenderSettings::UDataflowGeometryCollectionSurfaceRenderSettings(const FObjectInitializer& ObjectInitializer)
+UDataflowGeometryCollectionSurfaceRenderSettings::UDataflowGeometryCollectionSurfaceRenderSettings(const FObjectInitializer& ObjectInitializer) : Super(ObjectInitializer)
 {
 	SelectionMaterial = LoadObject<UMaterialInterface>(nullptr, GeometryCollectionSurfaceRenderableType::Private::DefaultDataflowSelectionMaterialName);
 
