[Dataflow] In cases where a dataflow toolkit is editing an EditorContent that itself is an instance of UAssetUserData - prefer the dataflow asset name instead of the dataflow owner name. Asset user data is used as a wrapper in some cases around a dataflow instance and asset rather than being the asset itself. Preferring the asset name gives a better indication in the tab title of what the user is actuallly editing.

#Synced-CL 50592532

diff --git a/Engine/Plugins/Experimental/Dataflow/Source/DataflowEditor/Private/Dataflow/DataflowEditorToolkit.cpp b/Engine/Plugins/Experimental/Dataflow/Source/DataflowEditor/Private/Dataflow/DataflowEditorToolkit.cpp
--- a/Engine/Plugins/Experimental/Dataflow/Source/DataflowEditor/Private/Dataflow/DataflowEditorToolkit.cpp
+++ b/Engine/Plugins/Experimental/Dataflow/Source/DataflowEditor/Private/Dataflow/DataflowEditorToolkit.cpp
@@ -3194,11 +3194,20 @@ FText FDataflowEditorToolkit::GetToolkitName() const
 {
 	if (const TObjectPtr<UDataflowBaseContent>& EditorContent = GetEditorContent())
 	{
-		if (EditorContent->GetDataflowOwner())
+		if(EditorContent->GetDataflowOwner())
 		{
-			return  GetLabelForObject(EditorContent->GetDataflowOwner());
+			// If we have an asset userdata type it is potentially a wrapper around a dataflow asset.
+			// In this case, check for an underlying asset and prefer to use the asset name rather
+			// than the name of the wrapper for the tab label
+			const UAssetUserData* AsUserData = Cast<UAssetUserData>(EditorContent->GetDataflowOwner());
+
+			if(!AsUserData || !EditorContent->GetDataflowAsset())
+			{
+				return  GetLabelForObject(EditorContent->GetDataflowOwner());
+			}
 		}
-		else if (EditorContent->GetDataflowAsset())
+
+		if (EditorContent->GetDataflowAsset())
 		{
 			return  GetLabelForObject(EditorContent->GetDataflowAsset());
 		}
