Dataflow : add support for new DataflowEvaluator
- update Dataflow editor progress bar to support the new evaluator
- make context call the new evelaution code path

#Synced-CL 50604725

diff --git a/Engine/Plugins/Experimental/Dataflow/Source/DataflowEditor/Private/Dataflow/DataflowEditorToolkit.cpp b/Engine/Plugins/Experimental/Dataflow/Source/DataflowEditor/Private/Dataflow/DataflowEditorToolkit.cpp
--- a/Engine/Plugins/Experimental/Dataflow/Source/DataflowEditor/Private/Dataflow/DataflowEditorToolkit.cpp
+++ b/Engine/Plugins/Experimental/Dataflow/Source/DataflowEditor/Private/Dataflow/DataflowEditorToolkit.cpp
@@ -91,6 +91,9 @@ bool bDataflowForceGraphEditable = false;
 FAutoConsoleVariableRef CVARbDataflowForceGraphEditable(TEXT("p.Dataflow.Editor.ForceGraphEditable"), bDataflowForceGraphEditable,
 	TEXT("when enabled all graph will be editable and the Unlock/lock editing buttons will not be visible.[def:false]"));
 
+bool bDataflowEnableNewEvaluation = false;
+FAutoConsoleVariableRef CVARDataflowEnableNewEvaluation(TEXT("p.Dataflow.Editor.EnableNewEvaluation"), bDataflowEnableNewEvaluation,
+	TEXT("When enable, use the new evaluation system using evaluators directly [def:false]"));
 
 const FName FDataflowEditorToolkit::GraphCanvasTabId(TEXT("DataflowEditor_GraphCanvas"));
 const FName FDataflowEditorToolkit::SubGraphCanvasTabId(TEXT("DataflowEditor_SubGraphTab"));
@@ -2302,10 +2305,17 @@ void FDataflowEditorToolkit::EvaluateNode(const FDataflowNode* Node, const FData
 								ToolKit->bViewsNeedRefresh = true;
 							}
 						};
-					
 
-					// If Node is null, the terminal node with the given name will be used instead
-					FDataflowEditorCommands::EvaluateNode(*EditorContent->GetDataflowContext().Get(), *Node, Output, EditorContent->GetTerminalAsset(), InOutTimestamp, OnEvaluationCompleted);
+					if (bDataflowEnableNewEvaluation)
+					{
+						// TODO : maybe this should be in the commands EvaluateNode as there's logic with the timestamp
+						EditorContent->EvaluateGraph(OnEvaluationCompleted);
+					}
+					else
+					{
+						// If Node is null, the terminal node with the given name will be used instead
+						FDataflowEditorCommands::EvaluateNode(*EditorContent->GetDataflowContext().Get(), *Node, Output, EditorContent->GetTerminalAsset(), InOutTimestamp, OnEvaluationCompleted);
+					}
 				}
 			}
 		}
diff --git a/Engine/Plugins/Experimental/Dataflow/Source/DataflowEditor/Private/Dataflow/DataflowGraphEditor.cpp b/Engine/Plugins/Experimental/Dataflow/Source/DataflowEditor/Private/Dataflow/DataflowGraphEditor.cpp
--- a/Engine/Plugins/Experimental/Dataflow/Source/DataflowEditor/Private/Dataflow/DataflowGraphEditor.cpp
+++ b/Engine/Plugins/Experimental/Dataflow/Source/DataflowEditor/Private/Dataflow/DataflowGraphEditor.cpp
@@ -1231,6 +1231,16 @@ void SDataflowGraphEditor::InitEvaluationProgressBar()
 			{
 				bVisible = Context->IsAsyncEvaluating();
 			}
+			if (bVisible == false && DataflowEditor)
+			{
+				if (TObjectPtr<UDataflowBaseContent> EditorContent = DataflowEditor->GetEditorContent())
+				{
+					if (TSharedPtr<UE::Dataflow::IDataflowEvaluator> Evaluator = EditorContent->GetDataflowEvalutor())
+					{
+						bVisible = Evaluator->IsRunning();
+					}
+				}
+			}
 			return bVisible ? EVisibility::SelfHitTestInvisible : EVisibility::Collapsed;
 		};
 
@@ -1250,6 +1260,16 @@ void SDataflowGraphEditor::InitEvaluationProgressBar()
 					NumRunningTasks,
 					NumCompletedTasks
 				);
+				if (NumPendingTasks == 0)
+				{
+					if (TObjectPtr<UDataflowBaseContent> EditorContent = DataflowEditor->GetEditorContent())
+					{
+						if (TSharedPtr<UE::Dataflow::IDataflowEvaluator> Evaluator = EditorContent->GetDataflowEvalutor())
+						{
+							ProgressText = Evaluator->GetProgressMessage();
+						}
+					}
+				}
 			}
 			return ProgressText;
 		};
@@ -1269,6 +1289,17 @@ void SDataflowGraphEditor::InitEvaluationProgressBar()
 				{
 					Percent = static_cast<float>(NumCompletedTasks) / static_cast<float>(TotalTasks);
 				}
+
+				if (NumPendingTasks == 0)
+				{
+					if (TObjectPtr<UDataflowBaseContent> EditorContent = DataflowEditor->GetEditorContent())
+					{
+						if (TSharedPtr<UE::Dataflow::IDataflowEvaluator> Evaluator = EditorContent->GetDataflowEvalutor())
+						{
+							Percent = Evaluator->GetProgress();
+						}
+					}
+				}
 			}
 			return Percent;
 		};
diff --git a/Engine/Source/Runtime/Experimental/Dataflow/Engine/Private/Dataflow/DataflowContextObject.cpp b/Engine/Source/Runtime/Experimental/Dataflow/Engine/Private/Dataflow/DataflowContextObject.cpp
--- a/Engine/Source/Runtime/Experimental/Dataflow/Engine/Private/Dataflow/DataflowContextObject.cpp
+++ b/Engine/Source/Runtime/Experimental/Dataflow/Engine/Private/Dataflow/DataflowContextObject.cpp
@@ -1,10 +1,25 @@
 // Copyright Epic Games, Inc. All Rights Reserved.
 
 #include "Dataflow/DataflowContextObject.h"
+#include "Dataflow/DataflowObjectInterface.h"
+#include "Dataflow/DataflowContextEvaluator.h"
 
 #include UE_INLINE_GENERATED_CPP_BY_NAME(DataflowContextObject)
 
 
+void UDataflowContextObject::EvaluateGraph(UE::Dataflow::FOnPostEvaluationFunction OnPostEvaluation)
+{
+	if (DataflowEvaluator)
+	{
+		DataflowEvaluator->Cancel();
+	}
+	if (DataflowContext)
+	{
+		DataflowEvaluator = DataflowContext->EvaluateGraph<UE::Dataflow::FDataflowTaskGraphEvaluator>(OnPostEvaluation);
+		
+	}
+}
+
 void UDataflowContextObject::AddReferencedObjects(UObject* InThis, FReferenceCollector& Collector)
 {
 	UDataflowContextObject* This = CastChecked<UDataflowContextObject>(InThis);
diff --git a/Engine/Source/Runtime/Experimental/Dataflow/Engine/Public/Dataflow/DataflowContextObject.h b/Engine/Source/Runtime/Experimental/Dataflow/Engine/Public/Dataflow/DataflowContextObject.h
--- a/Engine/Source/Runtime/Experimental/Dataflow/Engine/Public/Dataflow/DataflowContextObject.h
+++ b/Engine/Source/Runtime/Experimental/Dataflow/Engine/Public/Dataflow/DataflowContextObject.h
@@ -63,6 +63,8 @@ public:
 	virtual void SetDataflowAsset(const TObjectPtr<UDataflow>& InAsset) { DataflowGraph = InAsset; }
 	TObjectPtr<UDataflow> GetDataflowAsset() const { return DataflowGraph;  }
 
+	TSharedPtr<UE::Dataflow::IDataflowEvaluator> GetDataflowEvalutor() const { return DataflowEvaluator; }
+
 	/** Collection passing through the currently selected node and a flag indicating whether it's on a node input or output */
 	void SetSelectedCollection(TSharedPtr<const FManagedArrayCollection> InSelectedCollection, bool bInUsingInputCollection) 
 	{ 
@@ -72,6 +74,7 @@ public:
 	const TSharedPtr<const FManagedArrayCollection> GetSelectedCollection() const { return SelectedCollection; }
 	bool IsUsingInputCollection() const { return bUsingInputCollection; }
 
+	UE_API void EvaluateGraph(UE::Dataflow::FOnPostEvaluationFunction OnPostEvaluation);
 
 	//~ UObject interface
 	static UE_API void AddReferencedObjects(UObject* InThis, FReferenceCollector& Collector);
@@ -91,6 +94,9 @@ protected:
 	/** Engine context (data flow owner/asset) to be used for dataflow evaluation */
 	TSharedPtr<UE::Dataflow::FEngineContext> DataflowContext = nullptr;
 
+	/** EValuator being used to evaluate the graph*/
+	TSharedPtr<UE::Dataflow::IDataflowEvaluator> DataflowEvaluator = nullptr;
+
 	/** Dataflow graph for evaluation */
 	UPROPERTY(Transient, SkipSerialization)
 	TObjectPtr<UDataflow> DataflowGraph = nullptr;
