Dataflow:
- Added Instance.GetComponentName() to specify the component name

#Synced-CL 50577292

diff --git a/Engine/Plugins/Experimental/Dataflow/Source/DataflowEditor/Private/DataflowRendering/DataflowBoxRenderableType.cpp b/Engine/Plugins/Experimental/Dataflow/Source/DataflowEditor/Private/DataflowRendering/DataflowBoxRenderableType.cpp
--- a/Engine/Plugins/Experimental/Dataflow/Source/DataflowEditor/Private/DataflowRendering/DataflowBoxRenderableType.cpp
+++ b/Engine/Plugins/Experimental/Dataflow/Source/DataflowEditor/Private/DataflowRendering/DataflowBoxRenderableType.cpp
@@ -36,7 +36,7 @@ namespace UE::Dataflow::Private
 		{
 			const FBox Box = GetBox(Instance, {});
 
-			static const FName BoxComponentName = TEXT("Box");
+			const FName BoxComponentName = Instance.GetComponentName(TEXT("Box"));
 
 			UStaticMesh* BoxMesh = UE::Dataflow::RenderGeometry::GetDataflowBox();
 			if (BoxMesh)
@@ -93,7 +93,7 @@ namespace UE::Dataflow::Private
 
 				if (BoxArray.Num())
 				{
-					static const FName BoxArrayComponentName = TEXT("BoxArray");
+					const FName BoxArrayComponentName = Instance.GetComponentName(TEXT("BoxArray"));
 
 					const int32 NumBoxes = BoxArray.Num();
 
diff --git a/Engine/Plugins/Experimental/Dataflow/Source/DataflowEditor/Private/DataflowRendering/DataflowDynamicMeshRenderableType.cpp b/Engine/Plugins/Experimental/Dataflow/Source/DataflowEditor/Private/DataflowRendering/DataflowDynamicMeshRenderableType.cpp
--- a/Engine/Plugins/Experimental/Dataflow/Source/DataflowEditor/Private/DataflowRendering/DataflowDynamicMeshRenderableType.cpp
+++ b/Engine/Plugins/Experimental/Dataflow/Source/DataflowEditor/Private/DataflowRendering/DataflowDynamicMeshRenderableType.cpp
@@ -32,7 +32,9 @@ namespace UE::Dataflow::Private
 		{
 			if (const TObjectPtr<UDynamicMesh> DynamicMesh = GetDynamicMesh(Instance, nullptr))
 			{
-				if (UDynamicMeshComponent* Component = OutComponents.AddNewComponent<UDynamicMeshComponent>())
+				const FName ComponentName = Instance.GetComponentName(TEXT("Mesh"));
+
+				if (UDynamicMeshComponent* Component = OutComponents.AddNewComponent<UDynamicMeshComponent>(DynamicMesh->GetFName()))
 				{
 					Component->SetDynamicMesh(DynamicMesh);
 
@@ -62,7 +64,9 @@ namespace UE::Dataflow::Private
 			{
 				if (const UE::Geometry::FDynamicMesh3* DynamicMesh = DataflowMesh->GetDynamicMesh())
 				{
-					if (UDynamicMeshComponent* Component = OutComponents.AddNewComponent<UDynamicMeshComponent>())
+					const FName ComponentName = Instance.GetComponentName(TEXT("Mesh"));
+
+					if (UDynamicMeshComponent* Component = OutComponents.AddNewComponent<UDynamicMeshComponent>(DataflowMesh->GetFName()))
 					{
 						UE::Geometry::FDynamicMesh3 MeshToSet = *DynamicMesh;
 						Component->SetMesh(MoveTemp(MeshToSet));
@@ -94,7 +98,9 @@ namespace UE::Dataflow::Private
 			{
 				if (DynamicMesh)
 				{
-					if (UDynamicMeshComponent* Component = OutComponents.AddNewComponent<UDynamicMeshComponent>())
+					const FName ComponentName = Instance.GetComponentName(DynamicMesh->GetFName());
+
+					if (UDynamicMeshComponent* Component = OutComponents.AddNewComponent<UDynamicMeshComponent>(ComponentName))
 					{
 						Component->SetDynamicMesh(DynamicMesh);
 					}
diff --git a/Engine/Plugins/Experimental/Dataflow/Source/DataflowEditor/Private/DataflowRendering/DataflowGeometryCollectionProximityRenderableType.cpp b/Engine/Plugins/Experimental/Dataflow/Source/DataflowEditor/Private/DataflowRendering/DataflowGeometryCollectionProximityRenderableType.cpp
--- a/Engine/Plugins/Experimental/Dataflow/Source/DataflowEditor/Private/DataflowRendering/DataflowGeometryCollectionProximityRenderableType.cpp
+++ b/Engine/Plugins/Experimental/Dataflow/Source/DataflowEditor/Private/DataflowRendering/DataflowGeometryCollectionProximityRenderableType.cpp
@@ -136,7 +136,8 @@ namespace UE::Dataflow::Private
 
 				if (bShowCenters && Centers.Num())
 				{
-					static const FName PointComponentName = TEXT("Connectivity_Points");
+					const FName PointComponentName = Instance.GetComponentName(TEXT("Connectivity_Points"));
+
 					if (UPointSetComponent* PointComponent = OutComponents.AddNewComponent<UPointSetComponent>(PointComponentName))
 					{
 						const float PointSize = Settings ? Settings->PointSize : 15.0f;
@@ -149,7 +150,8 @@ namespace UE::Dataflow::Private
 
 				if (bShowConnections && LineStarts.Num() && LineStarts.Num() == LineEnds.Num())
 				{
-					static const FName LineComponentName = TEXT("Connectivity_Lines");
+					const FName LineComponentName = Instance.GetComponentName(TEXT("Connectivity_Lines"));
+
 					if (ULineSetComponent* LineComponent = OutComponents.AddNewComponent<ULineSetComponent>(LineComponentName))
 					{
 						const float LineThickness = Settings? Settings->LineThickness: 4.0f;
diff --git a/Engine/Plugins/Experimental/Dataflow/Source/DataflowEditor/Private/DataflowRendering/DataflowGeometryCollectionRenderableType.cpp b/Engine/Plugins/Experimental/Dataflow/Source/DataflowEditor/Private/DataflowRendering/DataflowGeometryCollectionRenderableType.cpp
--- a/Engine/Plugins/Experimental/Dataflow/Source/DataflowEditor/Private/DataflowRendering/DataflowGeometryCollectionRenderableType.cpp
+++ b/Engine/Plugins/Experimental/Dataflow/Source/DataflowEditor/Private/DataflowRendering/DataflowGeometryCollectionRenderableType.cpp
@@ -270,7 +270,8 @@ namespace UE::Dataflow::Private
 							}
 						}
 
-						static const FName PointComponentName = TEXT("VertexSelection");
+						const FName PointComponentName = Instance.GetComponentName(TEXT("VertexSelection"));
+
 						if (UPointSetComponent* PointComponent = OutComponents.AddNewComponent<UPointSetComponent>(PointComponentName))
 						{
 							constexpr float PointSize = 8.f;
@@ -498,7 +499,8 @@ namespace UE::Dataflow::Private
 		{
 			if (const TObjectPtr<UGeometryCollection> AssetPtr = GetGeometryCollectionAsset(Instance))
 			{
-				static FName ComponentName(TEXT("Geometry Collection Asset"));
+				const FName ComponentName = Instance.GetComponentName(TEXT("Geometry Collection Asset"));
+
 				if (UGeometryCollectionComponent* Component = OutComponents.AddNewComponent<UGeometryCollectionComponent>(ComponentName))
 				{
 					Component->SetCastShadow(false);
diff --git a/Engine/Plugins/Experimental/Dataflow/Source/DataflowEditor/Private/DataflowRendering/DataflowMaterialRenderableType.cpp b/Engine/Plugins/Experimental/Dataflow/Source/DataflowEditor/Private/DataflowRendering/DataflowMaterialRenderableType.cpp
--- a/Engine/Plugins/Experimental/Dataflow/Source/DataflowEditor/Private/DataflowRendering/DataflowMaterialRenderableType.cpp
+++ b/Engine/Plugins/Experimental/Dataflow/Source/DataflowEditor/Private/DataflowRendering/DataflowMaterialRenderableType.cpp
@@ -38,7 +38,8 @@ namespace UE::Dataflow::Private
 		{
 			const TObjectPtr<UMaterialInterface> Material = GetMaterial(Instance, {});
 
-			static const FName MaterialComponentName = TEXT("Material");
+			const FName MaterialComponentName = Instance.GetComponentName(TEXT("Material"));
+
 			if (UStaticMeshComponent* Component = OutComponents.AddNewComponent<UStaticMeshComponent>(MaterialComponentName))
 			{
 				//LoadMaterialPreviewMesh
diff --git a/Engine/Plugins/Experimental/Dataflow/Source/DataflowEditor/Private/DataflowRendering/DataflowPointRenderableType.cpp b/Engine/Plugins/Experimental/Dataflow/Source/DataflowEditor/Private/DataflowRendering/DataflowPointRenderableType.cpp
--- a/Engine/Plugins/Experimental/Dataflow/Source/DataflowEditor/Private/DataflowRendering/DataflowPointRenderableType.cpp
+++ b/Engine/Plugins/Experimental/Dataflow/Source/DataflowEditor/Private/DataflowRendering/DataflowPointRenderableType.cpp
@@ -53,7 +53,7 @@ namespace UE::Dataflow::Private
 				: LoadObject<UMaterial>(nullptr, TEXT("/MeshModelingToolsetExp/Materials/PointSetOverlaidComponentMaterial"));
 		}
 
-		virtual void GetPrimitiveComponents(const UE::Dataflow::FRenderableTypeInstance& Instance, FRenderableComponents& OutComponents) const override
+		virtual void GetPrimitiveComponents(const FRenderableTypeInstance& Instance, FRenderableComponents& OutComponents) const override
 		{
 			const UDataflowPointRenderSettings* Settings = Instance.GetTypedRenderSettings<UDataflowPointRenderSettings>();
 			if (Settings)
diff --git a/Engine/Plugins/Experimental/Dataflow/Source/DataflowEditor/Private/DataflowRendering/DataflowSkeletalMeshRenderableType.cpp b/Engine/Plugins/Experimental/Dataflow/Source/DataflowEditor/Private/DataflowRendering/DataflowSkeletalMeshRenderableType.cpp
--- a/Engine/Plugins/Experimental/Dataflow/Source/DataflowEditor/Private/DataflowRendering/DataflowSkeletalMeshRenderableType.cpp
+++ b/Engine/Plugins/Experimental/Dataflow/Source/DataflowEditor/Private/DataflowRendering/DataflowSkeletalMeshRenderableType.cpp
@@ -31,7 +31,9 @@ namespace UE::Dataflow::Private
 		{
 			if (const TObjectPtr<USkeletalMesh> SkeletalMesh = GetSkeletalMesh(Instance, nullptr))
 			{
-				if (USkeletalMeshComponent* Component = OutComponents.AddNewComponent<USkeletalMeshComponent>(SkeletalMesh->GetFName()))
+				const FName ComponentName = Instance.GetComponentName(SkeletalMesh->GetFName());
+
+				if (USkeletalMeshComponent* Component = OutComponents.AddNewComponent<USkeletalMeshComponent>(ComponentName))
 				{
 					Component->SetSkeletalMesh(SkeletalMesh);
 				}
diff --git a/Engine/Plugins/Experimental/Dataflow/Source/DataflowEditor/Private/DataflowRendering/DataflowSphereRenderableType.cpp b/Engine/Plugins/Experimental/Dataflow/Source/DataflowEditor/Private/DataflowRendering/DataflowSphereRenderableType.cpp
--- a/Engine/Plugins/Experimental/Dataflow/Source/DataflowEditor/Private/DataflowRendering/DataflowSphereRenderableType.cpp
+++ b/Engine/Plugins/Experimental/Dataflow/Source/DataflowEditor/Private/DataflowRendering/DataflowSphereRenderableType.cpp
@@ -35,7 +35,7 @@ namespace UE::Dataflow::Private
 		{
 			const FSphere Sphere = GetSphere(Instance, {});
 
-			static const FName SphereComponentName = TEXT("Sphere");
+			const FName SphereComponentName = Instance.GetComponentName(TEXT("Sphere"));
 
 			UStaticMesh* SphereMesh = UE::Dataflow::RenderGeometry::GetDataflowSphere();
 			if (SphereMesh)
@@ -92,7 +92,7 @@ namespace UE::Dataflow::Private
 
 				if (SphereArray.Num())
 				{
-					static const FName SphereArrayComponentName = TEXT("SphereArray");
+					const FName SphereArrayComponentName = Instance.GetComponentName(TEXT("SphereArray"));
 
 					const int32 NumSpheres = SphereArray.Num();
 
diff --git a/Engine/Plugins/Experimental/Dataflow/Source/DataflowEditor/Private/DataflowRendering/DataflowStaticMeshRenderableType.cpp b/Engine/Plugins/Experimental/Dataflow/Source/DataflowEditor/Private/DataflowRendering/DataflowStaticMeshRenderableType.cpp
--- a/Engine/Plugins/Experimental/Dataflow/Source/DataflowEditor/Private/DataflowRendering/DataflowStaticMeshRenderableType.cpp
+++ b/Engine/Plugins/Experimental/Dataflow/Source/DataflowEditor/Private/DataflowRendering/DataflowStaticMeshRenderableType.cpp
@@ -30,7 +30,9 @@ namespace UE::Dataflow::Private
 		{
 			if (const TObjectPtr<UStaticMesh> StaticMesh = GetStaticMesh(Instance, nullptr))
 			{
-				if (UStaticMeshComponent* Component = OutComponents.AddNewComponent<UStaticMeshComponent>(StaticMesh->GetFName()))
+				const FName ComponentName = Instance.GetComponentName(StaticMesh->GetFName());
+
+				if (UStaticMeshComponent* Component = OutComponents.AddNewComponent<UStaticMeshComponent>(ComponentName))
 				{
 					Component->SetStaticMesh(StaticMesh);
 				}
diff --git a/Engine/Plugins/Experimental/Dataflow/Source/DataflowEditor/Private/DataflowRendering/DataflowTextureRenderableType.cpp b/Engine/Plugins/Experimental/Dataflow/Source/DataflowEditor/Private/DataflowRendering/DataflowTextureRenderableType.cpp
--- a/Engine/Plugins/Experimental/Dataflow/Source/DataflowEditor/Private/DataflowRendering/DataflowTextureRenderableType.cpp
+++ b/Engine/Plugins/Experimental/Dataflow/Source/DataflowEditor/Private/DataflowRendering/DataflowTextureRenderableType.cpp
@@ -104,7 +104,8 @@ namespace UE::Dataflow::Private
 			// this allow us to display multiple images at the same time without having them overlapping 
 			const int32 NumComponents = OutComponents.GetComponents().Num();
 
-			static const FName ComponentName = TEXT("Texture");
+			const FName ComponentName = Instance.GetComponentName(TEXT("Texture"));
+
 			if (UStaticMeshComponent* Component = OutComponents.AddNewComponent<UStaticMeshComponent>(ComponentName))
 			{
 				if (Texture)
diff --git a/Engine/Plugins/Experimental/Dataflow/Source/DataflowEditor/Public/DataflowRendering/DataflowRenderableTypeInstance.h b/Engine/Plugins/Experimental/Dataflow/Source/DataflowEditor/Public/DataflowRendering/DataflowRenderableTypeInstance.h
--- a/Engine/Plugins/Experimental/Dataflow/Source/DataflowEditor/Public/DataflowRendering/DataflowRenderableTypeInstance.h
+++ b/Engine/Plugins/Experimental/Dataflow/Source/DataflowEditor/Public/DataflowRendering/DataflowRenderableTypeInstance.h
@@ -38,7 +38,7 @@ namespace UE::Dataflow
 		FString GetDisplayName() const;
 
 		/** Returns the display name of this instance */
-		FName GetComponentName(const FName& InComponentName) const;
+		UE_API FName GetComponentName(const FName& InComponentName) const;
 
 		/** Get the corresponding  */
 		void GetPrimitiveComponents(FRenderableComponents& OutComponents) const;
