Dataflow : change the value sampler to float sampler and do not use anytype
- Update the associated node accordingly

#Synced-CL 50604725

#Synced-CL 50604725

diff --git a/Engine/Plugins/Experimental/Dataflow/Source/DataflowNodes/Private/Dataflow/DataflowNodesPlugin.cpp b/Engine/Plugins/Experimental/Dataflow/Source/DataflowNodes/Private/Dataflow/DataflowNodesPlugin.cpp
--- a/Engine/Plugins/Experimental/Dataflow/Source/DataflowNodes/Private/Dataflow/DataflowNodesPlugin.cpp
+++ b/Engine/Plugins/Experimental/Dataflow/Source/DataflowNodes/Private/Dataflow/DataflowNodesPlugin.cpp
@@ -17,9 +17,8 @@
 #include "Dataflow/DataflowTextureToAttributeNode.h"
 #include "Dataflow/DataflowVertexColorToAttributeNode.h"
 #include "Dataflow/DataflowVisualizeAttributeNode.h"
-#include "Dataflow/DataflowValueSamplerToAttributeNode.h"
+#include "Dataflow/DataflowSamplerToAttributeNode.h"
 #include "Dataflow/DataflowGetCollectionBoundsNode.h"
-#include "Dataflow/SamplerNodes/DataflowSamplerTypes.h"
 #include "Dataflow/SamplerNodes/DataflowGradientSamplerNode.h"
 #include "Dataflow/SamplerNodes/DataflowTransformSamplerNode.h"
 
@@ -53,8 +52,6 @@ const FName FGeometryCollectionAddScalarVertexPropertyCallbacks::Name = FName("F
 
 void IDataflowNodesPlugin::StartupModule()
 {
-	UE::Dataflow::RegisterValueSamplerTypes();
-
 	UE::Dataflow::RegisterSkeletalMeshNodes();
 	UE::Dataflow::RegisterStaticMeshNodes();
 	UE::Dataflow::RegisterSelectionNodes();
@@ -65,9 +62,9 @@ void IDataflowNodesPlugin::StartupModule()
 	UE::Dataflow::RegisterVertexColorToAttributeNodes();
 	UE::Dataflow::RegisterVisualizeAttributeNodes();
 
-	UE::Dataflow::RegisterValueSamplerToAttributeNodes();
-	UE::Dataflow::RegisterGradientValueSamplerNodes();
-	UE::Dataflow::RegisterTransformValueSamplerNodes();
+	UE::Dataflow::RegisterSamplerToAttributeNodes();
+	UE::Dataflow::RegisterGradientSamplerNodes();
+	UE::Dataflow::RegisterTransformSamplerNodes();
 
 	UE::Dataflow::DataflowCollectionAttributeKeyNodes();
 	DATAFLOW_NODE_REGISTER_CREATION_FACTORY(FDataflowCollectionAddScalarVertexPropertyNode);
diff --git a/Engine/Plugins/Experimental/Dataflow/Source/DataflowNodes/Private/Dataflow/DataflowSamplerToAttributeNode.cpp b/Engine/Plugins/Experimental/Dataflow/Source/DataflowNodes/Private/Dataflow/DataflowSamplerToAttributeNode.cpp
new file mode 100644
--- /dev/null
+++ b/Engine/Plugins/Experimental/Dataflow/Source/DataflowNodes/Private/Dataflow/DataflowSamplerToAttributeNode.cpp
@@ -0,0 +1,92 @@
+// Copyright Epic Games, Inc. All Rights Reserved.
+
+#include "Dataflow/DataflowSamplerToAttributeNode.h"
+
+#include "Dataflow/DataflowNodeFactory.h"
+#include "GeometryCollection/GeometryCollection.h"
+#include "GeometryCollection/Facades/CollectionMeshFacade.h"
+
+#include UE_INLINE_GENERATED_CPP_BY_NAME(DataflowSamplerToAttributeNode)
+
+#define LOCTEXT_NAMESPACE "DataflowSamplerToAttributeNode"
+
+namespace UE::Dataflow
+{
+	void RegisterSamplerToAttributeNodes()
+	{
+		DATAFLOW_NODE_REGISTER_CREATION_FACTORY(FDataflowFloatSamplerToAttributeNode);
+	}
+};
+
+FDataflowFloatSamplerToAttributeNode::FDataflowFloatSamplerToAttributeNode(const UE::Dataflow::FNodeParameters& InParam, FGuid InGuid)
+	: FDataflowNode(InParam, InGuid)
+{
+	RegisterInputConnection(&Collection);
+	RegisterInputConnection(&Sampler);
+	RegisterInputConnection(&AttributeName);
+	RegisterOutputConnection(&Collection, &Collection);
+}
+
+void FDataflowFloatSamplerToAttributeNode::Evaluate(UE::Dataflow::FContext& Context, const FDataflowOutput* Out) const
+{
+	if (Out->IsA(&Collection))
+	{
+		const FName InAttributeName = FName(GetValue(Context, &AttributeName, AttributeName));
+		if (InAttributeName.IsNone())
+		{
+			Context.Error(LOCTEXT("NoAttributeName_Msg", "Attribute name is empty, data will not be set"));
+			SafeForwardInput(Context, &Collection, &Collection);
+			return;
+		}
+
+		FManagedArrayCollection InCollection = GetValue(Context, &Collection, Collection);
+
+		const FName InVertexGroupName = VertexGroup.Name;
+		if (!InCollection.HasGroup(InVertexGroupName))
+		{
+			Context.Error(LOCTEXT("InvalidVertexGroup_Msg", "Vertex group does not exists, data will not be set"));
+			SafeForwardInput(Context, &Collection, &Collection);
+			return;
+		}
+
+		// Create or get the new attribute 
+		TManagedArray<float>& TargetAttribute = InCollection.AddAttribute<float>(InAttributeName, InVertexGroupName);
+		TargetAttribute.Fill(0);
+
+		const TManagedArray<FVector3f>* VertexAttribute = InCollection.FindAttribute<FVector3f>(FGeometryCollection::VertexPositionAttribute, InVertexGroupName);
+		if (!VertexAttribute)
+		{
+			// cannot really include cloth here so we need to directly search for the vertex position attribute in the render vertices ( eventually we may converge to use only geometry collections )
+			VertexAttribute = InCollection.FindAttribute<FVector3f>("RenderPosition", InVertexGroupName);
+		}
+		if (!VertexAttribute)
+		{
+			// cannot really include cloth here so we need to directly search for the vertex position  attribute in the render vertices ( eventually we may converge to use only geometry collections )
+			VertexAttribute = InCollection.FindAttribute<FVector3f>("SimPosition3D", InVertexGroupName);
+		}
+		if (!VertexAttribute)
+		{
+			Context.Warning(LOCTEXT("NoVertexPosition_Msg", "No vertex position could be found in this collection, no data will be set and the attribute will be inititialized to zeros"));
+		}
+		else
+		{
+			const FDataflowFloatSampler InSampler = GetValue(Context, &Sampler);
+			TArrayView<float> Values = MakeArrayView<float>(TargetAttribute.GetData(), TargetAttribute.Num());
+			InSampler.Sample(VertexAttribute->GetConstArray(), Values);
+		}
+		
+		SetValue(Context, MoveTemp(InCollection), &Collection);
+	}
+}
+
+FDataflowNode::FAttributeKey FDataflowFloatSamplerToAttributeNode::GetVertexAttributeToVisualize(UE::Dataflow::FContext& Context, const FName OutputName, const FName RenderGroup) const
+{
+	return FAttributeKey
+	{
+		.AttributeName = FName(GetValue(Context, &AttributeName, AttributeName)),
+		.GroupName = VertexGroup.Name,
+	};
+}
+
+#undef LOCTEXT_NAMESPACE 
+
diff --git a/Engine/Plugins/Experimental/Dataflow/Source/DataflowNodes/Private/Dataflow/DataflowSamplerToAttributeNode.h b/Engine/Plugins/Experimental/Dataflow/Source/DataflowNodes/Private/Dataflow/DataflowSamplerToAttributeNode.h
new file mode 100644
--- /dev/null
+++ b/Engine/Plugins/Experimental/Dataflow/Source/DataflowNodes/Private/Dataflow/DataflowSamplerToAttributeNode.h
@@ -0,0 +1,49 @@
+// Copyright Epic Games, Inc. All Rights Reserved.
+
+#pragma once 
+
+#include "Dataflow/DataflowNode.h"
+#include "GeometryCollection/ManagedArrayCollection.h"
+#include "Dataflow/DataflowCollectionAddScalarVertexPropertyNode.h"
+#include "Dataflow/SamplerNodes/DataflowSamplerTypes.h"
+
+#include "DataflowSamplerToAttributeNode.generated.h"
+
+/**
+* Use a foat sampler to generate the value of an attribute
+*/
+USTRUCT()
+struct FDataflowFloatSamplerToAttributeNode : public FDataflowNode
+{
+	GENERATED_USTRUCT_BODY()
+	DATAFLOW_NODE_DEFINE_INTERNAL(FDataflowFloatSamplerToAttributeNode, "FloatSamplerToAttribute", "Dataflow", "Transfer Generate")
+
+public:
+	FDataflowFloatSamplerToAttributeNode(const UE::Dataflow::FNodeParameters& InParam, FGuid InGuid = FGuid::NewGuid());
+
+private:
+	/** Collection to set the attribute data to */
+	UPROPERTY(meta = (DataflowInput, DataflowOutput, DataflowPassthrough = Collection))
+	FManagedArrayCollection Collection;
+
+	/** Sampler to use to set the value of the attribute */
+	UPROPERTY(EditAnywhere, Category = Source, meta = (DataflowInput))
+	FDataflowFloatSampler Sampler;
+
+	/** Name of the attribute to set the value on */
+	UPROPERTY(EditAnywhere, Category = Target, meta = (DataflowInput))
+	FString AttributeName;
+
+	/** Vertex group to use */
+	UPROPERTY(EditAnywhere, Category = Target)
+	FScalarVertexPropertyGroup VertexGroup;
+
+	virtual void Evaluate(UE::Dataflow::FContext& Context, const FDataflowOutput* Out) const override;
+	virtual FAttributeKey GetVertexAttributeToVisualize(UE::Dataflow::FContext& Context, const FName OutputName, const FName RenderGroup) const override;
+};
+
+namespace UE::Dataflow
+{
+	void RegisterSamplerToAttributeNodes();
+}
+
diff --git a/Engine/Plugins/Experimental/Dataflow/Source/DataflowNodes/Private/Dataflow/DataflowValueSamplerToAttributeNode.cpp b/Engine/Plugins/Experimental/Dataflow/Source/DataflowNodes/Private/Dataflow/DataflowValueSamplerToAttributeNode.cpp
deleted file mode 100644
diff --git a/Engine/Plugins/Experimental/Dataflow/Source/DataflowNodes/Private/Dataflow/DataflowValueSamplerToAttributeNode.h b/Engine/Plugins/Experimental/Dataflow/Source/DataflowNodes/Private/Dataflow/DataflowValueSamplerToAttributeNode.h
deleted file mode 100644
diff --git a/Engine/Plugins/Experimental/Dataflow/Source/DataflowNodes/Private/Dataflow/SamplerNodes/DataflowGradientSamplerNode.cpp b/Engine/Plugins/Experimental/Dataflow/Source/DataflowNodes/Private/Dataflow/SamplerNodes/DataflowGradientSamplerNode.cpp
--- a/Engine/Plugins/Experimental/Dataflow/Source/DataflowNodes/Private/Dataflow/SamplerNodes/DataflowGradientSamplerNode.cpp
+++ b/Engine/Plugins/Experimental/Dataflow/Source/DataflowNodes/Private/Dataflow/SamplerNodes/DataflowGradientSamplerNode.cpp
@@ -5,17 +5,17 @@
 
 namespace UE::Dataflow
 {
-	void RegisterGradientValueSamplerNodes()
+	void RegisterGradientSamplerNodes()
 	{
-		DATAFLOW_NODE_REGISTER_CREATION_FACTORY(FDataflowLinearGradientValueSamplerNode);
-		DATAFLOW_NODE_REGISTER_CREATION_FACTORY(FDataflowLinearGradientFromBoxValueSamplerNode);
-		DATAFLOW_NODE_REGISTER_CREATION_FACTORY(FDataflowRadialGradientValueSamplerNode);
+		DATAFLOW_NODE_REGISTER_CREATION_FACTORY(FDataflowLinearGradientFloatSamplerNode);
+		DATAFLOW_NODE_REGISTER_CREATION_FACTORY(FDataflowLinearGradientFromBoxFloatSamplerNode);
+		DATAFLOW_NODE_REGISTER_CREATION_FACTORY(FDataflowRadialGradientFloatSamplerNode);
 	}
 }
 
 ////////////////////////////////////////////////////////////////////////////////////////////////////////////////////////
 
-void FDataflowLinearGradientValueSampler::Sample(TArrayView<const FVector3f> Positions, TArrayView<float> OutValues) const
+void FDataflowLinearGradientFloatSampler::Sample(TArrayView<const FVector3f> Positions, TArrayView<float> OutValues) const
 {
 	const FVector3f GradientDirection = EndPoint - StartPoint;
 	const double SquareLength = GradientDirection.SquaredLength();
@@ -43,34 +43,36 @@ void FDataflowLinearGradientValueSampler::Sample(TArrayView<const FVector3f> Pos
 	}
 }
 
-FDataflowLinearGradientValueSamplerNode::FDataflowLinearGradientValueSamplerNode(const UE::Dataflow::FNodeParameters& InParam, FGuid InGuid)
+FDataflowLinearGradientFloatSamplerNode::FDataflowLinearGradientFloatSamplerNode(const UE::Dataflow::FNodeParameters& InParam, FGuid InGuid)
 	: FDataflowNode(InParam, InGuid)
 {
-	RegisterInputConnection(&Gradient.StartPoint, GET_MEMBER_NAME_CHECKED(FDataflowLinearGradientValueSampler, StartPoint));
-	RegisterInputConnection(&Gradient.StartValue, GET_MEMBER_NAME_CHECKED(FDataflowLinearGradientValueSampler, StartValue));
-	RegisterInputConnection(&Gradient.EndPoint, GET_MEMBER_NAME_CHECKED(FDataflowLinearGradientValueSampler, EndPoint));
-	RegisterInputConnection(&Gradient.EndValue, GET_MEMBER_NAME_CHECKED(FDataflowLinearGradientValueSampler, EndValue));
-	RegisterInputConnection(&Gradient.bClamp, GET_MEMBER_NAME_CHECKED(FDataflowLinearGradientValueSampler, bClamp));
-	RegisterOutputConnection(&Gradient);
+	RegisterInputConnection(&Gradient.StartPoint, GET_MEMBER_NAME_CHECKED(FDataflowLinearGradientFloatSampler, StartPoint));
+	RegisterInputConnection(&Gradient.StartValue, GET_MEMBER_NAME_CHECKED(FDataflowLinearGradientFloatSampler, StartValue));
+	RegisterInputConnection(&Gradient.EndPoint, GET_MEMBER_NAME_CHECKED(FDataflowLinearGradientFloatSampler, EndPoint));
+	RegisterInputConnection(&Gradient.EndValue, GET_MEMBER_NAME_CHECKED(FDataflowLinearGradientFloatSampler, EndValue));
+	RegisterInputConnection(&Gradient.bClamp, GET_MEMBER_NAME_CHECKED(FDataflowLinearGradientFloatSampler, bClamp));
+	RegisterOutputConnection(&Sampler);
 }
 
-void FDataflowLinearGradientValueSamplerNode::Evaluate(UE::Dataflow::FContext& Context, const FDataflowOutput* Out) const
+void FDataflowLinearGradientFloatSamplerNode::Evaluate(UE::Dataflow::FContext& Context, const FDataflowOutput* Out) const
 {
-	if (Out->IsA(&Gradient))
+	if (Out->IsA(&Sampler))
 	{
-		FDataflowLinearGradientValueSampler OutGradient;
-		OutGradient.StartPoint = GetValue(Context, &Gradient.StartPoint, Gradient.StartPoint);
-		OutGradient.StartValue = GetValue(Context, &Gradient.StartValue, Gradient.StartValue);
-		OutGradient.EndPoint = GetValue(Context, &Gradient.EndPoint, Gradient.EndPoint);
-		OutGradient.EndValue = GetValue(Context, &Gradient.EndValue, Gradient.EndValue);
-		OutGradient.bClamp = GetValue(Context, &Gradient.bClamp, Gradient.bClamp);
-		SetValue(Context, OutGradient, &Gradient);
+		TSharedRef<FDataflowLinearGradientFloatSampler> Impl = MakeShared<FDataflowLinearGradientFloatSampler>(Gradient);
+		Impl->StartPoint = GetValue(Context, &Gradient.StartPoint);
+		Impl->StartValue = GetValue(Context, &Gradient.StartValue);
+		Impl->EndPoint = GetValue(Context, &Gradient.EndPoint);
+		Impl->EndValue = GetValue(Context, &Gradient.EndValue);
+		Impl->bClamp = GetValue(Context, &Gradient.bClamp);
+
+		FDataflowFloatSampler OutSampler(Impl.ToSharedPtr());
+		SetValue(Context, OutSampler, &Sampler);
 	}
 }
 
 ////////////////////////////////////////////////////////////////////////////////////////////////////////////////////////
 
-void FDataflowLinearGradientFromBoxValueSampler::Sample(TArrayView<const FVector3f> Positions, TArrayView<float> OutValues) const
+void FDataflowLinearGradientFromBoxFloatSampler::Sample(TArrayView<const FVector3f> Positions, TArrayView<float> OutValues) const
 {
 	const FVector BoxCenter = Box.GetCenter();
 	const FVector BoxHalfExtent = Box.GetExtent();
@@ -89,7 +91,7 @@ void FDataflowLinearGradientFromBoxValueSampler::Sample(TArrayView<const FVector
 	const FVector StartPoint = Transform.GetTranslation() + BoxCenter - GradientHalfExtent;
 	const FVector EndPoint = Transform.GetTranslation() + BoxCenter + GradientHalfExtent;
 
-	FDataflowLinearGradientValueSampler RegularGradientSampler;
+	FDataflowLinearGradientFloatSampler RegularGradientSampler;
 	RegularGradientSampler.StartPoint = FVector3f(StartPoint);
 	RegularGradientSampler.StartValue = StartValue;
 	RegularGradientSampler.EndPoint = FVector3f(EndPoint);
@@ -99,38 +101,40 @@ void FDataflowLinearGradientFromBoxValueSampler::Sample(TArrayView<const FVector
 	RegularGradientSampler.Sample(Positions, OutValues);
 }
 
-FDataflowLinearGradientFromBoxValueSamplerNode::FDataflowLinearGradientFromBoxValueSamplerNode(const UE::Dataflow::FNodeParameters& InParam, FGuid InGuid)
+FDataflowLinearGradientFromBoxFloatSamplerNode::FDataflowLinearGradientFromBoxFloatSamplerNode(const UE::Dataflow::FNodeParameters& InParam, FGuid InGuid)
 	: FDataflowNode(InParam, InGuid)
 {
-	RegisterInputConnection(&Gradient.Box, GET_MEMBER_NAME_CHECKED(FDataflowLinearGradientFromBoxValueSampler, Box));
-	RegisterInputConnection(&Gradient.Axis, GET_MEMBER_NAME_CHECKED(FDataflowLinearGradientFromBoxValueSampler, Axis));
-	RegisterInputConnection(&Gradient.Transform, GET_MEMBER_NAME_CHECKED(FDataflowLinearGradientFromBoxValueSampler, Transform));
-	RegisterInputConnection(&Gradient.StartValue, GET_MEMBER_NAME_CHECKED(FDataflowLinearGradientFromBoxValueSampler, StartValue));
-	RegisterInputConnection(&Gradient.EndValue, GET_MEMBER_NAME_CHECKED(FDataflowLinearGradientFromBoxValueSampler, EndValue));
-	RegisterInputConnection(&Gradient.bClamp, GET_MEMBER_NAME_CHECKED(FDataflowLinearGradientFromBoxValueSampler, bClamp));
-	RegisterOutputConnection(&Gradient);
+	RegisterInputConnection(&Gradient.Box, GET_MEMBER_NAME_CHECKED(FDataflowLinearGradientFromBoxFloatSampler, Box));
+	RegisterInputConnection(&Gradient.Axis, GET_MEMBER_NAME_CHECKED(FDataflowLinearGradientFromBoxFloatSampler, Axis));
+	RegisterInputConnection(&Gradient.Transform, GET_MEMBER_NAME_CHECKED(FDataflowLinearGradientFromBoxFloatSampler, Transform));
+	RegisterInputConnection(&Gradient.StartValue, GET_MEMBER_NAME_CHECKED(FDataflowLinearGradientFromBoxFloatSampler, StartValue));
+	RegisterInputConnection(&Gradient.EndValue, GET_MEMBER_NAME_CHECKED(FDataflowLinearGradientFromBoxFloatSampler, EndValue));
+	RegisterInputConnection(&Gradient.bClamp, GET_MEMBER_NAME_CHECKED(FDataflowLinearGradientFromBoxFloatSampler, bClamp));
+	RegisterOutputConnection(&Sampler);
 }
 
-void FDataflowLinearGradientFromBoxValueSamplerNode::Evaluate(UE::Dataflow::FContext& Context, const FDataflowOutput* Out) const
+void FDataflowLinearGradientFromBoxFloatSamplerNode::Evaluate(UE::Dataflow::FContext& Context, const FDataflowOutput* Out) const
 {
-	if (Out->IsA(&Gradient))
+	if (Out->IsA(&Sampler))
 	{
-		FDataflowLinearGradientFromBoxValueSampler OutGradient{ Gradient };
-		OutGradient.Box = GetValue(Context, &Gradient.Box, Gradient.Box);
-		OutGradient.Axis = GetValue(Context, &Gradient.Axis, Gradient.Axis);
-		OutGradient.Transform = GetValue(Context, &Gradient.Transform, Gradient.Transform);
-		OutGradient.StartValue= GetValue(Context, &Gradient.StartValue, Gradient.StartValue);
-		OutGradient.EndValue = GetValue(Context, &Gradient.EndValue, Gradient.EndValue);
-		OutGradient.bClamp = GetValue(Context, &Gradient.bClamp, Gradient.bClamp);
-		SetValue(Context, OutGradient, &Gradient);
+		TSharedRef<FDataflowLinearGradientFromBoxFloatSampler> Impl = MakeShared<FDataflowLinearGradientFromBoxFloatSampler>(Gradient);
+		Impl->Box = GetValue(Context, &Gradient.Box);
+		Impl->Axis = GetValue(Context, &Gradient.Axis);
+		Impl->Transform = GetValue(Context, &Gradient.Transform);
+		Impl->StartValue= GetValue(Context, &Gradient.StartValue);
+		Impl->EndValue = GetValue(Context, &Gradient.EndValue);
+		Impl->bClamp = GetValue(Context, &Gradient.bClamp);
+
+		FDataflowFloatSampler OutSampler(Impl.ToSharedPtr());
+		SetValue(Context, OutSampler, &Sampler);
 	}
 }
 
 ////////////////////////////////////////////////////////////////////////////////////////////////////////////////////////
 
-void FDataflowRadialGradientValueSampler::Sample(TArrayView<const FVector3f> Positions, TArrayView<float> OutValues) const
+void FDataflowRadialGradientFloatSampler::Sample(TArrayView<const FVector3f> Positions, TArrayView<float> OutValues) const
 {
-	if (Radius == UE_SMALL_NUMBER)
+	if (Radius < UE_SMALL_NUMBER)
 	{
 		for (float& OutValue : OutValues)
 		{
@@ -154,28 +158,30 @@ void FDataflowRadialGradientValueSampler::Sample(TArrayView<const FVector3f> Pos
 	}
 }
 
-FDataflowRadialGradientValueSamplerNode::FDataflowRadialGradientValueSamplerNode(const UE::Dataflow::FNodeParameters& InParam, FGuid InGuid)
+FDataflowRadialGradientFloatSamplerNode::FDataflowRadialGradientFloatSamplerNode(const UE::Dataflow::FNodeParameters& InParam, FGuid InGuid)
 	: FDataflowNode(InParam, InGuid)
 {
-	RegisterInputConnection(&Gradient.Center, GET_MEMBER_NAME_CHECKED(FDataflowRadialGradientValueSampler, Center));
-	RegisterInputConnection(&Gradient.Radius, GET_MEMBER_NAME_CHECKED(FDataflowRadialGradientValueSampler, Radius));
-	RegisterInputConnection(&Gradient.CenterValue, GET_MEMBER_NAME_CHECKED(FDataflowRadialGradientValueSampler, CenterValue));
-	RegisterInputConnection(&Gradient.EdgeValue, GET_MEMBER_NAME_CHECKED(FDataflowRadialGradientValueSampler, EdgeValue));
-	RegisterInputConnection(&Gradient.bClamp, GET_MEMBER_NAME_CHECKED(FDataflowRadialGradientValueSampler, bClamp));
-	RegisterOutputConnection(&Gradient);
+	RegisterInputConnection(&Gradient.Center, GET_MEMBER_NAME_CHECKED(FDataflowRadialGradientFloatSampler, Center));
+	RegisterInputConnection(&Gradient.Radius, GET_MEMBER_NAME_CHECKED(FDataflowRadialGradientFloatSampler, Radius));
+	RegisterInputConnection(&Gradient.CenterValue, GET_MEMBER_NAME_CHECKED(FDataflowRadialGradientFloatSampler, CenterValue));
+	RegisterInputConnection(&Gradient.EdgeValue, GET_MEMBER_NAME_CHECKED(FDataflowRadialGradientFloatSampler, EdgeValue));
+	RegisterInputConnection(&Gradient.bClamp, GET_MEMBER_NAME_CHECKED(FDataflowRadialGradientFloatSampler, bClamp));
+	RegisterOutputConnection(&Sampler);
 }
 
-void FDataflowRadialGradientValueSamplerNode::Evaluate(UE::Dataflow::FContext& Context, const FDataflowOutput* Out) const
+void FDataflowRadialGradientFloatSamplerNode::Evaluate(UE::Dataflow::FContext& Context, const FDataflowOutput* Out) const
 {
-	if (Out->IsA(&Gradient))
+	if (Out->IsA(&Sampler))
 	{
-		FDataflowRadialGradientValueSampler OutGradient;
-		OutGradient.Center = GetValue(Context, &Gradient.Center, Gradient.Center);
-		OutGradient.Radius = GetValue(Context, &Gradient.Radius, Gradient.Radius);
-		OutGradient.CenterValue = GetValue(Context, &Gradient.CenterValue, Gradient.CenterValue);
-		OutGradient.EdgeValue = GetValue(Context, &Gradient.EdgeValue, Gradient.EdgeValue);
-		OutGradient.bClamp = GetValue(Context, &Gradient.bClamp, Gradient.bClamp);
-		SetValue(Context, OutGradient, &Gradient);
+		TSharedRef<FDataflowRadialGradientFloatSampler> Impl = MakeShared<FDataflowRadialGradientFloatSampler>(Gradient);
+		Impl->Center = GetValue(Context, &Gradient.Center, Gradient.Center);
+		Impl->Radius = GetValue(Context, &Gradient.Radius, Gradient.Radius);
+		Impl->CenterValue = GetValue(Context, &Gradient.CenterValue, Gradient.CenterValue);
+		Impl->EdgeValue = GetValue(Context, &Gradient.EdgeValue, Gradient.EdgeValue);
+		Impl->bClamp = GetValue(Context, &Gradient.bClamp, Gradient.bClamp);
+
+		FDataflowFloatSampler OutSampler(Impl.ToSharedPtr());
+		SetValue(Context, OutSampler, &Sampler);
 	}
 }
 
diff --git a/Engine/Plugins/Experimental/Dataflow/Source/DataflowNodes/Private/Dataflow/SamplerNodes/DataflowGradientSamplerNode.h b/Engine/Plugins/Experimental/Dataflow/Source/DataflowNodes/Private/Dataflow/SamplerNodes/DataflowGradientSamplerNode.h
--- a/Engine/Plugins/Experimental/Dataflow/Source/DataflowNodes/Private/Dataflow/SamplerNodes/DataflowGradientSamplerNode.h
+++ b/Engine/Plugins/Experimental/Dataflow/Source/DataflowNodes/Private/Dataflow/SamplerNodes/DataflowGradientSamplerNode.h
@@ -10,11 +10,11 @@
 ////////////////////////////////////////////////////////////////////////////////////////////////////////////////////////
 
 USTRUCT()
-struct FDataflowLinearGradientValueSampler : public FDataflowValueSampler
+struct FDataflowLinearGradientFloatSampler : public FDataflowFloatSamplerBase
 {
 	GENERATED_BODY();
 
-	virtual ~FDataflowLinearGradientValueSampler() override = default;
+	virtual ~FDataflowLinearGradientFloatSampler() override = default;
 	virtual void Sample(TArrayView<const FVector3f> Positions, TArrayView<float> OutValues) const override;
 
 public:
@@ -43,17 +43,20 @@ public:
  * Output a linear gradient sampler 
  */
 USTRUCT()
-struct FDataflowLinearGradientValueSamplerNode : public FDataflowNode
+struct FDataflowLinearGradientFloatSamplerNode : public FDataflowNode
 {
 	GENERATED_BODY()
-	DATAFLOW_NODE_DEFINE_INTERNAL(FDataflowLinearGradientValueSamplerNode, "Linear Gradient Sampler", "Samplers", "")
+	DATAFLOW_NODE_DEFINE_INTERNAL(FDataflowLinearGradientFloatSamplerNode, "Linear Gradient Sampler", "Samplers", "")
 
 public:
-	FDataflowLinearGradientValueSamplerNode(const UE::Dataflow::FNodeParameters& InParam, FGuid InGuid = FGuid::NewGuid());
+	FDataflowLinearGradientFloatSamplerNode(const UE::Dataflow::FNodeParameters& InParam, FGuid InGuid = FGuid::NewGuid());
 
 private:
-	UPROPERTY(EditAnywhere, Category = Sampler, meta = (DataflowOutput, ShowOnlyInnerProperties))
-	FDataflowLinearGradientValueSampler Gradient;
+	UPROPERTY(meta = (DataflowOutput))
+	FDataflowFloatSampler Sampler;
+
+	UPROPERTY(EditAnywhere, Category = Sampler, meta = (ShowOnlyInnerProperties))
+	FDataflowLinearGradientFloatSampler Gradient;
 
 	void Evaluate(UE::Dataflow::FContext& Context, const FDataflowOutput* Out) const override;
 };
@@ -61,11 +64,11 @@ private:
 ////////////////////////////////////////////////////////////////////////////////////////////////////////////////////////
 
 USTRUCT()
-struct FDataflowLinearGradientFromBoxValueSampler : public FDataflowValueSampler
+struct FDataflowLinearGradientFromBoxFloatSampler : public FDataflowFloatSamplerBase
 {
 	GENERATED_BODY();
 
-	virtual ~FDataflowLinearGradientFromBoxValueSampler() override = default;
+	virtual ~FDataflowLinearGradientFromBoxFloatSampler() override = default;
 	virtual void Sample(TArrayView<const FVector3f> Positions, TArrayView<float> OutValues) const override;
 
 public:
@@ -98,29 +101,32 @@ public:
  * Output a linear gradient sampler
  */
 USTRUCT()
-struct FDataflowLinearGradientFromBoxValueSamplerNode : public FDataflowNode
+struct FDataflowLinearGradientFromBoxFloatSamplerNode : public FDataflowNode
 {
 	GENERATED_BODY()
-	DATAFLOW_NODE_DEFINE_INTERNAL(FDataflowLinearGradientFromBoxValueSamplerNode, "Linear Gradient From Box Sampler", "Samplers", "Bounds Bounding Box Axis")
+	DATAFLOW_NODE_DEFINE_INTERNAL(FDataflowLinearGradientFromBoxFloatSamplerNode, "Linear Gradient From Box Sampler", "Samplers", "Bounds Bounding Box Axis")
 
 public:
-	FDataflowLinearGradientFromBoxValueSamplerNode(const UE::Dataflow::FNodeParameters& InParam, FGuid InGuid = FGuid::NewGuid());
+	FDataflowLinearGradientFromBoxFloatSamplerNode(const UE::Dataflow::FNodeParameters& InParam, FGuid InGuid = FGuid::NewGuid());
 
 private:
-	UPROPERTY(EditAnywhere, Category = Sampler, meta = (DataflowOutput, ShowOnlyInnerProperties))
-	FDataflowLinearGradientFromBoxValueSampler Gradient;
+	UPROPERTY(meta = (DataflowOutput))
+	FDataflowFloatSampler Sampler;
 
+	UPROPERTY(EditAnywhere, Category = Sampler, meta = (ShowOnlyInnerProperties))
+	FDataflowLinearGradientFromBoxFloatSampler Gradient;
+	
 	void Evaluate(UE::Dataflow::FContext& Context, const FDataflowOutput* Out) const override;
 };
 
 ////////////////////////////////////////////////////////////////////////////////////////////////////////////////////////
 
 USTRUCT()
-struct FDataflowRadialGradientValueSampler : public FDataflowValueSampler
+struct FDataflowRadialGradientFloatSampler : public FDataflowFloatSamplerBase
 {
 	GENERATED_BODY();
 
-	virtual ~FDataflowRadialGradientValueSampler() override = default;
+	virtual ~FDataflowRadialGradientFloatSampler() override = default;
 	virtual void Sample(TArrayView<const FVector3f> Positions, TArrayView<float> OutValues) const override;
 
 public:
@@ -149,17 +155,20 @@ public:
  * Output a radial gradient sampler
  */
 USTRUCT()
-struct FDataflowRadialGradientValueSamplerNode : public FDataflowNode
+struct FDataflowRadialGradientFloatSamplerNode : public FDataflowNode
 {
 	GENERATED_BODY()
-	DATAFLOW_NODE_DEFINE_INTERNAL(FDataflowRadialGradientValueSamplerNode, "Radial Gradient Sampler", "Samplers", "")
+	DATAFLOW_NODE_DEFINE_INTERNAL(FDataflowRadialGradientFloatSamplerNode, "Radial Gradient Sampler", "Samplers", "")
 
 public:
-	FDataflowRadialGradientValueSamplerNode(const UE::Dataflow::FNodeParameters& InParam, FGuid InGuid = FGuid::NewGuid());
+	FDataflowRadialGradientFloatSamplerNode(const UE::Dataflow::FNodeParameters& InParam, FGuid InGuid = FGuid::NewGuid());
 
 private:
-	UPROPERTY(EditAnywhere, Category = Sampler, meta = (DataflowOutput, ShowOnlyInnerProperties))
-	FDataflowRadialGradientValueSampler Gradient;
+	UPROPERTY(meta = (DataflowOutput))
+	FDataflowFloatSampler Sampler;
+
+	UPROPERTY(EditAnywhere, Category = Sampler, meta = (ShowOnlyInnerProperties))
+	FDataflowRadialGradientFloatSampler Gradient;;
 
 	void Evaluate(UE::Dataflow::FContext& Context, const FDataflowOutput* Out) const override;
 };
@@ -168,5 +177,5 @@ private:
 
 namespace UE::Dataflow
 {
-	void RegisterGradientValueSamplerNodes();
+	void RegisterGradientSamplerNodes();
 }
\ No newline at end of file
diff --git a/Engine/Plugins/Experimental/Dataflow/Source/DataflowNodes/Private/Dataflow/SamplerNodes/DataflowSamplerTypes.cpp b/Engine/Plugins/Experimental/Dataflow/Source/DataflowNodes/Private/Dataflow/SamplerNodes/DataflowSamplerTypes.cpp
--- a/Engine/Plugins/Experimental/Dataflow/Source/DataflowNodes/Private/Dataflow/SamplerNodes/DataflowSamplerTypes.cpp
+++ b/Engine/Plugins/Experimental/Dataflow/Source/DataflowNodes/Private/Dataflow/SamplerNodes/DataflowSamplerTypes.cpp
@@ -3,10 +3,23 @@
 #include "Dataflow/SamplerNodes/DataflowSamplerTypes.h"
 #include "Dataflow/DataflowAnyTypeRegistry.h"
 
-namespace UE::Dataflow
+FDataflowFloatSampler::FDataflowFloatSampler(const TSharedPtr<const FDataflowFloatSamplerBase>& InImpl)
+	: Impl(InImpl)
 {
-	void RegisterValueSamplerTypes()
+}
+
+void FDataflowFloatSampler::Sample(TArrayView<const FVector3f> Positions, TArrayView<float> OutValues) const
+{
+	if (Impl)
+	{
+		Impl->Sample(Positions, OutValues);
+	}
+}
+FBox FDataflowFloatSampler::GetRenderBounds() const
+{
+	if (Impl)
 	{
-		UE_DATAFLOW_REGISTER_ANYTYPE(FDataflowValueSamplerTypes);
+		return Impl->GetRenderBounds();
 	}
+	return FDataflowFloatSamplerBase::GetRenderBounds();
 }
diff --git a/Engine/Plugins/Experimental/Dataflow/Source/DataflowNodes/Private/Dataflow/SamplerNodes/DataflowSamplerTypes.h b/Engine/Plugins/Experimental/Dataflow/Source/DataflowNodes/Private/Dataflow/SamplerNodes/DataflowSamplerTypes.h
--- a/Engine/Plugins/Experimental/Dataflow/Source/DataflowNodes/Private/Dataflow/SamplerNodes/DataflowSamplerTypes.h
+++ b/Engine/Plugins/Experimental/Dataflow/Source/DataflowNodes/Private/Dataflow/SamplerNodes/DataflowSamplerTypes.h
@@ -8,27 +8,33 @@
 #include "DataflowSamplerTypes.generated.h"
 
 /**
-* A Value sampler, sample at a specific 3D point and return a single value
+* A float sampler, sample at a specific 3D position  and return a single float value per position 
 */
 USTRUCT()
-struct FDataflowValueSampler
+struct FDataflowFloatSamplerBase
 {
 	GENERATED_BODY();
 public:
-	virtual ~FDataflowValueSampler() = default;
-	virtual void Sample(TArrayView<const FVector3f> Positions, TArrayView<float> OutValues) const { check(false); }
+	virtual ~FDataflowFloatSamplerBase() = default;
+	virtual void Sample(TArrayView<const FVector3f> Positions, TArrayView<float> OutValues) const {};
+	virtual FBox GetRenderBounds() const { return FBox(EForceInit::ForceInitToZero); };
 };
 
-/** Value sampler anytype */
 USTRUCT()
-struct FDataflowValueSamplerTypes : public FDataflowAnyType
+struct FDataflowFloatSampler : public FDataflowFloatSamplerBase
 {
-	using FPolicyType = FDataflowUDerivedStructTypePolicy<FDataflowValueSampler>;
-	using FStorageType = void;
 	GENERATED_BODY()
+
+public:
+	FDataflowFloatSampler() = default;
+	FDataflowFloatSampler(const TSharedPtr<const FDataflowFloatSamplerBase>& InImpl);
+
+	virtual void Sample(TArrayView<const FVector3f> Positions, TArrayView<float> OutValues) const override;
+	virtual FBox GetRenderBounds() const override;
+	
+	const TSharedPtr<const FDataflowFloatSamplerBase>& GetImpl() const { return Impl; }
+
+private:
+	TSharedPtr<const FDataflowFloatSamplerBase> Impl;
 };
 
-namespace UE::Dataflow
-{
-	void RegisterValueSamplerTypes();
-}
diff --git a/Engine/Plugins/Experimental/Dataflow/Source/DataflowNodes/Private/Dataflow/SamplerNodes/DataflowTransformSamplerNode.cpp b/Engine/Plugins/Experimental/Dataflow/Source/DataflowNodes/Private/Dataflow/SamplerNodes/DataflowTransformSamplerNode.cpp
--- a/Engine/Plugins/Experimental/Dataflow/Source/DataflowNodes/Private/Dataflow/SamplerNodes/DataflowTransformSamplerNode.cpp
+++ b/Engine/Plugins/Experimental/Dataflow/Source/DataflowNodes/Private/Dataflow/SamplerNodes/DataflowTransformSamplerNode.cpp
@@ -5,43 +5,52 @@
 
 namespace UE::Dataflow
 {
-	void RegisterTransformValueSamplerNodes()
+	void RegisterTransformSamplerNodes()
 	{
-		DATAFLOW_NODE_REGISTER_CREATION_FACTORY(FDataflowTransformValueSamplerNode);
+		DATAFLOW_NODE_REGISTER_CREATION_FACTORY(FDataflowTransformFloatSamplerNode);
 	}
 }
 
 ////////////////////////////////////////////////////////////////////////////////////////////////////////////////////////
 
-void FDataflowTransformValueSampler::Sample(TArrayView<const FVector3f> Positions, TArrayView<float> OutValues) const
+void FDataflowTransformFloatSampler::Sample(TArrayView<const FVector3f> Positions, TArrayView<float> OutValues) const
 {
-	if (const FDataflowValueSampler* SamplerPtr = Sampler.GetPtr())
+	if (Sampler.IsValid())
 	{
 		TArray<FVector3f> TransformedPositions;
 		TransformedPositions.SetNumUninitialized(Positions.Num());
 		for (int32 Index = 0; Index < Positions.Num(); ++Index)
 		{
-			TransformedPositions[Index] = FVector3f(Transform.InverseTransformPosition(FVector(Positions[Index])));
+			if (bInverseTransform)
+			{
+				TransformedPositions[Index] = FVector3f(Transform.InverseTransformPosition(FVector(Positions[Index])));
+			}
+			else
+			{
+				TransformedPositions[Index] = FVector3f(Transform.TransformPosition(FVector(Positions[Index])));
+			}
 		}
-		SamplerPtr->Sample(TransformedPositions, OutValues);
+		Sampler->Sample(TransformedPositions, OutValues);
 	}
 }
 
-FDataflowTransformValueSamplerNode::FDataflowTransformValueSamplerNode(const UE::Dataflow::FNodeParameters& InParam, FGuid InGuid)
+FDataflowTransformFloatSamplerNode::FDataflowTransformFloatSamplerNode(const UE::Dataflow::FNodeParameters& InParam, FGuid InGuid)
 	: FDataflowNode(InParam, InGuid)
 {
-	RegisterInputConnection(&TransformSampler.Transform, GET_MEMBER_NAME_CHECKED(FDataflowTransformValueSampler, Transform));
 	RegisterInputConnection(&Sampler);
-	RegisterOutputConnection(&TransformSampler);
+	RegisterInputConnection(&TransformSampler.Transform, GET_MEMBER_NAME_CHECKED(FDataflowTransformFloatSampler, Transform));
+	RegisterOutputConnection(&Sampler, &Sampler);
 }
 
-void FDataflowTransformValueSamplerNode::Evaluate(UE::Dataflow::FContext& Context, const FDataflowOutput* Out) const
+void FDataflowTransformFloatSamplerNode::Evaluate(UE::Dataflow::FContext& Context, const FDataflowOutput* Out) const
 {
-	if (Out->IsA(&TransformSampler))
+	if (Out->IsA(&Sampler))
 	{
-		FDataflowTransformValueSampler OutTransformSampler;
-		OutTransformSampler.Transform = GetValue(Context, &TransformSampler.Transform);
-		OutTransformSampler.Sampler = GetValue(Context, &Sampler).GetTypedInstancedStruct<FDataflowValueSampler>();
-		SetValue(Context, OutTransformSampler, &TransformSampler);
+		TSharedRef<FDataflowTransformFloatSampler> Impl = MakeShared<FDataflowTransformFloatSampler>(TransformSampler);
+		Impl->Transform = GetValue(Context, &TransformSampler.Transform);
+		Impl->Sampler = GetValue(Context, &Sampler).GetImpl();
+	
+		FDataflowFloatSampler OutSampler(Impl.ToSharedPtr());
+		SetValue(Context, OutSampler, &Sampler);
 	}
-}
+}
\ No newline at end of file
diff --git a/Engine/Plugins/Experimental/Dataflow/Source/DataflowNodes/Private/Dataflow/SamplerNodes/DataflowTransformSamplerNode.h b/Engine/Plugins/Experimental/Dataflow/Source/DataflowNodes/Private/Dataflow/SamplerNodes/DataflowTransformSamplerNode.h
--- a/Engine/Plugins/Experimental/Dataflow/Source/DataflowNodes/Private/Dataflow/SamplerNodes/DataflowTransformSamplerNode.h
+++ b/Engine/Plugins/Experimental/Dataflow/Source/DataflowNodes/Private/Dataflow/SamplerNodes/DataflowTransformSamplerNode.h
@@ -7,14 +7,12 @@
 
 #include "DataflowTransformSamplerNode.generated.h"
 
-////////////////////////////////////////////////////////////////////////////////////////////////////////////////////////
-
 USTRUCT()
-struct FDataflowTransformValueSampler : public FDataflowValueSampler
+struct FDataflowTransformFloatSampler : public FDataflowFloatSamplerBase
 {
 	GENERATED_BODY();
 
-	virtual ~FDataflowTransformValueSampler() override = default;
+	virtual ~FDataflowTransformFloatSampler() override = default;
 	virtual void Sample(TArrayView<const FVector3f> Positions, TArrayView<float> OutValues) const override;
 
 public:
@@ -22,34 +20,36 @@ public:
 	UPROPERTY(EditAnywhere, Category = Transform, meta = (DataflowInput))
 	FTransform Transform;
 
+	UPROPERTY(EditAnywhere, Category = Transform)
+	bool bInverseTransform = false;
+
 	/** Sampler to send the transformed point to */
-	UPROPERTY()
-	TInstancedStruct<FDataflowValueSampler> Sampler;
+	TSharedPtr<const FDataflowFloatSamplerBase> Sampler;
 };
 
 /**
- * Output a transform sampler 
+ * Output a transform sampler
  */
 USTRUCT()
-struct FDataflowTransformValueSamplerNode : public FDataflowNode
+struct FDataflowTransformFloatSamplerNode : public FDataflowNode
 {
 	GENERATED_BODY()
-	DATAFLOW_NODE_DEFINE_INTERNAL(FDataflowTransformValueSamplerNode, "Transform Sampler", "Samplers", "Offset Translation Rotation Scale Space points")
+	DATAFLOW_NODE_DEFINE_INTERNAL(FDataflowTransformFloatSamplerNode, "Transform Sampler", "Samplers", "Offset Translation Rotation Scale Space points")
 
 public:
-	FDataflowTransformValueSamplerNode(const UE::Dataflow::FNodeParameters& InParam, FGuid InGuid = FGuid::NewGuid());
+	FDataflowTransformFloatSamplerNode(const UE::Dataflow::FNodeParameters& InParam, FGuid InGuid = FGuid::NewGuid());
 
 private:
-	UPROPERTY(meta = (DataflowInput))
-	FDataflowValueSamplerTypes Sampler;
+	UPROPERTY(meta = (DataflowInput, DataflowOutput, DataflowPassthrough = Sampler))
+	FDataflowFloatSampler Sampler;
 
-	UPROPERTY(EditAnywhere, Category = Sampler, meta = (DataflowOutput, ShowOnlyInnerProperties))
-	FDataflowTransformValueSampler TransformSampler;
+	UPROPERTY(EditAnywhere, Category = Sampler, meta = (ShowOnlyInnerProperties))
+	FDataflowTransformFloatSampler TransformSampler;
 
 	void Evaluate(UE::Dataflow::FContext& Context, const FDataflowOutput* Out) const override;
 };
 
 namespace UE::Dataflow
 {
-	void RegisterTransformValueSamplerNodes();
+	void RegisterTransformSamplerNodes();
 }
\ No newline at end of file
