Dataflow : fix issue where node generating dynamic meshes from the new rendering system do not get recognized as target for tools

#Synced-CL 50641382

diff --git a/Engine/Plugins/Experimental/Dataflow/Source/DataflowEditor/Private/Dataflow/DataflowConstructionScene.cpp b/Engine/Plugins/Experimental/Dataflow/Source/DataflowEditor/Private/Dataflow/DataflowConstructionScene.cpp
--- a/Engine/Plugins/Experimental/Dataflow/Source/DataflowEditor/Private/Dataflow/DataflowConstructionScene.cpp
+++ b/Engine/Plugins/Experimental/Dataflow/Source/DataflowEditor/Private/Dataflow/DataflowConstructionScene.cpp
@@ -723,6 +723,19 @@ void FDataflowConstructionScene::UpdateComponentsFromNode(const UDataflowEdNode*
 					ProgressNotification.AddProgress(RenderableTypeProgress);
 				}
 				PrimitiveComponents.Append(NodeRenderableComponents.GetComponents());
+
+				int32 MeshIndex = 0;
+				const TArray<UPrimitiveComponent*>& NodeComponents = NodeRenderableComponents.GetComponents();
+				for (UPrimitiveComponent* PrimitiveComponent : NodeComponents)
+				{
+					if (UDynamicMeshComponent* DynamicMeshComponent = Cast<UDynamicMeshComponent>(PrimitiveComponent))
+					{
+						FDataflowRenderKey DynamicMeshKey{ EdNode, MeshIndex };
+						DynamicMeshComponents.Emplace(DynamicMeshKey, DynamicMeshComponent);
+						++MeshIndex;
+					}
+				}
+				
 			}
 		}
 	}
