Dataflow : add curve selection display to geometry collection renderable type

#Synced-CL 50764219

diff --git a/Engine/Plugins/Experimental/Dataflow/Source/DataflowEditor/Private/DataflowRendering/DataflowGeometryCollectionRenderableType.cpp b/Engine/Plugins/Experimental/Dataflow/Source/DataflowEditor/Private/DataflowRendering/DataflowGeometryCollectionRenderableType.cpp
--- a/Engine/Plugins/Experimental/Dataflow/Source/DataflowEditor/Private/DataflowRendering/DataflowGeometryCollectionRenderableType.cpp
+++ b/Engine/Plugins/Experimental/Dataflow/Source/DataflowEditor/Private/DataflowRendering/DataflowGeometryCollectionRenderableType.cpp
@@ -4,6 +4,7 @@
 
 #include "Components/DynamicMeshComponent.h"
 #include "Drawing/PointSetComponent.h"
+#include "Drawing/LineSetComponent.h"
 
 #include "Dataflow/DataflowRenderingViewMode.h"
 #include "DataflowRendering/DataflowRenderableComponents.h"
@@ -118,6 +119,21 @@ namespace UE::Dataflow::Private
 		UE_DATAFLOW_IRENDERABLE_VIEW_MODE(FDataflowConstruction3DViewMode);
 		UE_DATAFLOW_IRENDERABLE_SETTINGS(UDataflowGeometryCollectionSurfaceRenderSettings);
 
+		struct FSelection
+		{
+			FDataflowTransformSelection TransformSelection;
+			bool bValidTransformSelection = false;
+
+			FDataflowVertexSelection VertexSelection;
+			bool bValidVertexSelection = false;
+
+			FDataflowFaceSelection FaceSelection;
+			bool bValidFaceSelection = false;
+
+			FDataflowFaceSelection CurveSelection;
+			bool bValidCurveSelection = false;
+		};
+
 		virtual bool CanRender(const FRenderableTypeInstance& Instance) const
 		{
 			const FManagedArrayCollection& Collection = GetCollection(Instance);
@@ -207,14 +223,12 @@ namespace UE::Dataflow::Private
 					Instance.CacheValue(CurrentSettingsHash, SettingsHashCacheName);
 				}
 
-				FDataflowTransformSelection TransformSelection;
-				FDataflowVertexSelection VertexSelection;
+				const FColor SelectionColor = FColor::Orange;
+				FSelection Selection;
 				if (bShowSelection)
 				{
-					GetSelection(Instance, Collection, TransformSelection, VertexSelection);
+					GetSelection(Instance, Collection, Selection);
 				}
-				const bool bIsValidTransformSelection = (TransformSelection.GetGroupName() == FTransformCollection::TransformGroup && TransformSelection.IsValidForCollection(Collection));
-				const bool bIsValidVertexSelection = (VertexSelection.GetGroupName() == FGeometryCollection::VerticesGroup && VertexSelection.IsValidForCollection(Collection));
 
 				UMaterialInterface* SelectionMaterial = Settings
 					? Settings->SelectionMaterial.Get()
@@ -233,7 +247,7 @@ namespace UE::Dataflow::Private
 							Component->SetCastShadow(false);
 							Component->SetMesh(MoveTemp(DynamicMeshes[TransformIndex]));
 
-							if (bIsValidTransformSelection && TransformSelection.IsSelected(TransformIndex))
+							if (Selection.bValidTransformSelection && Selection.TransformSelection.IsSelected(TransformIndex))
 							{
 								Component->SetOverrideRenderMaterial(SelectionMaterial);
 							}
@@ -246,7 +260,7 @@ namespace UE::Dataflow::Private
 				}
 
 				// Create point component for vertex selection
-				if (bIsValidVertexSelection)
+				if (Selection.bValidVertexSelection)
 				{
 					const TManagedArray<FVector3f>* VertexAttribute = Collection.FindAttribute<FVector3f>(FGeometryCollection::VertexPositionAttribute, FGeometryCollection::VerticesGroup);
 					const TManagedArray<int32>* VertexToTransformIndex = Collection.FindAttribute<int32>(FGeometryCollection::VertexBoneMapAttribute, FGeometryCollection::VerticesGroup);
@@ -255,13 +269,13 @@ namespace UE::Dataflow::Private
 						const GeometryCollection::Facades::FCollectionTransformFacade TransformFacade(Collection);
 						const TArray<FTransform> CollectionSpaceTransforms = TransformFacade.ComputeCollectionSpaceTransforms();
 
-						const int32 NumVertices = VertexSelection.Num();
+						const int32 NumVertices = Selection.VertexSelection.Num();
 						TArray<FVector> Points;
 						Points.Reserve(NumVertices);
 
 						for (int32 Index = 0; Index < NumVertices; ++Index)
 						{
-							if (VertexSelection.IsSelected(Index))
+							if (Selection.VertexSelection.IsSelected(Index))
 							{
 								const int32 TransformIndex = VertexToTransformIndex ? (*VertexToTransformIndex)[Index] : INDEX_NONE;
 								const FTransform Transform = CollectionSpaceTransforms.IsValidIndex(TransformIndex) ? CollectionSpaceTransforms[TransformIndex] : FTransform::Identity;
@@ -276,20 +290,47 @@ namespace UE::Dataflow::Private
 						{
 							constexpr float PointSize = 8.f;
 							PointComponent->ReservePoints(Points.Num());
-							PointComponent->AddPoints(Points, FColor::Orange, PointSize);
+							PointComponent->AddPoints(Points, SelectionColor, PointSize);
 							UMaterialInterface* Material = LoadObject<UMaterialInterface>(nullptr, TEXT("/MeshModelingToolsetExp/Materials/PointSetOverlaidComponentMaterial_Round"));
 							PointComponent->SetMaterial(0, Material);
 
 						}
 					}
 				}
+
+				// Create line component from curve selection 
+				if (Selection.bValidCurveSelection)
+				{
+					const FName LineComponentName = Instance.GetComponentName(TEXT("CurveSelection"));
+
+					if (ULineSetComponent* LineComponent = OutComponents.AddNewComponent<ULineSetComponent>(LineComponentName))
+					{
+						GeometryCollection::Facades::FCollectionCurveGeometryFacade CurveFacade(Collection);
+						const int32 NumCurves = CurveFacade.GetNumCurves();
+						TArray<FVector3f> CurvePositions;
+						for (int32 CurveIndex = 0; CurveIndex < NumCurves; ++CurveIndex)
+						{
+							if (Selection.CurveSelection.IsSelected(CurveIndex))
+							{
+								CurveFacade.GetCurvePointPositions(CurveIndex, CurvePositions);
+								for (int32 PointIndex = 1; PointIndex < CurvePositions.Num(); ++PointIndex)
+								{
+									LineComponent->AddLine(FVector(CurvePositions[PointIndex - 1]), FVector(CurvePositions[PointIndex]), SelectionColor, 3.0f);
+								}
+							}
+						}
+						LineComponent->SetLineMaterial(LoadObject<UMaterialInterface>(nullptr, TEXT("/MeshModelingToolsetExp/Materials/LineSetOverlaidComponentMaterial")));
+					}
+				}
 			}
 		}
 
-		void GetSelection(const FRenderableTypeInstance& Instance, const FManagedArrayCollection& Collection, FDataflowTransformSelection& OutTransformSelection, FDataflowVertexSelection& OutVertexSelection) const
+		void GetSelection(const FRenderableTypeInstance& Instance, 
+			const FManagedArrayCollection& Collection, 
+			FSelection& OutSelection
+		) const
 		{
-			OutTransformSelection.Clear();
-			OutVertexSelection.Clear();
+			OutSelection = FSelection();
 
 			FDataflowSelection Selection;
 			Instance.GetSelectionToVisualize(Selection);
@@ -298,14 +339,34 @@ namespace UE::Dataflow::Private
 				FDataflowSelection DefaultValue;
 				Selection = Instance.GetOutputValueByType<FDataflowSelectionTypes>(DefaultValue);
 			}
-			if (Selection.GetGroupName() == FTransformCollection::TransformGroup && Selection.AnySelected())
+			if (!Selection.IsValidForCollection(Collection))
 			{
-				OutTransformSelection.Initialize(Selection);
-				GeometryCollectionSurfaceRenderableType::Private::AddLeavesToSelection(Collection, Selection, OutTransformSelection);
+				return;
 			}
-			else if (Selection.GetGroupName() == FGeometryCollection::VerticesGroup && Selection.AnySelected())
+			if (Selection.AnySelected())
 			{
-				OutVertexSelection.Initialize(Selection);
+				const FName GroupName = Selection.GetGroupName();
+				if (GroupName == FDataflowTransformSelection::TransformGroupName)
+				{
+					OutSelection.TransformSelection.Initialize(Selection);
+					GeometryCollectionSurfaceRenderableType::Private::AddLeavesToSelection(Collection, Selection, OutSelection.TransformSelection);
+					OutSelection.bValidTransformSelection = true;
+				}
+				else if (GroupName == FDataflowVertexSelection::VerticesGroupName)
+				{
+					OutSelection.VertexSelection.Initialize(Selection);
+					OutSelection.bValidVertexSelection = true;
+				}
+				else if (GroupName == FDataflowFaceSelection::FacesGroupName)
+				{
+					OutSelection.FaceSelection.Initialize(Selection);
+					OutSelection.bValidFaceSelection = true;
+				}
+				else if (GroupName == FDataflowCurveSelection::CurveGroupName)
+				{
+					OutSelection.CurveSelection.Initialize(Selection);
+					OutSelection.bValidCurveSelection = true;
+				}
 			}
 		}
 
diff --git a/Engine/Source/Runtime/Experimental/Chaos/Private/GeometryCollection/Facades/CollectionCurveFacade.cpp b/Engine/Source/Runtime/Experimental/Chaos/Private/GeometryCollection/Facades/CollectionCurveFacade.cpp
--- a/Engine/Source/Runtime/Experimental/Chaos/Private/GeometryCollection/Facades/CollectionCurveFacade.cpp
+++ b/Engine/Source/Runtime/Experimental/Chaos/Private/GeometryCollection/Facades/CollectionCurveFacade.cpp
@@ -293,6 +293,26 @@ namespace GeometryCollection::Facades
 		}
 	}
 	
+	void FCollectionCurveGeometryFacade::GetCurvePointPositions(const int32 CurveIndex, TArray<FVector3f>& OutPoints) const
+	{
+		OutPoints.Reset();
+
+		const int32 NumCurves = GetNumCurves();
+		if (CurveIndex < 0 || CurveIndex >= NumCurves)
+		{
+			return;
+		}
+
+		const int32 StartPointIndex = (CurveIndex == 0) ? 0 : CurvePointOffsets[CurveIndex-1];
+		const int32 EndPointIndex = CurvePointOffsets[CurveIndex];
+		const int32 NumPoints = (EndPointIndex - StartPointIndex);
+		OutPoints.Reserve(NumPoints);
+		for (int32 Index = StartPointIndex; Index < EndPointIndex; ++Index)
+		{
+			OutPoints.Add(PointRestPositions[Index]);
+		}
+	}
+
 	FCollectionCurveHierarchyFacade::FCollectionCurveHierarchyFacade(FManagedArrayCollection& InCollection) :
 		ConstCollection(InCollection), Collection(&InCollection),
 		CurveParentIndices(InCollection, CurveParentIndicesAttribute, FCollectionCurveGeometryFacade::CurvesGroup),
diff --git a/Engine/Source/Runtime/Experimental/Chaos/Public/GeometryCollection/Facades/CollectionCurveFacade.h b/Engine/Source/Runtime/Experimental/Chaos/Public/GeometryCollection/Facades/CollectionCurveFacade.h
--- a/Engine/Source/Runtime/Experimental/Chaos/Public/GeometryCollection/Facades/CollectionCurveFacade.h
+++ b/Engine/Source/Runtime/Experimental/Chaos/Public/GeometryCollection/Facades/CollectionCurveFacade.h
@@ -99,6 +99,9 @@ namespace GeometryCollection::Facades
 		/** Get the managed array collection */
 		const FManagedArrayCollection& GetManagedArrayCollection() const {return ConstCollection;}
 		
+		/** Get the point positions of a specific curve */
+		CHAOS_API void GetCurvePointPositions(const int32 CurveIndex, TArray<FVector3f>& OutPoints) const;
+
 	protected :
 
 		/** Build geometry collection from the curves one */
