Dataflow : texture to attribute node : make sure to check if the attribute if of type float is it already exists in the collection

#Synced-CL 50812477

diff --git a/Engine/Plugins/Experimental/Dataflow/Source/DataflowNodes/Private/Dataflow/DataflowTextureToAttributeNode.cpp b/Engine/Plugins/Experimental/Dataflow/Source/DataflowNodes/Private/Dataflow/DataflowTextureToAttributeNode.cpp
--- a/Engine/Plugins/Experimental/Dataflow/Source/DataflowNodes/Private/Dataflow/DataflowTextureToAttributeNode.cpp
+++ b/Engine/Plugins/Experimental/Dataflow/Source/DataflowNodes/Private/Dataflow/DataflowTextureToAttributeNode.cpp
@@ -59,6 +59,7 @@ FDataflowTextureToAttributeNode::FDataflowTextureToAttributeNode(const UE::Dataf
 	RegisterInputConnection(&UVChannel);
 	RegisterInputConnection(&AttributeName);
 	RegisterOutputConnection(&Collection, &Collection);
+	RegisterOutputConnection(&AttributeName, &AttributeName);
 }
 
 FDataflowTextureToAttributeNode::~FDataflowTextureToAttributeNode()
@@ -75,7 +76,11 @@ void FDataflowTextureToAttributeNode::PostSerialize(const FArchive& Ar)
 
 void FDataflowTextureToAttributeNode::Evaluate(UE::Dataflow::FContext& Context, const FDataflowOutput* Out) const
 {
-	if (Out->IsA(&Collection))
+	if (Out->IsA(&AttributeName))
+	{
+		SafeForwardInput(Context, &AttributeName, &AttributeName);
+	}
+	else if (Out->IsA(&Collection))
 	{
 		const FName InAttributeName = FName(GetValue(Context, &AttributeName, AttributeName));
 		if (InAttributeName.IsNone())
@@ -95,6 +100,14 @@ void FDataflowTextureToAttributeNode::Evaluate(UE::Dataflow::FContext& Context,
 			return;
 		}
 
+		const TManagedArray<float>* ExistingAttribute = InCollection.FindAttributeTyped<float>(InAttributeName, InVertexGroupName);
+		if (ExistingAttribute == nullptr && InCollection.HasAttribute(InAttributeName, InVertexGroupName))
+		{
+			Context.Warning(LOCTEXT("ExistingAttributeWrongType_Msg", "Attribute already exist but is not a float type, transfer will be skipped"));
+			SafeForwardInput(Context, &Collection, &Collection);
+			return;
+		}
+
 		const TObjectPtr<UTexture2D> InTexture = GetValue(Context, &Texture, Texture);
 		if (!InTexture)
 		{
