Dataflow : attempt to fix the issue where animation is oscillating between two poses when pausing the timeline
- disable animation outside of teh update from the dataflow simulation scene so that the world tick does not get the opprtunity to animate it too

#Synced-CL 50812477

diff --git a/Engine/Plugins/Experimental/Dataflow/Source/DataflowEditor/Private/Dataflow/DataflowSimulationControls.cpp b/Engine/Plugins/Experimental/Dataflow/Source/DataflowEditor/Private/Dataflow/DataflowSimulationControls.cpp
--- a/Engine/Plugins/Experimental/Dataflow/Source/DataflowEditor/Private/Dataflow/DataflowSimulationControls.cpp
+++ b/Engine/Plugins/Experimental/Dataflow/Source/DataflowEditor/Private/Dataflow/DataflowSimulationControls.cpp
@@ -262,15 +262,23 @@ namespace UE::Dataflow
 			{
 				if(USkeletalMeshComponent* SkeletalMeshComponent = Cast<USkeletalMeshComponent>(PrimComponent))
 				{
-					SkeletalMeshComponent->SetPosition(SimulationTime);
-					SkeletalMeshComponent->TickAnimation(0.f, false /*bNeedsValidRootMotion*/);
-					SkeletalMeshComponent->RefreshBoneTransforms(nullptr /*TickFunction*/);
+					// Make sure only Dataflow can update the animation to avoid the world tick to eventually update it again 
+					// where this can cause issues where the animation jitter between two state when the dataflow timeline is paused
+					// Note : we do not call SetEnableAnimation because it does not work in that way 
+					// we want to explicily enable animation only for this portion fo the code and force it to false when done 
+					SkeletalMeshComponent->bEnableAnimation = true;
+					{
+						SkeletalMeshComponent->SetPosition(SimulationTime);
+						SkeletalMeshComponent->TickAnimation(0.f, false /*bNeedsValidRootMotion*/);
+						SkeletalMeshComponent->RefreshBoneTransforms(nullptr /*TickFunction*/);
 
-					SkeletalMeshComponent->RefreshFollowerComponents();
-					SkeletalMeshComponent->UpdateComponentToWorld();
-					SkeletalMeshComponent->FinalizeBoneTransform();
-					SkeletalMeshComponent->MarkRenderTransformDirty();
-					SkeletalMeshComponent->MarkRenderDynamicDataDirty();
+						SkeletalMeshComponent->RefreshFollowerComponents();
+						SkeletalMeshComponent->UpdateComponentToWorld();
+						SkeletalMeshComponent->FinalizeBoneTransform();
+						SkeletalMeshComponent->MarkRenderTransformDirty();
+						SkeletalMeshComponent->MarkRenderDynamicDataDirty();
+					}
+					SkeletalMeshComponent->bEnableAnimation = false;
 				}
 			}
 		}
