Dataflow : texture to attribute : make attribute name to also be an output

#Synced-CL 50812477

diff --git a/Engine/Plugins/Experimental/Dataflow/Source/DataflowNodes/Private/Dataflow/DataflowTextureToAttributeNode.h b/Engine/Plugins/Experimental/Dataflow/Source/DataflowNodes/Private/Dataflow/DataflowTextureToAttributeNode.h
--- a/Engine/Plugins/Experimental/Dataflow/Source/DataflowNodes/Private/Dataflow/DataflowTextureToAttributeNode.h
+++ b/Engine/Plugins/Experimental/Dataflow/Source/DataflowNodes/Private/Dataflow/DataflowTextureToAttributeNode.h
@@ -40,7 +40,7 @@ private:
 	int32 UVChannel = 0;
 
 	/** Name of the attribute to to transfer the texture color to */
-	UPROPERTY(EditAnywhere, Category = Target, meta = (DataflowInput))
+	UPROPERTY(EditAnywhere, Category = Target, meta = (DataflowInput, DataflowOutput, DataflowPassthrough = AttributeName))
 	FString AttributeName;
 
 	/** Vertex group to use */
