Dataflow : update command loc text to make sure Dataflow starts with a capital D and adjust the capitalization on the word graph too

#Synced-CL 50850263

diff --git a/Engine/Plugins/Experimental/Dataflow/Source/DataflowEditor/Private/Dataflow/DataflowEditorCommands.cpp b/Engine/Plugins/Experimental/Dataflow/Source/DataflowEditor/Private/Dataflow/DataflowEditorCommands.cpp
--- a/Engine/Plugins/Experimental/Dataflow/Source/DataflowEditor/Private/Dataflow/DataflowEditorCommands.cpp
+++ b/Engine/Plugins/Experimental/Dataflow/Source/DataflowEditor/Private/Dataflow/DataflowEditorCommands.cpp
@@ -50,18 +50,18 @@ void FDataflowEditorCommandsImpl::RegisterCommands()
 	TBaseCharacterFXEditorCommands::RegisterCommands();
 	
 	UI_COMMAND(EvaluateNode, "Evaluate", "Trigger an evaluation of the selected node.", EUserInterfaceActionType::Button, FInputChord());
-	UI_COMMAND(EvaluateGraph, "Evaluate Dataflow Graph", "Trigger an evaluation of the graph Terminal node.", EUserInterfaceActionType::Button, FInputChord(EKeys::F5));
+	UI_COMMAND(EvaluateGraph, "Evaluate Dataflow Graph", "Trigger an evaluation of the Dataflow graph terminal nodes.", EUserInterfaceActionType::Button, FInputChord(EKeys::F5));
 	UI_COMMAND(ToggleSimulation, "Simulate Dataflow Scene", "Toggle the simulation of the registered components in the simulation scene.", EUserInterfaceActionType::RadioButton, FInputChord(EKeys::S, EModifierKey::Alt));
-	UI_COMMAND(StartSimulation, "Start Dataflow Simulation", "Start the dataflow simulation", EUserInterfaceActionType::Button, FInputChord());
-	UI_COMMAND(StopSimulation, "Stop Dataflow Simulation", "Stop the dataflow simulation.", EUserInterfaceActionType::Button, FInputChord());
-	UI_COMMAND(PauseSimulation, "Pause Dataflow Simulation", "Pause the dataflow simulation.", EUserInterfaceActionType::Button, FInputChord());
-	UI_COMMAND(StepSimulation, "Step Dataflow Simulation", "Advance one frame in the dataflow simulation.", EUserInterfaceActionType::Button, FInputChord());
-	UI_COMMAND(ResetSimulation, "Reset Dataflow Simulation", "Reset the dataflow simulation.", EUserInterfaceActionType::Button, FInputChord());
-	UI_COMMAND(EvaluateGraphAutomatic, "Automatic Graph Evaluation", "Set the evaluation mode of the graph to Automatic.", EUserInterfaceActionType::RadioButton, FInputChord());
-	UI_COMMAND(EvaluateGraphManual, "Manual Graph Evaluation", "Set the evaluation mode of the graph to Manual.", EUserInterfaceActionType::RadioButton, FInputChord());
-	UI_COMMAND(ClearGraphCache, "Force Graph Re-evaluation", "Force the entire graph to re-evaluate by deleting the evaluation cache and re-evaluating the Terminal node.", EUserInterfaceActionType::Button, FInputChord());
+	UI_COMMAND(StartSimulation, "Start Dataflow Simulation", "Start the simulation of the asset", EUserInterfaceActionType::Button, FInputChord());
+	UI_COMMAND(StopSimulation, "Stop Dataflow Simulation", "Stop the simulation of the asset.", EUserInterfaceActionType::Button, FInputChord());
+	UI_COMMAND(PauseSimulation, "Pause Dataflow Simulation", "Pause the simulation of the asset.", EUserInterfaceActionType::Button, FInputChord());
+	UI_COMMAND(StepSimulation, "Step Dataflow Simulation", "Advance the simulation of the asset by one frame.", EUserInterfaceActionType::Button, FInputChord());
+	UI_COMMAND(ResetSimulation, "Reset Dataflow Simulation", "Reset the simulation of the asset.", EUserInterfaceActionType::Button, FInputChord());
+	UI_COMMAND(EvaluateGraphAutomatic, "Automatic Graph Evaluation", "Set the evaluation mode of the Dataflow graph to Automatic.", EUserInterfaceActionType::RadioButton, FInputChord());
+	UI_COMMAND(EvaluateGraphManual, "Manual Graph Evaluation", "Set the evaluation mode of the Dataflow graph to Manual.", EUserInterfaceActionType::RadioButton, FInputChord());
+	UI_COMMAND(ClearGraphCache, "Force Graph Re-evaluation", "Force the entire Dataflow graph to re-evaluate by deleting the evaluation cache and re-evaluating the terminal nodes.", EUserInterfaceActionType::Button, FInputChord());
 	UI_COMMAND(TogglePerfData, "Performance Data", "Toggle the evaluation performance data for each node.", EUserInterfaceActionType::ToggleButton, FInputChord());
-	UI_COMMAND(ToggleAsyncEvaluation, "Asynchronous Evaluation (Experimental)", "Toggle asynchronous evaluation of the graph. This is an experimental feature.", EUserInterfaceActionType::ToggleButton, FInputChord());
+	UI_COMMAND(ToggleAsyncEvaluation, "Asynchronous Evaluation (Experimental)", "Toggle asynchronous evaluation of the Dataflow graph (This is an experimental feature).", EUserInterfaceActionType::ToggleButton, FInputChord());
 	UI_COMMAND(FreezeNodes, "FreezeNodes", "Freeze the evaluation of the selected nodes.", EUserInterfaceActionType::Button, FInputChord());
 	UI_COMMAND(UnfreezeNodes, "UnfreezeNodes", "Unfreeze the evaluation of the selected nodes.", EUserInterfaceActionType::Button, FInputChord());
 	UI_COMMAND(CreateComment, "CreateComment", "Create a Comment node.", EUserInterfaceActionType::None, FInputChord());
@@ -71,17 +71,17 @@ void FDataflowEditorCommandsImpl::RegisterCommands()
 	UI_COMMAND(ToggleVertexSelection, "ToggleVertexSelection", "Enable vertex selection in editor.", EUserInterfaceActionType::ToggleButton, FInputChord());
 	UI_COMMAND(AddOptionPin, "AddOptionPin", "Add an option pin to the selected nodes.", EUserInterfaceActionType::Button, FInputChord());
 	UI_COMMAND(RemoveOptionPin, "RemoveOptionPin", "Remove the last option pin from the selected nodes.", EUserInterfaceActionType::Button, FInputChord());
-	UI_COMMAND(ZoomToFitGraph, "ZoomToFitGraph", "Fit the graph in the graph editor viewport.", EUserInterfaceActionType::None, FInputChord(EKeys::F));
+	UI_COMMAND(ZoomToFitGraph, "ZoomToFitGraph", "Fit the Dataflow graph in the graph editor viewport.", EUserInterfaceActionType::None, FInputChord(EKeys::F));
 
 	UI_COMMAND(AddWeightMapNode, "Add Weight Map", "Paint weight maps on the mesh", EUserInterfaceActionType::Button, FInputChord());
-	UI_COMMAND(AddNewVariable, "Variable", "Adds a new variable to this dataflow graph.", EUserInterfaceActionType::Button, FInputChord());
-	UI_COMMAND(AddNewSubGraph, "SubGraph", "Adds a new subgraph to this dataflow graph.", EUserInterfaceActionType::Button, FInputChord());
+	UI_COMMAND(AddNewVariable, "Variable", "Adds a new variable to this Dataflow graph.", EUserInterfaceActionType::Button, FInputChord());
+	UI_COMMAND(AddNewSubGraph, "SubGraph", "Adds a new subgraph to this Dataflow graph.", EUserInterfaceActionType::Button, FInputChord());
 	UI_COMMAND(ConvertToBasicSubGraph, "Convert to Basic Subgraph", "Convert the subgraph to a basic one (no loop).", EUserInterfaceActionType::Button, FInputChord());
 	UI_COMMAND(ConvertToForEachSubGraph, "Convert to For Each Subgraph", "Convert the subgraph to be able to iterate through an array.", EUserInterfaceActionType::Button, FInputChord());
 	UI_COMMAND(CollapseToSubGraph, "Collapse to Sub-graph", "Collapse the selection into a subgraph", EUserInterfaceActionType::Button, FInputChord());
 
 	UI_COMMAND(ConvertToEmbeddedDataflow, "Embed Dataflow", "Embed the Dataflow graph in the current asset", EUserInterfaceActionType::Button, FInputChord());
-	UI_COMMAND(ConvertToExternalDataflow, "Extract and Reference Dataflow...", "Convert this Dataflow to be an external Dataflow asset that can be shared by other asset and assign it to this asset", EUserInterfaceActionType::Button, FInputChord());
+	UI_COMMAND(ConvertToExternalDataflow, "Extract and Reference Dataflow...", "Convert this Dataflow graph to be an external Dataflow asset that can be shared by other asset and assign it to this asset", EUserInterfaceActionType::Button, FInputChord());
 	UI_COMMAND(ToggleDataflowEditingLock, "Lock/Unlock Editing", "Toggle editing for the current Dataflow graph.", EUserInterfaceActionType::ToggleButton, FInputChord());
 
 	for (const TPair<FName, TUniquePtr<UE::Dataflow::IDataflowConstructionViewMode>>& NameAndMode : UE::Dataflow::FRenderingViewModeFactory::GetInstance().GetViewModes())
@@ -113,7 +113,7 @@ void FDataflowEditorCommandsImpl::RegisterCommands()
 			ToolRegistry.GetToolCommandForNode(NodeName),
 			FName(NodeName.ToString() + "_Tool"),
 			ToolRegistry.GetAddNodeButtonText(NodeName),					// TODO: Replace placeholder names
-			FText::Format(LOCTEXT("LaunchDataflowToolTooltip", "Launch a \"{0}\" dataflow tool"), FText::FromString(NodeName.ToString())),
+			FText::Format(LOCTEXT("LaunchDataflowToolTooltip", "Launch a \"{0}\" Dataflow tool"), FText::FromString(NodeName.ToString())),
 			ToolRegistry.GetAddNodeButtonIcon(NodeName),
 			EUserInterfaceActionType::RadioButton,
 			FInputChord()
@@ -124,7 +124,7 @@ void FDataflowEditorCommandsImpl::RegisterCommands()
 			ToolRegistry.GetAddNodeCommandForNode(NodeName),
 			FName("Add_" + NodeName.ToString()),
 			ToolRegistry.GetAddNodeButtonText(NodeName),
-			FText::Format(LOCTEXT("AddDataflowNodeTooltip", "Add a \"{0}\" node to the graph"), FText::FromString(NodeName.ToString())),
+			FText::Format(LOCTEXT("AddDataflowNodeTooltip", "Add a \"{0}\" node to the Dataflow graph"), FText::FromString(NodeName.ToString())),
 			ToolRegistry.GetAddNodeButtonIcon(NodeName),
 			EUserInterfaceActionType::Button,
 			FInputChord()
diff --git a/Engine/Plugins/Experimental/Dataflow/Source/DataflowEditor/Private/Dataflow/DataflowEditorToolkit.cpp b/Engine/Plugins/Experimental/Dataflow/Source/DataflowEditor/Private/Dataflow/DataflowEditorToolkit.cpp
--- a/Engine/Plugins/Experimental/Dataflow/Source/DataflowEditor/Private/Dataflow/DataflowEditorToolkit.cpp
+++ b/Engine/Plugins/Experimental/Dataflow/Source/DataflowEditor/Private/Dataflow/DataflowEditorToolkit.cpp
@@ -1015,7 +1015,7 @@ void FDataflowEditorToolkit::AddEvaluationWidget(FToolMenuSection& InSection)
 					FUIAction(),
 					FOnGetContent::CreateSP(Toolkit.ToSharedRef(), &FDataflowEditorToolkit::GenerateEvaluationOptionsMenu),
 					LOCTEXT("DataflowEvaluationOptions", "Options"),
-					LOCTEXT("DataflowEvaluationOptions_ToolbarTooltip", "Options to customize how Dataflow evaluate"),
+					LOCTEXT("DataflowEvaluationOptions_ToolbarTooltip", "Options to customize how the Dataflow Graph evaluates"),
 					TAttribute<FSlateIcon>(),
 					true
 				);
diff --git a/Engine/Plugins/Runtime/HairStrands/Source/HairStrandsEditor/Private/GroomActions.cpp b/Engine/Plugins/Runtime/HairStrands/Source/HairStrandsEditor/Private/GroomActions.cpp
--- a/Engine/Plugins/Runtime/HairStrands/Source/HairStrandsEditor/Private/GroomActions.cpp
+++ b/Engine/Plugins/Runtime/HairStrands/Source/HairStrandsEditor/Private/GroomActions.cpp
@@ -715,7 +715,7 @@ static FDelayedAutoRegisterHelper DelayedAutoRegister(EDelayedRegisterRunPhase::
 
 			{
 				const TAttribute<FText> Label = LOCTEXT("CreateDataflowAsset", "Create Dataflow Asset");
-				const TAttribute<FText> ToolTip = LOCTEXT("CreateDataflowAssetooltip", "Create a dataflow asset from a list of templates");
+				const TAttribute<FText> ToolTip = LOCTEXT("CreateDataflowAssetooltip", "Create a Dataflow asset from a list of templates");
 				const FSlateIcon Icon = FSlateIcon(FAppStyle::GetAppStyleSetName(), "ContentBrowser.AssetActions");
 
 				FToolUIAction UIAction;
