Dataflow paint tool : make sure to refocus on the viewport when interacting with the tool properties so that the shortcuts work properly

#Synced-CL 50850263

diff --git a/Engine/Plugins/Experimental/Dataflow/Source/DataflowEditor/Private/Dataflow/DataflowEditorMode.cpp b/Engine/Plugins/Experimental/Dataflow/Source/DataflowEditor/Private/Dataflow/DataflowEditorMode.cpp
--- a/Engine/Plugins/Experimental/Dataflow/Source/DataflowEditor/Private/Dataflow/DataflowEditorMode.cpp
+++ b/Engine/Plugins/Experimental/Dataflow/Source/DataflowEditor/Private/Dataflow/DataflowEditorMode.cpp
@@ -48,6 +48,9 @@
 #include "Tools/UEdMode.h"
 #include "Selection.h"
 #include "UnrealClient.h"
+#include "Framework/Application/SlateApplication.h"
+#include "SEditorViewport.h"
+#include "Slate/SceneViewport.h"
 
 #include UE_INLINE_GENERATED_CPP_BY_NAME(DataflowEditorMode)
 
@@ -638,6 +641,23 @@ void UDataflowEditorMode::FirstTimeFocusConstructionViewport()
 	}
 }
 
+void UDataflowEditorMode::SetInputFocusOnConstructionViewport()
+{
+	if (TSharedPtr<FDataflowConstructionViewportClient> PinnedConstructionViewportClient = ConstructionViewportClient.Pin())
+	{
+		if (PinnedConstructionViewportClient->Viewport)
+		{
+			if (FSceneViewport* SceneViewport = PinnedConstructionViewportClient->Viewport->AsSceneViewport())
+			{
+				if (TSharedPtr<SWidget> ViewportWidget = SceneViewport->GetViewportWidget().Pin())
+				{
+					FSlateApplication::Get().SetUserFocus(FSlateApplication::Get().GetUserIndexForKeyboard(), ViewportWidget, EFocusCause::SetDirectly);
+				}
+			}
+		}
+	}
+}
+
 void UDataflowEditorMode::FirstTimeFocusSimulationViewport()
 {
 	// If this is the first time seeing a valid 2D or 3D mesh, refocus the camera on it.
diff --git a/Engine/Plugins/Experimental/Dataflow/Source/DataflowEditor/Private/DataflowEditorTools/DataflowEditorVertexAttributePaintTool.cpp b/Engine/Plugins/Experimental/Dataflow/Source/DataflowEditor/Private/DataflowEditorTools/DataflowEditorVertexAttributePaintTool.cpp
--- a/Engine/Plugins/Experimental/Dataflow/Source/DataflowEditor/Private/DataflowEditorTools/DataflowEditorVertexAttributePaintTool.cpp
+++ b/Engine/Plugins/Experimental/Dataflow/Source/DataflowEditor/Private/DataflowEditorTools/DataflowEditorVertexAttributePaintTool.cpp
@@ -27,6 +27,7 @@
 #include "InteractiveGizmoManager.h"
 #include "Intersection/IntrLine2Line2.h"
 #include "Intersection/IntrSegment2Segment2.h"
+#include "IPersonaEditorModeManager.h"
 #include "Framework/Notifications/NotificationManager.h"
 #include "HAL/PlatformApplicationMisc.h"
 #include "Logging/LogMacros.h"
@@ -45,6 +46,7 @@
 #include "Widgets/Notifications/SNotificationList.h"
 #include "Generators/RectangleMeshGenerator.h"
 
+
 #include UE_INLINE_GENERATED_CPP_BY_NAME(DataflowEditorVertexAttributePaintTool)
 
 DEFINE_LOG_CATEGORY_STATIC(LogDataflowEditorVertexAttributePaintTool, Warning, All);
@@ -864,6 +866,19 @@ void UDataflowEditorVertexAttributePaintTool::InitializeMirrorIndicator()
 	UpdateMirrorPlaneMesh();
 }
 
+void UDataflowEditorVertexAttributePaintTool::SetFocusInViewport() const
+{
+	if (ensure(Mode))
+	{
+		Mode->SetInputFocusOnConstructionViewport();
+	}
+
+	//if (UPersonaEditorModeManagerContext* PersonaModeManagerContext = GetToolManager()->GetContextObjectStore()->FindContext<UPersonaEditorModeManagerContext>())
+	//{
+	//	PersonaModeManagerContext->SetFocusInViewport();
+	//}
+}
+
 void UDataflowEditorVertexAttributePaintTool::UpdateMirrorPlaneMesh()
 {
 	if (MirrorPlaneMesh)
@@ -1117,6 +1132,7 @@ void UDataflowEditorVertexAttributePaintTool::CommitResult(UBaseDynamicMeshCompo
 void UDataflowEditorVertexAttributePaintTool::EnableValuePicker(bool bEnable)
 {
 	bPendingPickWeight = bEnable;
+	SetFocusInViewport();
 }
 
 bool UDataflowEditorVertexAttributePaintTool::IsValuePickerEnabled() const
@@ -1127,6 +1143,7 @@ bool UDataflowEditorVertexAttributePaintTool::IsValuePickerEnabled() const
 void UDataflowEditorVertexAttributePaintTool::EnableBrushMirroring(bool bEnable)
 {
 	bBrushMirroringEnabled = bEnable;
+	SetFocusInViewport();
 }
 
 bool UDataflowEditorVertexAttributePaintTool::IsBrushMirroringEnabled() const
@@ -1138,6 +1155,8 @@ void UDataflowEditorVertexAttributePaintTool::SetMirrorPlaneWidgetVisible(bool b
 {
 	bMirrorPlaneWidgetVisible = bVisible;
 	UpdateMirrorPlaneMesh();
+
+	SetFocusInViewport();
 }
 
 bool UDataflowEditorVertexAttributePaintTool::IsMirrorPlaneWidgetVisible()
@@ -1156,6 +1175,8 @@ void UDataflowEditorVertexAttributePaintTool::SetMirrorAxis(EAxis::Type InAxis)
 	{
 		ToolProperties->MirrorProperties.MirrorAxis = InAxis;
 		UpdateMirrorPlaneMesh();
+
+		SetFocusInViewport();
 	}
 }
 
@@ -1169,6 +1190,7 @@ void UDataflowEditorVertexAttributePaintTool::SetMeshSelectionTool(EMeshSelector
 	if (MeshSelector)
 	{
 		MeshSelector->SetSelectionTool(SelectorTool);
+		SetFocusInViewport();
 	}
 }
 
@@ -1407,6 +1429,7 @@ void UDataflowEditorVertexAttributePaintTool::SetBrushAdaptiveSize(float Adaptiv
 		}
 		NotifyOfPropertyChangeByTool(BrushProperties);
 		CalculateBrushRadius();
+		SetFocusInViewport();
 	}
 }
 
@@ -1424,6 +1447,7 @@ void UDataflowEditorVertexAttributePaintTool::SetBrushRadius(float NewRadius)
 		}
 		NotifyOfPropertyChangeByTool(BrushProperties);
 		CalculateBrushRadius();
+		SetFocusInViewport();
 	}
 }
 
@@ -1438,6 +1462,7 @@ void UDataflowEditorVertexAttributePaintTool::SetBrushFalloff(float NewFalloff)
 	{
 		BrushProperties->BrushFalloffAmount = FMath::Clamp(NewFalloff, 0.0f, 1.0f);
 		NotifyOfPropertyChangeByTool(BrushProperties);
+		SetFocusInViewport();
 	}
 }
 
@@ -1466,6 +1491,19 @@ void UDataflowEditorVertexAttributePaintTool::LoadBrushConfig(EDataflowEditorToo
 
 }
 
+void UDataflowEditorVertexAttributePaintTool::SetEditingMode(EDataflowEditorToolEditMode Editmode)
+{
+	if (ToolProperties)
+	{
+		ToolProperties->EditingMode = Editmode;
+		if (Editmode == EDataflowEditorToolEditMode::Mesh)
+		{
+			SetMirrorPlaneWidgetVisible(false);
+		}
+		SetFocusInViewport();
+	}
+}
+
 void UDataflowEditorVertexAttributePaintTool::SetBrushMode(EDataflowEditorToolEditOperation NewBrushMode)
 {
 	if (ToolProperties && BrushProperties)
@@ -1485,12 +1523,23 @@ void UDataflowEditorVertexAttributePaintTool::SetBrushMode(EDataflowEditorToolEd
 
 		ToolProperties->SaveConfig();
 		BrushProperties->SaveConfig();
+
+		SetFocusInViewport();
 	}
 }
 
 bool UDataflowEditorVertexAttributePaintTool::IsInBrushMode() const
 {
-	return (ToolProperties->EditingMode == EDataflowEditorToolEditMode::Brush);
+	return (ToolProperties && ToolProperties->EditingMode == EDataflowEditorToolEditMode::Brush);
+}
+
+void UDataflowEditorVertexAttributePaintTool::SetBrushAreaMode(EMeshVertexPaintBrushAreaType BrushAreaMode)
+{
+	if (ToolProperties)
+	{
+		ToolProperties->BrushProperties.BrushAreaMode = BrushAreaMode;
+		SetFocusInViewport();
+	}
 }
 
 bool UDataflowEditorVertexAttributePaintTool::IsVolumetricBrush() const
@@ -2052,6 +2101,7 @@ void UDataflowEditorVertexAttributePaintTool::GrowSelection() const
 	if (MeshSelector)
 	{
 		MeshSelector->GrowSelection();
+		SetFocusInViewport();
 	}
 }
 
@@ -2060,6 +2110,7 @@ void UDataflowEditorVertexAttributePaintTool::ShrinkSelection() const
 	if (MeshSelector)
 	{
 		MeshSelector->ShrinkSelection();
+		SetFocusInViewport();
 	}
 }
 
@@ -2068,6 +2119,7 @@ void UDataflowEditorVertexAttributePaintTool::InvertSelection() const
 	if (MeshSelector)
 	{
 		MeshSelector->InvertSelection();
+		SetFocusInViewport();
 	}
 }
 
@@ -2076,6 +2128,7 @@ void UDataflowEditorVertexAttributePaintTool::FloodSelection() const
 	if (MeshSelector)
 	{
 		MeshSelector->FloodSelection();
+		SetFocusInViewport();
 	}
 }
 
@@ -2084,6 +2137,7 @@ void UDataflowEditorVertexAttributePaintTool::SelectBorder() const
 	if (MeshSelector)
 	{
 		MeshSelector->SelectBorder();
+		SetFocusInViewport();
 	}
 }
 
@@ -2098,6 +2152,7 @@ void UDataflowEditorVertexAttributePaintTool::SetComponentSelectionMode(ECompone
 		{
 			MeshSelector->SetComponentSelectionMode(NewMode);
 		}
+		SetFocusInViewport();
 	}
 }
 
@@ -2148,6 +2203,7 @@ void UDataflowEditorVertexAttributePaintTool::ApplyValueToSelection(EDataflowEdi
 	{
 		EndChange();
 	}
+	SetFocusInViewport();
 }
 
 void UDataflowEditorVertexAttributePaintTool::PruneSelection(float Threshold)
@@ -2181,6 +2237,8 @@ void UDataflowEditorVertexAttributePaintTool::PruneSelection(float Threshold)
 	UpdateVertexColorOverlay();
 	GetToolManager()->PostInvalidation();
 	EndChange();
+
+	SetFocusInViewport();
 }
 
 namespace UE::DataflowEditorVertexAttributePaintTool::Private
@@ -2209,6 +2267,8 @@ void UDataflowEditorVertexAttributePaintTool::CopyAverageFromSelectionToClipboar
 	// notify user
 	const FText NotificationText = FText::FromString("Selection average value copied to clipboard.");
 	ShowEditorMessage(ELogVerbosity::Log, NotificationText);
+
+	SetFocusInViewport();
 }
 
 void UDataflowEditorVertexAttributePaintTool::PasteValueToSelectionFromClipboard()
@@ -2251,6 +2311,8 @@ void UDataflowEditorVertexAttributePaintTool::PasteValueToSelectionFromClipboard
 	// notify user
 	const FText NotificationText = FText::FromString("Pasted weights.");
 	ShowEditorMessage(ELogVerbosity::Log, NotificationText);
+
+	SetFocusInViewport();
 }
 
 void UDataflowEditorVertexAttributePaintTool::OnSelectionModified()
@@ -3043,6 +3105,8 @@ void UDataflowEditorVertexAttributePaintTool::SetColorMode(EDataflowEditorToolCo
 
 		FPropertyChangedEvent ColorModePropertyChangedEvent(UBrushBaseProperties::StaticClass()->FindPropertyByName(GET_MEMBER_NAME_CHECKED(FDataflowEditorVertexAttributePaintToolDisplayProperties, ColorMode)));
 		ToolProperties->PostEditChangeProperty(ColorModePropertyChangedEvent);
+
+		SetFocusInViewport();
 	}
 }
 
@@ -3154,6 +3218,7 @@ void UDataflowEditorVertexAttributePaintTool::MirrorValues()
 		GetToolManager()->PostInvalidation();
 		EndChange();
 
+		SetFocusInViewport();
 	}
 }
 
diff --git a/Engine/Plugins/Experimental/Dataflow/Source/DataflowEditor/Private/DataflowEditorTools/DataflowEditorVertexAttributePaintToolDetailCustomization.cpp b/Engine/Plugins/Experimental/Dataflow/Source/DataflowEditor/Private/DataflowEditorTools/DataflowEditorVertexAttributePaintToolDetailCustomization.cpp
--- a/Engine/Plugins/Experimental/Dataflow/Source/DataflowEditor/Private/DataflowEditorTools/DataflowEditorVertexAttributePaintToolDetailCustomization.cpp
+++ b/Engine/Plugins/Experimental/Dataflow/Source/DataflowEditor/Private/DataflowEditorTools/DataflowEditorVertexAttributePaintToolDetailCustomization.cpp
@@ -142,13 +142,9 @@ namespace UE::Dataflow::Editor
 					})
 				.OnValueChanged_Lambda([this](EDataflowEditorToolEditMode Mode)
 					{
-						if (Tool.IsValid() && ToolProperties.IsValid())
+						if (Tool.IsValid())
 						{
-							ToolProperties->EditingMode = Mode;
-							if (Mode == EDataflowEditorToolEditMode::Mesh)
-							{
-								Tool->SetMirrorPlaneWidgetVisible(false);
-							}
+							Tool->SetEditingMode(Mode);
 							if (CurrentDetailBuilder)
 							{
 								CurrentDetailBuilder->ForceRefreshDetails();
@@ -685,10 +681,9 @@ namespace UE::Dataflow::Editor
 					})
 				.OnValueChanged_Lambda([this](EMeshVertexPaintBrushAreaType Mode)
 					{
-						if (ToolProperties.IsValid())
+						if (Tool.IsValid())
 						{
-							ToolProperties->BrushProperties.BrushAreaMode = Mode; // TODO(ccaillaud) ==>> 
-							// TODO(ccaillaud) ==>> ToolProperties->SetFalloffMode(Mode);
+							Tool->SetBrushAreaMode(Mode);
 						}
 					})
 				+ SSegmentedControl<EMeshVertexPaintBrushAreaType>::Slot(EMeshVertexPaintBrushAreaType::Connected)
@@ -731,9 +726,9 @@ namespace UE::Dataflow::Editor
 				})
 			.OnValueCommitted_Lambda([this](float NewValue, ETextCommit::Type CommitType)
 				{
-					if (ToolProperties.IsValid())
+					if (Tool.IsValid())
 					{
-						// TODO(ccaillaud) ==>> ToolProperties->SaveConfig();
+						Tool->SetFocusInViewport();
 					}
 				})
 		];
@@ -811,9 +806,9 @@ namespace UE::Dataflow::Editor
 					})
 				.OnValueCommitted_Lambda([this](float NewValue, ETextCommit::Type CommitType)
 					{
-						if (ToolProperties.IsValid())
+						if (Tool.IsValid())
 						{
-							// TODO(ccaillaud) ==>> ToolProperties->SaveConfig();
+							Tool->SetFocusInViewport();
 						}
 					})
 			]
@@ -877,9 +872,9 @@ namespace UE::Dataflow::Editor
 						})
 					.OnValueCommitted_Lambda([this](float NewValue, ETextCommit::Type CommitType)
 						{
-							if (ToolProperties.IsValid())
+							if (Tool.IsValid())
 							{
-								// TODO(ccaillaud) ==>> ToolProperties->SaveConfig();
+								Tool->SetFocusInViewport();
 							}
 						})
 			];
@@ -1027,6 +1022,7 @@ namespace UE::Dataflow::Editor
 						if (Tool.IsValid() && Tool->BrushProperties)
 						{
 							Tool->BrushProperties->bHitBackFaces = (InCheckBoxState == ECheckBoxState::Checked);
+							Tool->SetFocusInViewport();
 						}
 					})
 			];
@@ -1136,11 +1132,15 @@ namespace UE::Dataflow::Editor
 						const bool bIsChecked = (ToolProperties.IsValid() && ToolProperties->MirrorProperties.bHideOnBrushStroke);
 						return bIsChecked ? ECheckBoxState::Checked : ECheckBoxState::Unchecked;
 					})
-				.OnCheckStateChanged_Lambda([ToolProperties = ToolProperties](ECheckBoxState InCheckBoxState)
+				.OnCheckStateChanged_Lambda([Tool = Tool, ToolProperties = ToolProperties](ECheckBoxState InCheckBoxState)
 					{
 						if (ToolProperties.IsValid())
 						{
 							ToolProperties->MirrorProperties.bHideOnBrushStroke = (InCheckBoxState == ECheckBoxState::Checked);
+							if (Tool.IsValid())
+							{
+								Tool->SetFocusInViewport();
+							}
 						}
 					})
 			];
@@ -1652,6 +1652,7 @@ namespace UE::Dataflow::Editor
 					if (Tool.IsValid())
 					{
 						Tool->EndChange();
+						Tool->SetFocusInViewport();
 					}
 					ResetValue(AddMultiplyButtonsWidget->GetValue());
 				})
@@ -1778,9 +1779,9 @@ namespace UE::Dataflow::Editor
 			.ToolTipText(LOCTEXT("PruneValueSliderToolTip", "Prune Threshold - Values below this threshold will be set to zero"))
 			.OnValueCommitted_Lambda([this](float NewValue, ETextCommit::Type CommitType)
 				{
-					if (ToolProperties.IsValid())
+					if (Tool.IsValid())
 					{
-						// TOOD(ccaillaud) ==> ToolProperties->SaveConfig();
+						Tool->SetFocusInViewport();
 					}
 				});
 
@@ -1936,6 +1937,10 @@ namespace UE::Dataflow::Editor
 					if (ToolProperties.IsValid())
 					{
 						ToolProperties->MirrorProperties.MirrorDirection = Mode;
+						if (Tool.IsValid())
+						{
+							Tool->SetFocusInViewport();
+						}
 					}
 				})
 			+ SSegmentedControl<EDataflowEditorToolMirrorDirection>::Slot(EDataflowEditorToolMirrorDirection::PositiveToNegative)
@@ -1991,11 +1996,15 @@ namespace UE::Dataflow::Editor
 					const bool bIsChecked = (ToolProperties.IsValid() && ToolProperties->MirrorProperties.bObjectSpace);
 					return bIsChecked ? ECheckBoxState::Checked : ECheckBoxState::Unchecked;
 				})
-			.OnCheckStateChanged_Lambda([ToolProperties = ToolProperties](ECheckBoxState InCheckBoxState)
+			.OnCheckStateChanged_Lambda([Tool = Tool, ToolProperties = ToolProperties](ECheckBoxState InCheckBoxState)
 				{
 					if (ToolProperties.IsValid())
 					{
 						ToolProperties->MirrorProperties.bObjectSpace = (InCheckBoxState == ECheckBoxState::Checked);
+						if (Tool.IsValid())
+						{
+							Tool->SetFocusInViewport();
+						}
 					}
 				});
 	}
diff --git a/Engine/Plugins/Experimental/Dataflow/Source/DataflowEditor/Public/Dataflow/DataflowEditorMode.h b/Engine/Plugins/Experimental/Dataflow/Source/DataflowEditor/Public/Dataflow/DataflowEditorMode.h
--- a/Engine/Plugins/Experimental/Dataflow/Source/DataflowEditor/Public/Dataflow/DataflowEditorMode.h
+++ b/Engine/Plugins/Experimental/Dataflow/Source/DataflowEditor/Public/Dataflow/DataflowEditorMode.h
@@ -92,6 +92,8 @@ public:
 	UE_API void StartToolForSelectedNode(const UObject* SelectedNode);
 	UE_API void ShutdownActiveToolIfNeeded(EToolShutdownType ShutdownType);
 
+	UE_API void SetInputFocusOnConstructionViewport();
+
 	/** Each node has a set of view modes it can use. For each distinct set of available view modes, we store the last mode that the user has chosen.
 	 *  Use this function to find the most recently used view mode for nodes having the same set of available modes. 
 	 *	Returns nullptr if there is nothing stored for the input set of modes 
diff --git a/Engine/Plugins/Experimental/Dataflow/Source/DataflowEditor/Public/DataflowEditorTools/DataflowEditorVertexAttributePaintTool.h b/Engine/Plugins/Experimental/Dataflow/Source/DataflowEditor/Public/DataflowEditorTools/DataflowEditorVertexAttributePaintTool.h
--- a/Engine/Plugins/Experimental/Dataflow/Source/DataflowEditor/Public/DataflowEditorTools/DataflowEditorVertexAttributePaintTool.h
+++ b/Engine/Plugins/Experimental/Dataflow/Source/DataflowEditor/Public/DataflowEditorTools/DataflowEditorVertexAttributePaintTool.h
@@ -435,8 +435,12 @@ public:
 
 	UE_API virtual FBox GetWorldSpaceFocusBox() override;
 
+	UE_API void SetEditingMode(EDataflowEditorToolEditMode Editmode);
+
 	UE_API void SetBrushMode(EDataflowEditorToolEditOperation BrushMode);
 	UE_API bool IsInBrushMode() const;
+
+	UE_API void SetBrushAreaMode(EMeshVertexPaintBrushAreaType BrushAreaMode);
 	UE_API bool IsVolumetricBrush() const;
 
 	UE_API virtual void CommitResult(UBaseDynamicMeshComponent* Component, bool bModifiedTopology) override;
@@ -495,6 +499,8 @@ public:
 
 	UE_API TArray<int32> GetSelectedVertices() const;
 
+	UE_API void SetFocusInViewport() const;
+
 public:
 	static FInputChord GreyscaleDisplayInputChord;
 	static FInputChord ColorMapDisplayInputChord;
