Dataflow paint tool : add support for display value at selection in the mesh tab

#Synced-CL 50850263

diff --git a/Engine/Plugins/Experimental/Dataflow/Source/DataflowEditor/Private/DataflowEditorTools/DataflowEditorVertexAttributePaintTool.cpp b/Engine/Plugins/Experimental/Dataflow/Source/DataflowEditor/Private/DataflowEditorTools/DataflowEditorVertexAttributePaintTool.cpp
--- a/Engine/Plugins/Experimental/Dataflow/Source/DataflowEditor/Private/DataflowEditorTools/DataflowEditorVertexAttributePaintTool.cpp
+++ b/Engine/Plugins/Experimental/Dataflow/Source/DataflowEditor/Private/DataflowEditorTools/DataflowEditorVertexAttributePaintTool.cpp
@@ -1380,7 +1380,7 @@ FBox UDataflowEditorVertexAttributePaintTool::GetWorldSpaceFocusBox()
 	{
 		if (PreviewMesh && PreviewMesh->GetMesh())
 		{
-			TArray<int32> SelectedVertices = GetSelectedVertices();
+			const TArray<int32>& SelectedVertices = GetSelectedVertices();
 			if (!SelectedVertices.IsEmpty())
 			{
 				const FDynamicMesh3* Mesh = PreviewMesh->GetMesh();
@@ -2156,13 +2156,14 @@ void UDataflowEditorVertexAttributePaintTool::SetComponentSelectionMode(ECompone
 	}
 }
 
-TArray<int32> UDataflowEditorVertexAttributePaintTool::GetSelectedVertices() const
+const TArray<int32>& UDataflowEditorVertexAttributePaintTool::GetSelectedVertices() const
 {
+	static TArray<int32> EmptySelection;
 	if (MeshSelector)
 	{
 		return MeshSelector->GetSelectedVertices();
 	}
-	return {};
+	return EmptySelection;
 }
 
 void UDataflowEditorVertexAttributePaintTool::ApplyValueToSelection(EDataflowEditorToolEditOperation Operation, float InValue, bool bWithTransaction)
@@ -2213,7 +2214,7 @@ void UDataflowEditorVertexAttributePaintTool::PruneSelection(float Threshold)
 		return;
 	}
 
-	TArray<int32> SelectVertices = UDataflowEditorVertexAttributePaintTool::GetSelectedVertices();
+	const TArray<int32>& SelectVertices = UDataflowEditorVertexAttributePaintTool::GetSelectedVertices();
 	if (SelectVertices.IsEmpty())
 	{
 		return;
@@ -2250,7 +2251,7 @@ void UDataflowEditorVertexAttributePaintTool::CopyAverageFromSelectionToClipboar
 {
 	using namespace UE::DataflowEditorVertexAttributePaintTool::Private;
 
-	TArray<int32> SelectVertices = GetSelectedVertices();
+	const TArray<int32>& SelectVertices = GetSelectedVertices();
 	if (SelectVertices.IsEmpty())
 	{
 		const FText NotificationText = FText::FromString(TEXT("No vertices were selected. Nothing was copied to the clipboard."));
@@ -2666,25 +2667,43 @@ void UDataflowEditorVertexAttributePaintTool::DrawHUD(FCanvas* Canvas, IToolsCon
 						}
 					});
 
-				for (int32 VertexID = 0; VertexID < MaxVertexId; ++VertexID)
+				auto DrawValue = 
+					[this, &VertexVisibility, &MinValue, &MaxValue, &ValueRange, &NumberFormattingOptions, &Canvas, &RenderAPI, &Mesh]
+					(int32 VertexID)
+					{
+						if (VertexVisibility.IsValidIndex(VertexID) && VertexVisibility[VertexID])
+						{
+							const float NormalizedValue = VertexData.GetValue(VertexID);
+							const float RemappedValue = MinValue + NormalizedValue * ValueRange;
+							const FText ValueText = FText::AsNumber(RemappedValue, &NumberFormattingOptions);
+
+							const FVector3d VtxPosition = Mesh->GetVertex(VertexID);
+							const FVector4 ScreenBrushPos = RenderAPI->GetSceneView()->WorldToScreen(VtxPosition);
+							FVector2D ScreenBrushPos2D = FVector2D(ScreenBrushPos.X, ScreenBrushPos.Y);
+							RenderAPI->GetSceneView()->ScreenToPixel(ScreenBrushPos, ScreenBrushPos2D);
+
+							FCanvasTextItem TextItem(ScreenBrushPos2D, ValueText, GEngine->GetLargeFont(), FLinearColor::White);
+							TextItem.EnableShadow(FLinearColor::Black);
+							TextItem.bCentreX = true;
+							Canvas->DrawItem(TextItem);
+						}
+					};
+
+				if (IsInBrushMode() || ToolProperties->DisplayProperties.bShowValuesOnlySelected == false)
 				{
-					if (VertexVisibility[VertexID])
+					for (int32 VertexID = 0; VertexID < MaxVertexId; ++VertexID)
 					{
-						const float NormalizedValue = VertexData.GetValue(VertexID);
-						const float RemappedValue = MinValue + NormalizedValue * ValueRange;
-						const FText ValueText = FText::AsNumber(RemappedValue, &NumberFormattingOptions);
-
-						const FVector3d VtxPosition = Mesh->GetVertex(VertexID);
-						const FVector4 ScreenBrushPos = RenderAPI->GetSceneView()->WorldToScreen(VtxPosition);
-						FVector2D ScreenBrushPos2D = FVector2D(ScreenBrushPos.X, ScreenBrushPos.Y);
-						RenderAPI->GetSceneView()->ScreenToPixel(ScreenBrushPos, ScreenBrushPos2D);
-
-						FCanvasTextItem TextItem(ScreenBrushPos2D, ValueText, GEngine->GetLargeFont(), FLinearColor::White);
-						TextItem.EnableShadow(FLinearColor::Black);
-						TextItem.bCentreX = true;
-						Canvas->DrawItem(TextItem);
+						DrawValue(VertexID);
 					}
-				};
+				}
+				else
+				{
+					const TArray<int32>& SelectedVertices = GetSelectedVertices();
+					for (const int32 VertexID : SelectedVertices)
+					{
+						DrawValue(VertexID);
+					}
+				}
 			}
 		}
 	}
@@ -2757,27 +2776,17 @@ void UDataflowEditorVertexAttributePaintTool::OnTick(float DeltaTime)
 	}
 
 	// Get value at brush location
-	const bool bShouldPickWeight = bPendingPickWeight && IsStampPending() == false;
-//	const bool bShouldUpdateValueAtBrush = IsInBrushSubMode();
-
-	//if (bShouldPickWeight /*|| bShouldUpdateValueAtBrush*/)
+	if (IsInBrushMode())
 	{
 		if (GetSculptMesh()->IsTriangle(GetBrushTriangleID()))
 		{
 			const double HitWeightValue = GetCurrentWeightValueUnderBrush();
-
-			//if (bShouldPickWeight)
-			//{
-			//	ToolProperties->BrushProperties.AttributeValue = HitWeightValue;
-			//	NotifyOfPropertyChangeByTool(ToolProperties);
-			//}
-
-			/*if (bShouldUpdateValueAtBrush)*/
-			{
-				ToolProperties->BrushProperties.ValueAtBrush = HitWeightValue;
-			}
+			ToolProperties->BrushProperties.ValueAtBrush = HitWeightValue;
 		}
-		//bPendingPickWeight = false;
+	}
+	else
+	{
+		ToolProperties->BrushProperties.ValueAtBrush = GetAverageWeightValueFromSelection();
 	}
 
 	auto ExecuteStampOperation = [this](int StampIndex, const FRay& StampRay)
@@ -3064,6 +3073,11 @@ int32 UDataflowEditorVertexAttributePaintTool::GetBrushNearestVertexAccurate() c
 	return NearestVertexIndex;
 }
 
+float UDataflowEditorVertexAttributePaintTool::GetAverageWeightValueFromSelection() const
+{
+	return VertexData.GetAverageValue(GetSelectedVertices());
+}
+
 void UDataflowEditorVertexAttributePaintTool::UpdateBrushType(EDataflowEditorToolEditOperation BrushMode)
 {
 	static const FText BaseMessage = LOCTEXT("OnStartTool", "Hold [Shift] to Relax values, Hold [Ctrl] to Invert Value, [/] and S/D change Size (+Shift to small-step)");
diff --git a/Engine/Plugins/Experimental/Dataflow/Source/DataflowEditor/Private/DataflowEditorTools/DataflowEditorVertexAttributePaintToolDetailCustomization.cpp b/Engine/Plugins/Experimental/Dataflow/Source/DataflowEditor/Private/DataflowEditorTools/DataflowEditorVertexAttributePaintToolDetailCustomization.cpp
--- a/Engine/Plugins/Experimental/Dataflow/Source/DataflowEditor/Private/DataflowEditorTools/DataflowEditorVertexAttributePaintToolDetailCustomization.cpp
+++ b/Engine/Plugins/Experimental/Dataflow/Source/DataflowEditor/Private/DataflowEditorTools/DataflowEditorVertexAttributePaintToolDetailCustomization.cpp
@@ -479,11 +479,49 @@ namespace UE::Dataflow::Editor
 							const bool bIsChecked = (ToolProperties.IsValid() && ToolProperties->DisplayProperties.bShowValues);
 							return bIsChecked ? ECheckBoxState::Checked : ECheckBoxState::Unchecked;
 						})
-					.OnCheckStateChanged_Lambda([ToolProperties = ToolProperties](ECheckBoxState InCheckBoxState)
+					.OnCheckStateChanged_Lambda([Tool = Tool, ToolProperties = ToolProperties](ECheckBoxState InCheckBoxState)
 						{
 							if (ToolProperties.IsValid())
 							{
 								ToolProperties->DisplayProperties.bShowValues = (InCheckBoxState == ECheckBoxState::Checked);
+								if (Tool.IsValid())
+								{
+									Tool->SetFocusInViewport();
+								}
+							}
+						})
+			];
+
+		MeshDisplayCategory.AddCustomRow(LOCTEXT("ValuesDisplay_OnlySelected_Row", "Values Display - Only Selected"), false)
+			.Visibility(TAttribute<EVisibility>::CreateLambda(
+				[Tool = Tool]() 
+				{ 
+					return (Tool.IsValid() && !Tool->IsInBrushMode()) ? EVisibility::Visible: EVisibility::Collapsed;
+				}))
+			.NameContent()
+			[
+				SNew(STextBlock)
+					.Text(LOCTEXT("ValuesDisplay_OnlySelected_Label", "Only Selected"))
+					.Font(FAppStyle::GetFontStyle(TEXT("PropertyWindow.NormalFont")))
+					.ToolTipText(LOCTEXT("ValuesDisplay_onlySelected_Tooltip", "In mesh selection mode, only show values for the corresponding selected vertices"))
+			]
+			.ValueContent()
+			[
+				SNew(SCheckBox)
+					.IsChecked_Lambda([ToolProperties = ToolProperties]()
+						{
+							const bool bIsChecked = (ToolProperties.IsValid() && ToolProperties->DisplayProperties.bShowValuesOnlySelected);
+							return bIsChecked ? ECheckBoxState::Checked : ECheckBoxState::Unchecked;
+						})
+					.OnCheckStateChanged_Lambda([Tool= Tool, ToolProperties = ToolProperties](ECheckBoxState InCheckBoxState)
+						{
+							if (ToolProperties.IsValid())
+							{
+								ToolProperties->DisplayProperties.bShowValuesOnlySelected = (InCheckBoxState == ECheckBoxState::Checked);
+								if (Tool.IsValid())
+								{
+									Tool->SetFocusInViewport();
+								}
 							}
 						})
 			];
@@ -888,9 +926,19 @@ namespace UE::Dataflow::Editor
 			.NameContent()
 			[
 				SNew(STextBlock)
-					.Text(LOCTEXT("ValueAtBrushLabel", "Value At Brush"))
+					.Text_Lambda([Tool = Tool]()
+						{
+							return (Tool.IsValid() && Tool->IsInBrushMode())
+								? LOCTEXT("ValueAtBrushLabel", "Value At Brush")
+								: LOCTEXT("ValueOfSelectionLabel", "Value Of Selection");
+						})
 					.Font(FAppStyle::GetFontStyle(TEXT("PropertyWindow.NormalFont")))
-					.ToolTipText(LOCTEXT("ValueAtBrushTooltip", "The value where the brush is loclated. The exact method used to compute it can be changed in the combobox below"))
+					.ToolTipText_Lambda([Tool = Tool]()
+						{
+							return (Tool.IsValid() && Tool->IsInBrushMode())
+								? LOCTEXT("ValueAtBrushTooltip", "The value where the brush is loclated. The exact method used to compute it can be changed in the combobox below")
+								: LOCTEXT("ValueOfSelectionTooltip", "The average value of selected vertices.");
+						})
 			]
 			.ValueContent()
 			[
@@ -907,6 +955,11 @@ namespace UE::Dataflow::Editor
 	void FVertexAttributePaintToolDetailCustomization::AddBrushValueQueryMethodRow(IDetailCategoryBuilder& BrushCategory) const
 	{
 		BrushCategory.AddCustomRow(LOCTEXT("ValueQueryMethodBrushCategory", "ValueAtBrush"), false)
+			.Visibility(TAttribute<EVisibility>::CreateLambda(
+				[Tool = Tool]() 
+				{ 
+					return (Tool.IsValid() && Tool->IsInBrushMode()) ? EVisibility::Visible : EVisibility::Collapsed;
+				}))
 			.NameContent()
 			[
 				SNew(STextBlock)
diff --git a/Engine/Plugins/Experimental/Dataflow/Source/DataflowEditor/Public/DataflowEditorTools/DataflowEditorVertexAttributePaintTool.h b/Engine/Plugins/Experimental/Dataflow/Source/DataflowEditor/Public/DataflowEditorTools/DataflowEditorVertexAttributePaintTool.h
--- a/Engine/Plugins/Experimental/Dataflow/Source/DataflowEditor/Public/DataflowEditorTools/DataflowEditorVertexAttributePaintTool.h
+++ b/Engine/Plugins/Experimental/Dataflow/Source/DataflowEditor/Public/DataflowEditorTools/DataflowEditorVertexAttributePaintTool.h
@@ -204,6 +204,9 @@ struct FDataflowEditorVertexAttributePaintToolDisplayProperties
 	UPROPERTY()
 	bool bShowValues = false;
 
+	UPROPERTY()
+	bool bShowValuesOnlySelected = false;
+
 	UPROPERTY()
 	float MinValue = 0;
 
@@ -497,7 +500,7 @@ public:
 	UE_API void ApplyValueToSelection(EDataflowEditorToolEditOperation Operation, float InValue, bool bWithTransaction = true);
 	UE_API void MirrorValues();
 
-	UE_API TArray<int32> GetSelectedVertices() const;
+	UE_API const TArray<int32>& GetSelectedVertices() const;
 
 	UE_API void SetFocusInViewport() const;
 
@@ -622,6 +625,7 @@ protected:
 	FVector3d CurrentBaryCentricCoords;
 	UE_API int32 GetBrushNearestVertex() const;
 	UE_API int32 GetBrushNearestVertexAccurate() const;
+	UE_API float GetAverageWeightValueFromSelection() const;
 
 	TSet<int32> AccumulatedTriangleROI;
 	bool bUndoUpdatePending = false;
