Dataflow : make graph paste operation more resilient to truncated pasted text

#Synced-CL 50850263

diff --git a/Engine/Plugins/Experimental/Dataflow/Source/DataflowEditor/Private/Dataflow/DataflowAssetEditUtils.cpp b/Engine/Plugins/Experimental/Dataflow/Source/DataflowEditor/Private/Dataflow/DataflowAssetEditUtils.cpp
--- a/Engine/Plugins/Experimental/Dataflow/Source/DataflowEditor/Private/Dataflow/DataflowAssetEditUtils.cpp
+++ b/Engine/Plugins/Experimental/Dataflow/Source/DataflowEditor/Private/Dataflow/DataflowAssetEditUtils.cpp
@@ -799,43 +799,48 @@ namespace UE::Dataflow
 				FString NodeOut, PropertyOut, TypeOut;
 				FDataflowConnectionData::GetNodePropertyAndType(Connection.Out, NodeOut, PropertyOut, TypeOut);
 
-				ensure(TypeIn == TypeOut);
-
-				const UDataflowEdNode* EdNodeIn = OriginalNodeNameToEdNode[NodeIn];
-				const UDataflowEdNode* EdNodeOut = OriginalNodeNameToEdNode[NodeOut];
+				if (TypeIn == TypeOut)
+				{
+					const UDataflowEdNode* EdNodeIn = OriginalNodeNameToEdNode[NodeIn];
+					const UDataflowEdNode* EdNodeOut = OriginalNodeNameToEdNode[NodeOut];
 
-				const FGuid GuidIn = EdNodeIn? EdNodeIn->DataflowNodeGuid: FGuid();
-				const FGuid GuidOut = EdNodeOut ? EdNodeOut->DataflowNodeGuid : FGuid();
+					const FGuid GuidIn = EdNodeIn ? EdNodeIn->DataflowNodeGuid : FGuid();
+					const FGuid GuidOut = EdNodeOut ? EdNodeOut->DataflowNodeGuid : FGuid();
 
-				const FName InputputName = *PropertyIn;
-				const FName OutputputName = *PropertyOut;
+					const FName InputputName = *PropertyIn;
+					const FName OutputputName = *PropertyOut;
 
-				if (TSharedPtr<UE::Dataflow::FGraph> DataflowGraph = DataflowAsset.GetDataflow())
-				{
-					if (TSharedPtr<FDataflowNode> DataflowNodeFrom = DataflowGraph->FindBaseNode(GuidOut))
+					if (TSharedPtr<UE::Dataflow::FGraph> DataflowGraph = DataflowAsset.GetDataflow())
 					{
-						if (TSharedPtr<FDataflowNode> DataflowNodeTo = DataflowGraph->FindBaseNode(GuidIn))
+						if (TSharedPtr<FDataflowNode> DataflowNodeFrom = DataflowGraph->FindBaseNode(GuidOut))
 						{
-							FDataflowInput* InputConnection = DataflowNodeTo->FindInput(InputputName);
-							FDataflowOutput* OutputConnection = DataflowNodeFrom->FindOutput(OutputputName);
+							if (TSharedPtr<FDataflowNode> DataflowNodeTo = DataflowGraph->FindBaseNode(GuidIn))
+							{
+								FDataflowInput* InputConnection = DataflowNodeTo->FindInput(InputputName);
+								FDataflowOutput* OutputConnection = DataflowNodeFrom->FindOutput(OutputputName);
 
-							// make sure we set the right type before attempting any connection
-							DataflowNodeFrom->TrySetConnectionType(OutputConnection, FName(TypeOut));
-							DataflowNodeTo->TrySetConnectionType(InputConnection, FName(TypeIn));
+								// make sure we set the right type before attempting any connection
+								DataflowNodeFrom->TrySetConnectionType(OutputConnection, FName(TypeOut));
+								DataflowNodeTo->TrySetConnectionType(InputConnection, FName(TypeIn));
 
-							// first connect the edgraph as this may affect the dataflow inputs ( for AnyType )
-							UEdGraphPin* OutputPin = FindPin(EdNodeOut, EEdGraphPinDirection::EGPD_Output, OutputputName);
-							UEdGraphPin* InputPin = FindPin(EdNodeIn, EEdGraphPinDirection::EGPD_Input, InputputName);
-							if (OutputPin && InputPin)
-							{
-								EdGraph->GetSchema()->TryCreateConnection(OutputPin, InputPin);
-							}
+								// first connect the edgraph as this may affect the dataflow inputs ( for AnyType )
+								UEdGraphPin* OutputPin = FindPin(EdNodeOut, EEdGraphPinDirection::EGPD_Output, OutputputName);
+								UEdGraphPin* InputPin = FindPin(EdNodeIn, EEdGraphPinDirection::EGPD_Input, InputputName);
+								if (OutputPin && InputPin)
+								{
+									EdGraph->GetSchema()->TryCreateConnection(OutputPin, InputPin);
+								}
 
-							// now connect the dataflow
-							DataflowGraph->Connect(OutputConnection, InputConnection);
+								// now connect the dataflow
+								DataflowGraph->Connect(OutputConnection, InputConnection);
+							}
 						}
 					}
 				}
+				else
+				{
+					UE_LOG(LogChaosDataflow, Error, TEXT("Failed to reconnect output [%s] to input [%s] - incompatible types"), *Connection.Out, *Connection.In);
+				}
 			}
 			return EChangeResult::Changed;
 		};
