Dataflow : fix issues with the Convert Dataflow to Embed/External and lock/Unlock Dataflow button , not working properly when muliple Dataflow Editors were opened

#Synced-CL 50850263

diff --git a/Engine/Plugins/Experimental/Dataflow/Source/DataflowEditor/Private/Dataflow/DataflowEditorToolkit.cpp b/Engine/Plugins/Experimental/Dataflow/Source/DataflowEditor/Private/Dataflow/DataflowEditorToolkit.cpp
--- a/Engine/Plugins/Experimental/Dataflow/Source/DataflowEditor/Private/Dataflow/DataflowEditorToolkit.cpp
+++ b/Engine/Plugins/Experimental/Dataflow/Source/DataflowEditor/Private/Dataflow/DataflowEditorToolkit.cpp
@@ -1104,13 +1104,15 @@ void FDataflowEditorToolkit::AddDataflowActionWidget(FToolMenuSection& InSection
 		FIsActionButtonVisible::CreateSP(this, &FDataflowEditorToolkit::IsExternalDataflow)
 	);
 
-	FToolMenuEntry ConvertToEmbeddedButton = FToolMenuEntry::InitToolBarButton(
-		FDataflowEditorCommands::Get().ConvertToEmbeddedDataflow,
-		TAttribute<FText>(),
-		TAttribute<FText>(),
-		FSlateIcon());
-
-	InSection.AddEntry(ConvertToEmbeddedButton);
+	InSection.AddDynamicEntry("DataflowActionConvertToEmbed", FNewToolMenuSectionDelegate::CreateLambda([](FToolMenuSection& DynamicSection)
+		{
+			DynamicSection.AddEntry(FToolMenuEntry::InitToolBarButton(
+				FDataflowEditorCommands::Get().ConvertToEmbeddedDataflow,
+				TAttribute<FText>(),
+				TAttribute<FText>(),
+				FSlateIcon()
+			));
+		}));
 
 	// Convert to external Dataflow button
 	ToolkitCommands->MapAction(FDataflowEditorCommands::Get().ConvertToExternalDataflow,
@@ -1120,13 +1122,15 @@ void FDataflowEditorToolkit::AddDataflowActionWidget(FToolMenuSection& InSection
 		FIsActionButtonVisible::CreateSP(this, &FDataflowEditorToolkit::IsEmbeddedDataflow)
 	);
 
-	FToolMenuEntry ConvertToExternalButton = FToolMenuEntry::InitToolBarButton(
-		FDataflowEditorCommands::Get().ConvertToExternalDataflow,
-		TAttribute<FText>(),
-		TAttribute<FText>(),
-		FSlateIcon());
-
-	InSection.AddEntry(ConvertToExternalButton);
+	InSection.AddDynamicEntry("DataflowActionConvertToExternal", FNewToolMenuSectionDelegate::CreateLambda([](FToolMenuSection& DynamicSection)
+		{
+			DynamicSection.AddEntry(FToolMenuEntry::InitToolBarButton(
+				FDataflowEditorCommands::Get().ConvertToExternalDataflow,
+				TAttribute<FText>(),
+				TAttribute<FText>(),
+				FSlateIcon()
+			));
+		}));
 
 	// Lock/Unlock editing button
 	ToolkitCommands->MapAction(FDataflowEditorCommands::Get().ToggleDataflowEditingLock,
@@ -1136,22 +1140,48 @@ void FDataflowEditorToolkit::AddDataflowActionWidget(FToolMenuSection& InSection
 		FIsActionButtonVisible::CreateSP(this, &FDataflowEditorToolkit::IsEditingLockButtonVisible)
 	);
 
-	const FSlateIcon UnlockedIcon = FSlateIcon(FAppStyle::GetAppStyleSetName(), "Icons.Unlock");
-	const FSlateIcon LockedIcon = FSlateIcon(FAppStyle::GetAppStyleSetName(), "Icons.Lock");
-	const FText UnlockedText = LOCTEXT("LockEditing_Text", "Lock Editing");
-	const FText UnlockedTooltip = LOCTEXT("lockEditing_Tooltip", "Lock Editing for this external Dataflow asset so that other assets using will not be affected by any potential changes (Default behavior)");
-	const FText LockedText = LOCTEXT("UnlockEditing_Text", "Unlock Editing");
-	const FText LockedTooltip = LOCTEXT("UnlockEditing_Tooltip", "Unlock Editing for this external Dataflow asset\nWarning : this asset may be used by other asset and changes made will affect them");
-
-	FToolMenuEntry ToggleEditingLockButton = FToolMenuEntry::InitToolBarButton(
-		FDataflowEditorCommands::Get().ToggleDataflowEditingLock,
-		TAttribute<FText>::CreateSPLambda(this, [this, UnlockedText, LockedText]() { return bAllowExternalDataflowEditing ? UnlockedText : LockedText; } ),
-		TAttribute<FText>::CreateSPLambda(this, [this, UnlockedTooltip, LockedTooltip]() { return bAllowExternalDataflowEditing ? UnlockedTooltip : LockedTooltip; }),
-		TAttribute<FSlateIcon>::CreateSPLambda(this, [this, UnlockedIcon, LockedIcon]() { return bAllowExternalDataflowEditing ? UnlockedIcon : LockedIcon; })
-	);
-
-	InSection.AddEntry(ToggleEditingLockButton);
-
+	InSection.AddDynamicEntry("DataflowActionLockUnlockGraph", FNewToolMenuSectionDelegate::CreateLambda([](FToolMenuSection& DynamicSection)
+		{
+			using namespace UE::Dataflow::Private;
+			if (const TSharedPtr<FDataflowEditorToolkit> Toolkit = GetAssetEditorToolkitFromMenuContext(DynamicSection.FindContext<UAssetEditorToolkitMenuContext>()))
+			{
+				const FSlateIcon UnlockedIcon = FSlateIcon(FAppStyle::GetAppStyleSetName(), "Icons.Unlock");
+				const FSlateIcon LockedIcon = FSlateIcon(FAppStyle::GetAppStyleSetName(), "Icons.Lock");
+				const FText UnlockedText = LOCTEXT("LockEditing_Text", "Lock Editing");
+				const FText UnlockedTooltip = LOCTEXT("lockEditing_Tooltip", "Lock Editing for this external Dataflow asset so that other assets using will not be affected by any potential changes (Default behavior)");
+				const FText LockedText = LOCTEXT("UnlockEditing_Text", "Unlock Editing");
+				const FText LockedTooltip = LOCTEXT("UnlockEditing_Tooltip", "Unlock Editing for this external Dataflow asset\nWarning : this asset may be used by other asset and changes made will affect them");
+
+				TWeakPtr<FDataflowEditorToolkit> WeakToolkit = Toolkit;
+				DynamicSection.AddEntry(FToolMenuEntry::InitToolBarButton(
+					FDataflowEditorCommands::Get().ToggleDataflowEditingLock,
+					TAttribute<FText>::CreateLambda([WeakToolkit, UnlockedText, LockedText]() 
+						{ 
+							if (TSharedPtr<FDataflowEditorToolkit> ToolkitPtr = WeakToolkit.Pin())
+							{
+								return ToolkitPtr->bAllowExternalDataflowEditing ? UnlockedText : LockedText;
+							}
+							return UnlockedText;
+						}),
+					TAttribute<FText>::CreateLambda([WeakToolkit, UnlockedTooltip, LockedTooltip]() 
+						{ 
+							if (TSharedPtr<FDataflowEditorToolkit> ToolkitPtr = WeakToolkit.Pin())
+							{
+								return ToolkitPtr->bAllowExternalDataflowEditing ? UnlockedTooltip : LockedTooltip;
+							}
+							return UnlockedTooltip;
+						}),
+					TAttribute<FSlateIcon>::CreateLambda([WeakToolkit, UnlockedIcon, LockedIcon]()
+						{ 
+							if (TSharedPtr<FDataflowEditorToolkit> ToolkitPtr = WeakToolkit.Pin())
+							{
+								return ToolkitPtr->bAllowExternalDataflowEditing ? UnlockedIcon : LockedIcon;
+							}
+							return UnlockedIcon;
+						})
+				));
+			}
+		}));
 }
 
 bool FDataflowEditorToolkit::IsEmbeddedDataflow() const
