copy the FReferenceSkeleton to the dataflow skeleton debug draw object rather than trying to hold an unsafe raw reference to its address

#Synced-CL 50857708

diff --git a/Engine/Plugins/Experimental/Dataflow/Source/DataflowNodes/Public/Dataflow/DataflowSkeletalMeshNodes.h b/Engine/Plugins/Experimental/Dataflow/Source/DataflowNodes/Public/Dataflow/DataflowSkeletalMeshNodes.h
--- a/Engine/Plugins/Experimental/Dataflow/Source/DataflowNodes/Public/Dataflow/DataflowSkeletalMeshNodes.h
+++ b/Engine/Plugins/Experimental/Dataflow/Source/DataflowNodes/Public/Dataflow/DataflowSkeletalMeshNodes.h
@@ -50,7 +50,7 @@ private :
 	TArray<FTransform> TransformOverrides;
 
 	/** Skeletal mesh to use to populate/draw scene elements */
-	const FReferenceSkeleton& ReferenceSkeleton;
+	const FReferenceSkeleton ReferenceSkeleton;
 		
 	/** Previous element selection */
 	TArray<bool> PreviousSelection;
