Dataflow:
- Added Sampler suite to Dataflow

#Synced-CL 51001165

diff --git a/Engine/Plugins/Experimental/Dataflow/Source/DataflowNodes/Private/Dataflow/DataflowNodesPlugin.cpp b/Engine/Plugins/Experimental/Dataflow/Source/DataflowNodes/Private/Dataflow/DataflowNodesPlugin.cpp
--- a/Engine/Plugins/Experimental/Dataflow/Source/DataflowNodes/Private/Dataflow/DataflowNodesPlugin.cpp
+++ b/Engine/Plugins/Experimental/Dataflow/Source/DataflowNodes/Private/Dataflow/DataflowNodesPlugin.cpp
@@ -19,8 +19,10 @@
 #include "Dataflow/DataflowVisualizeAttributeNode.h"
 #include "Dataflow/DataflowSamplerToAttributeNode.h"
 #include "Dataflow/DataflowGetCollectionBoundsNode.h"
+#include "Dataflow/SamplerNodes/DataflowSamplerTypes.h"
 #include "Dataflow/SamplerNodes/DataflowGradientSamplerNode.h"
 #include "Dataflow/SamplerNodes/DataflowTransformSamplerNode.h"
+#include "Dataflow/SamplerNodes/DataflowSamplerTypes.h"
 
 #define LOCTEXT_NAMESPACE "DataflowNodes"
 
@@ -62,10 +64,11 @@ void IDataflowNodesPlugin::StartupModule()
 	UE::Dataflow::RegisterVertexColorToAttributeNodes();
 	UE::Dataflow::RegisterVisualizeAttributeNodes();
 
+	UE::Dataflow::RegisterSamplerTypes();
+
 	UE::Dataflow::RegisterSamplerToAttributeNodes();
-	UE::Dataflow::RegisterGradientSamplerNodes();
-	UE::Dataflow::RegisterTransformSamplerNodes();
 
+	UE::Dataflow::RegisterSamplerNodes();
 	UE::Dataflow::DataflowCollectionAttributeKeyNodes();
 	DATAFLOW_NODE_REGISTER_CREATION_FACTORY(FDataflowCollectionAddScalarVertexPropertyNode);
 	DATAFLOW_NODE_REGISTER_CREATION_FACTORY(FDataflowCollectionEditSkinWeightsNode);
@@ -83,8 +86,6 @@ void IDataflowNodesPlugin::ShutdownModule()
 	FDataflowAddScalarVertexPropertyCallbackRegistry::Get().DeregisterCallbacks(FGeometryCollectionAddScalarVertexPropertyCallbacks::Name);
 }
 
-
 IMPLEMENT_MODULE(IDataflowNodesPlugin, DataflowNodes)
 
-
 #undef LOCTEXT_NAMESPACE
diff --git a/Engine/Plugins/Experimental/Dataflow/Source/DataflowNodes/Private/Dataflow/SamplerNodes/DataflowAbsSamplerNode.cpp b/Engine/Plugins/Experimental/Dataflow/Source/DataflowNodes/Private/Dataflow/SamplerNodes/DataflowAbsSamplerNode.cpp
new file mode 100644
--- /dev/null
+++ b/Engine/Plugins/Experimental/Dataflow/Source/DataflowNodes/Private/Dataflow/SamplerNodes/DataflowAbsSamplerNode.cpp
@@ -0,0 +1,111 @@
+// Copyright Epic Games, Inc. All Rights Reserved.
+
+#include "Dataflow/SamplerNodes/DataflowAbsSamplerNode.h"
+#include "Dataflow/SamplerNodes/DataflowSamplerFunctions.h"
+#include "Dataflow/DataflowNodeFactory.h"
+
+FBox FDataflowAbsFloatSampler::GetRenderBounds() const
+{
+	return FBox(FVector(-50.0), FVector(50.0));
+}
+
+void FDataflowAbsFloatSampler::Sample(TArrayView<const FVector3f> Positions, TArrayView<float> OutValues) const
+{
+	if (Sampler.IsValid() && Positions.Num() == OutValues.Num())
+	{
+		TArray<float> OutValueFromInputSampler;
+		OutValueFromInputSampler.SetNumUninitialized(OutValues.Num());
+	
+		Sampler->Sample(Positions, OutValueFromInputSampler);
+	
+		for (int32 Idx = 0; Idx < OutValues.Num(); ++Idx)
+		{
+			OutValues[Idx] = FMath::Abs(OutValueFromInputSampler[Idx]);
+		}
+	}
+	else
+	{
+		for (int32 Idx = 0; Idx < OutValues.Num(); ++Idx)
+		{
+			OutValues[Idx] = 0.f;
+		}
+	}
+}
+
+FBox FDataflowAbsVectorSampler::GetRenderBounds() const
+{
+	return FBox(FVector(-50.0), FVector(50.0));
+}
+
+void FDataflowAbsVectorSampler::Sample(TArrayView<const FVector3f> Positions, TArrayView<FVector3f> OutValues) const
+{
+	if (Sampler.IsValid() && Positions.Num() == OutValues.Num())
+	{
+		TArray<FVector3f> OutValueFromInputSampler;
+		OutValueFromInputSampler.SetNumUninitialized(OutValues.Num());
+	
+		Sampler->Sample(Positions, OutValueFromInputSampler);
+	
+		for (int32 Idx = 0; Idx < OutValues.Num(); ++Idx)
+		{
+			OutValues[Idx] = FVector3f(FMath::Abs(OutValueFromInputSampler[Idx].X), 
+				FMath::Abs(OutValueFromInputSampler[Idx].Y), 
+				FMath::Abs(OutValueFromInputSampler[Idx].Z));
+		}
+	}
+	else
+	{
+		for (int32 Idx = 0; Idx < OutValues.Num(); ++Idx)
+		{
+			OutValues[Idx] = FVector3f::ZeroVector;
+		}
+	}
+}
+
+FDataflowAbsSamplerNode::FDataflowAbsSamplerNode(const UE::Dataflow::FNodeParameters& InParam, FGuid InGuid)
+	: FDataflowNode(InParam, InGuid)
+{
+	static const FName MainTypeGroup("Main");
+	
+	RegisterInputConnection(&Sampler)
+		.SetTypeDependencyGroup(MainTypeGroup);
+	RegisterOutputConnection(&Sampler)
+		.SetPassthroughInput(&Sampler)
+		.SetTypeDependencyGroup(MainTypeGroup);
+}
+
+void FDataflowAbsSamplerNode::Evaluate(UE::Dataflow::FContext& Context, const FDataflowOutput* Out) const
+{
+	if (Out->IsA(&Sampler))
+	{
+		const FDataflowFloatSampler* FloatSampler = GetValue(Context, &Sampler).TryGet <FDataflowFloatSampler>();
+		if (FloatSampler)
+		{
+			TSharedRef<FDataflowAbsFloatSampler> Impl = MakeShared<FDataflowAbsFloatSampler>();
+
+			Impl->Sampler = FloatSampler->GetImpl();
+
+			FDataflowFloatSampler OutSampler(Impl.ToSharedPtr());
+
+			SetValue(Context, OutSampler, &Sampler);
+			return;
+		}
+
+		const FDataflowVectorSampler* VectorSampler = GetValue(Context, &Sampler).TryGet<FDataflowVectorSampler>();
+		if (VectorSampler)
+		{
+			TSharedRef<FDataflowAbsVectorSampler> Impl = MakeShared<FDataflowAbsVectorSampler>();
+		
+			Impl->Sampler = VectorSampler->GetImpl();
+		
+			FDataflowVectorSampler OutSampler(Impl.ToSharedPtr());
+
+			SetValue(Context, OutSampler, &Sampler);
+			return;
+		}
+
+		SafeForwardInput(Context, &Sampler, &Sampler);
+	}
+}
+
+
diff --git a/Engine/Plugins/Experimental/Dataflow/Source/DataflowNodes/Private/Dataflow/SamplerNodes/DataflowAbsSamplerNode.h b/Engine/Plugins/Experimental/Dataflow/Source/DataflowNodes/Private/Dataflow/SamplerNodes/DataflowAbsSamplerNode.h
new file mode 100644
--- /dev/null
+++ b/Engine/Plugins/Experimental/Dataflow/Source/DataflowNodes/Private/Dataflow/SamplerNodes/DataflowAbsSamplerNode.h
@@ -0,0 +1,58 @@
+// Copyright Epic Games, Inc. All Rights Reserved.
+
+#pragma once 
+
+#include "Dataflow/DataflowNode.h"
+#include "Dataflow/SamplerNodes/DataflowSamplerTypes.h"
+
+#include "DataflowAbsSamplerNode.generated.h"
+
+/***********************************************************************************************************************
+ * No Input Sampler nodes
+ ***********************************************************************************************************************/
+USTRUCT()
+struct FDataflowAbsFloatSampler : public FDataflowFloatSamplerBase
+{
+	GENERATED_BODY();
+
+	virtual ~FDataflowAbsFloatSampler() override = default;
+	virtual void Sample(TArrayView<const FVector3f> Positions, TArrayView<float> OutValues) const override;
+	virtual FBox GetRenderBounds() const override;
+
+public:
+	/**  */
+	TSharedPtr<const FDataflowFloatSamplerBase> Sampler;
+};
+
+USTRUCT()
+struct FDataflowAbsVectorSampler : public FDataflowVectorSamplerBase
+{
+	GENERATED_BODY();
+
+	virtual ~FDataflowAbsVectorSampler() override = default;
+	virtual void Sample(TArrayView<const FVector3f> Positions, TArrayView<FVector3f> OutValues) const override;
+	virtual FBox GetRenderBounds() const override;
+
+public:
+	/**  */
+	TSharedPtr<const FDataflowVectorSamplerBase> Sampler;
+};
+
+/**
+ * Output an abs float sampler
+ */
+USTRUCT()
+struct FDataflowAbsSamplerNode : public FDataflowNode
+{
+	GENERATED_BODY()
+	DATAFLOW_NODE_DEFINE_INTERNAL(FDataflowAbsSamplerNode, "Abs Sampler", "Samplers", "Absolute")
+
+public:
+	FDataflowAbsSamplerNode(const UE::Dataflow::FNodeParameters& InParam, FGuid InGuid = FGuid::NewGuid());
+
+private:
+	UPROPERTY(meta = (DataflowInput, DataflowOutput, DataflowPassthrough = Sampler))
+	FDataflowSamplerTypes Sampler;
+	
+	void Evaluate(UE::Dataflow::FContext& Context, const FDataflowOutput* Out) const override;
+};
diff --git a/Engine/Plugins/Experimental/Dataflow/Source/DataflowNodes/Private/Dataflow/SamplerNodes/DataflowAddSamplerNode.cpp b/Engine/Plugins/Experimental/Dataflow/Source/DataflowNodes/Private/Dataflow/SamplerNodes/DataflowAddSamplerNode.cpp
new file mode 100644
--- /dev/null
+++ b/Engine/Plugins/Experimental/Dataflow/Source/DataflowNodes/Private/Dataflow/SamplerNodes/DataflowAddSamplerNode.cpp
@@ -0,0 +1,117 @@
+// Copyright Epic Games, Inc. All Rights Reserved.
+
+#include "Dataflow/SamplerNodes/DataflowAddSamplerNode.h"
+#include "Dataflow/DataflowNodeFactory.h"
+
+#include UE_INLINE_GENERATED_CPP_BY_NAME(DataflowAddSamplerNode)
+
+#define LOCTEXT_NAMESPACE "DataflowAddSamplerNode"
+
+FBox FDataflowAddFloatSampler::GetRenderBounds() const
+{
+	return FBox(FVector(-50.0), FVector(50.0));
+}
+
+void FDataflowAddFloatSampler::Sample(TArrayView<const FVector3f> Positions, TArrayView<float> OutValues) const
+{
+	for (int32 Index = 0; Index < OutValues.Num(); ++Index)
+	{
+		OutValues[Index] = 0;
+	}
+
+	if (Positions.Num() == OutValues.Num())
+	{
+		TArray<float> OneSamplerValues;
+
+		for (const TSharedPtr<const FDataflowFloatSamplerBase>& Sampler : Samplers)
+		{
+			OneSamplerValues.Init(0, OutValues.Num());
+			Sampler->Sample(Positions, OneSamplerValues);
+			for (int32 Index = 0; Index < OutValues.Num(); ++Index)
+			{
+				OutValues[Index] += OneSamplerValues[Index];
+			}
+		}
+	}
+}
+
+FBox FDataflowAddVectorSampler::GetRenderBounds() const
+{
+	return FBox(FVector(-50.0), FVector(50.0));
+}
+
+void FDataflowAddVectorSampler::Sample(TArrayView<const FVector3f> Positions, TArrayView<FVector3f> OutValues) const
+{
+	for (int32 Index = 0; Index < OutValues.Num(); ++Index)
+	{
+		OutValues[Index] = FVector3f::ZeroVector;
+	}
+
+	if (Positions.Num() == OutValues.Num())
+	{
+		TArray<FVector3f> OneSamplerValues;
+
+		for (const TSharedPtr<const FDataflowVectorSamplerBase>& Sampler : Samplers)
+		{
+			OneSamplerValues.Init(FVector3f::ZeroVector, OutValues.Num());
+			Sampler->Sample(Positions, OneSamplerValues);
+			for (int32 Index = 0; Index < OutValues.Num(); ++Index)
+			{
+				OutValues[Index] += OneSamplerValues[Index];
+			}
+		}
+	}
+}
+
+//////////////////////////////////////////////////////////////////////////////////////////////////////////////////////
+
+FDataflowAddSamplerNode::FDataflowAddSamplerNode(const UE::Dataflow::FNodeParameters& InParam, FGuid InGuid)
+	: FDataflowMultiInputSamplerNodeBase(InParam, InGuid)
+{
+	RegisterInitialConnections();
+}
+
+void FDataflowAddSamplerNode::EvaluateSampler(UE::Dataflow::FContext& Context, FDataflowSamplerTypes::FStorageType& OutSampler) const
+{
+	// use the first sampler as a reference
+	const int32 NumSamplers = GetNumSamplerInputs();
+	if (NumSamplers == 0)
+	{
+		return;
+	}
+
+	const FDataflowSamplerTypes::FStorageType& FirstSampler = GetSamplerInput(Context, 0);
+	if (FirstSampler.TryGet<FDataflowFloatSampler>())
+	{
+		TSharedPtr<FDataflowAddFloatSampler> OutImpl = MakeShared<FDataflowAddFloatSampler>();
+		for (int32 Index = 0; Index < NumSamplers; ++Index)
+		{
+			if (const FDataflowFloatSampler* FloatSampler = GetSamplerInput(Context, Index).TryGet<FDataflowFloatSampler>())
+			{
+				if (TSharedPtr<const FDataflowFloatSamplerBase> SamplerImpl = FloatSampler->GetImpl())
+				{
+					OutImpl->Samplers.Add(SamplerImpl);
+				}
+			}
+		}
+		OutSampler = FDataflowFloatSampler(OutImpl);
+	}
+	else if (FirstSampler.TryGet<FDataflowVectorSampler>())
+	{
+		TSharedPtr<FDataflowAddVectorSampler> OutImpl = MakeShared<FDataflowAddVectorSampler>();
+		for (int32 Index = 0; Index < NumSamplers; ++Index)
+		{
+			if (const FDataflowVectorSampler* VectorSampler = GetSamplerInput(Context, Index).TryGet<FDataflowVectorSampler>())
+			{
+				if (TSharedPtr<const FDataflowVectorSamplerBase> SamplerImpl = VectorSampler->GetImpl())
+				{
+					OutImpl->Samplers.Add(SamplerImpl);
+				}
+			}
+		}
+		OutSampler = FDataflowVectorSampler(OutImpl);
+	}
+
+}
+
+#undef LOCTEXT_NAMESPACE 
diff --git a/Engine/Plugins/Experimental/Dataflow/Source/DataflowNodes/Private/Dataflow/SamplerNodes/DataflowAddSamplerNode.h b/Engine/Plugins/Experimental/Dataflow/Source/DataflowNodes/Private/Dataflow/SamplerNodes/DataflowAddSamplerNode.h
new file mode 100644
--- /dev/null
+++ b/Engine/Plugins/Experimental/Dataflow/Source/DataflowNodes/Private/Dataflow/SamplerNodes/DataflowAddSamplerNode.h
@@ -0,0 +1,49 @@
+// Copyright Epic Games, Inc. All Rights Reserved.
+
+#pragma once 
+#include "Dataflow/SamplerNodes/DataflowSamplerTypes.h"
+#include "Dataflow/SamplerNodes/DataflowSamplerMultiInput.h"
+
+#include "DataflowAddSamplerNode.generated.h"
+
+USTRUCT()
+struct FDataflowAddFloatSampler : public FDataflowFloatSamplerBase
+{
+	GENERATED_BODY();
+	virtual ~FDataflowAddFloatSampler() override = default;
+	virtual void Sample(TArrayView<const FVector3f> Positions, TArrayView<float> OutValues) const override;
+	virtual FBox GetRenderBounds() const override;
+
+public:
+	/** list of samplers to add results from */
+	TArray<TSharedPtr<const FDataflowFloatSamplerBase>> Samplers;
+};
+
+USTRUCT()
+struct FDataflowAddVectorSampler : public FDataflowVectorSamplerBase
+{
+	GENERATED_BODY();
+	virtual ~FDataflowAddVectorSampler() override = default;
+	virtual void Sample(TArrayView<const FVector3f> Positions, TArrayView<FVector3f> OutValues) const override;
+	virtual FBox GetRenderBounds() const override;
+
+public:
+	/** list of samplers to add results from */
+	TArray<TSharedPtr<const FDataflowVectorSamplerBase>> Samplers;
+};
+
+/**
+* add sampler results together
+*/
+USTRUCT()
+struct FDataflowAddSamplerNode : public FDataflowMultiInputSamplerNodeBase
+{
+	GENERATED_USTRUCT_BODY()
+
+	DATAFLOW_NODE_DEFINE_INTERNAL(FDataflowAddSamplerNode, "Add Sampler", "Samplers", "+ Addition Sum")
+
+public:
+	FDataflowAddSamplerNode(const UE::Dataflow::FNodeParameters& InParam, FGuid InGuid = FGuid::NewGuid());
+	virtual void EvaluateSampler(UE::Dataflow::FContext& Context, FDataflowSamplerTypes::FStorageType& OutSampler) const override;
+};
+
diff --git a/Engine/Plugins/Experimental/Dataflow/Source/DataflowNodes/Private/Dataflow/SamplerNodes/DataflowGradientSamplerNode.cpp b/Engine/Plugins/Experimental/Dataflow/Source/DataflowNodes/Private/Dataflow/SamplerNodes/DataflowGradientSamplerNode.cpp
--- a/Engine/Plugins/Experimental/Dataflow/Source/DataflowNodes/Private/Dataflow/SamplerNodes/DataflowGradientSamplerNode.cpp
+++ b/Engine/Plugins/Experimental/Dataflow/Source/DataflowNodes/Private/Dataflow/SamplerNodes/DataflowGradientSamplerNode.cpp
@@ -15,6 +15,16 @@ namespace UE::Dataflow
 
 ////////////////////////////////////////////////////////////////////////////////////////////////////////////////////////
 
+FBox FDataflowLinearGradientFloatSampler::GetRenderBounds() const
+{
+	FBox Box;
+
+	Box += FVector(StartPoint);
+	Box += FVector(EndPoint);
+
+	return Box;
+}
+
 void FDataflowLinearGradientFloatSampler::Sample(TArrayView<const FVector3f> Positions, TArrayView<float> OutValues) const
 {
 	const FVector3f GradientDirection = EndPoint - StartPoint;
@@ -72,7 +82,7 @@ void FDataflowLinearGradientFloatSamplerNode::Evaluate(UE::Dataflow::FContext& C
 
 ////////////////////////////////////////////////////////////////////////////////////////////////////////////////////////
 
-void FDataflowLinearGradientFromBoxFloatSampler::Sample(TArrayView<const FVector3f> Positions, TArrayView<float> OutValues) const
+void FDataflowLinearGradientFromBoxFloatSampler::ComputePoints(FVector& StartPoint, FVector& EndPoint) const
 {
 	const FVector BoxCenter = Box.GetCenter();
 	const FVector BoxHalfExtent = Box.GetExtent();
@@ -88,8 +98,31 @@ void FDataflowLinearGradientFromBoxFloatSampler::Sample(TArrayView<const FVector
 
 	GradientHalfExtent = Transform.TransformVector(GradientHalfExtent);
 
-	const FVector StartPoint = Transform.GetTranslation() + BoxCenter - GradientHalfExtent;
-	const FVector EndPoint = Transform.GetTranslation() + BoxCenter + GradientHalfExtent;
+	StartPoint = Transform.GetTranslation() + BoxCenter - GradientHalfExtent;
+	EndPoint = Transform.GetTranslation() + BoxCenter + GradientHalfExtent;
+}
+
+FBox FDataflowLinearGradientFromBoxFloatSampler::GetRenderBounds() const
+{
+	FVector StartPoint;
+	FVector EndPoint;
+
+	ComputePoints(StartPoint, EndPoint);
+
+	FBox BoundBox;
+
+	BoundBox += FVector(StartPoint);
+	BoundBox += FVector(EndPoint);
+
+	return BoundBox;
+}
+
+void FDataflowLinearGradientFromBoxFloatSampler::Sample(TArrayView<const FVector3f> Positions, TArrayView<float> OutValues) const
+{
+	FVector StartPoint;
+	FVector EndPoint;
+
+	ComputePoints(StartPoint, EndPoint);
 
 	FDataflowLinearGradientFloatSampler RegularGradientSampler;
 	RegularGradientSampler.StartPoint = FVector3f(StartPoint);
@@ -132,6 +165,15 @@ void FDataflowLinearGradientFromBoxFloatSamplerNode::Evaluate(UE::Dataflow::FCon
 
 ////////////////////////////////////////////////////////////////////////////////////////////////////////////////////////
 
+FBox FDataflowRadialGradientFloatSampler::GetRenderBounds() const
+{
+	// Radius is always greater than zero
+	FVector Extent(Radius);
+	FBox Box(FVector(Center) - Extent, FVector(Center) + Extent);
+
+	return Box;
+}
+
 void FDataflowRadialGradientFloatSampler::Sample(TArrayView<const FVector3f> Positions, TArrayView<float> OutValues) const
 {
 	if (Radius < UE_SMALL_NUMBER)
diff --git a/Engine/Plugins/Experimental/Dataflow/Source/DataflowNodes/Private/Dataflow/SamplerNodes/DataflowGradientSamplerNode.h b/Engine/Plugins/Experimental/Dataflow/Source/DataflowNodes/Private/Dataflow/SamplerNodes/DataflowGradientSamplerNode.h
--- a/Engine/Plugins/Experimental/Dataflow/Source/DataflowNodes/Private/Dataflow/SamplerNodes/DataflowGradientSamplerNode.h
+++ b/Engine/Plugins/Experimental/Dataflow/Source/DataflowNodes/Private/Dataflow/SamplerNodes/DataflowGradientSamplerNode.h
@@ -16,6 +16,7 @@ struct FDataflowLinearGradientFloatSampler : public FDataflowFloatSamplerBase
 
 	virtual ~FDataflowLinearGradientFloatSampler() override = default;
 	virtual void Sample(TArrayView<const FVector3f> Positions, TArrayView<float> OutValues) const override;
+	virtual FBox GetRenderBounds() const override;
 
 public:
 	// Point in space where the start value is defined
@@ -55,7 +56,7 @@ private:
 	UPROPERTY(meta = (DataflowOutput))
 	FDataflowFloatSampler Sampler;
 
-	UPROPERTY(EditAnywhere, Category = Sampler, meta = (ShowOnlyInnerProperties))
+	UPROPERTY(EditAnywhere, Category = Sampler, meta = (ShowOnlyInnerProperties, SkipInDisplayNameChain))
 	FDataflowLinearGradientFloatSampler Gradient;
 
 	void Evaluate(UE::Dataflow::FContext& Context, const FDataflowOutput* Out) const override;
@@ -70,6 +71,7 @@ struct FDataflowLinearGradientFromBoxFloatSampler : public FDataflowFloatSampler
 
 	virtual ~FDataflowLinearGradientFromBoxFloatSampler() override = default;
 	virtual void Sample(TArrayView<const FVector3f> Positions, TArrayView<float> OutValues) const override;
+	virtual FBox GetRenderBounds() const override;
 
 public:
 	// Box to use for the gradient - one axis is chosen to be the axis of the gradient
@@ -95,6 +97,9 @@ public:
 	// If true generated value will be clamped beyond end and start point, otherwise, they will be extrapolated
 	UPROPERTY(EditAnywhere, Category = Gradient, meta = (DataflowInput))
 	bool bClamp = true;
+
+private:
+	void ComputePoints(FVector& StartPoint, FVector& EndPoint) const;
 };
 
 /**
@@ -113,7 +118,7 @@ private:
 	UPROPERTY(meta = (DataflowOutput))
 	FDataflowFloatSampler Sampler;
 
-	UPROPERTY(EditAnywhere, Category = Sampler, meta = (ShowOnlyInnerProperties))
+	UPROPERTY(EditAnywhere, Category = Sampler, meta = (ShowOnlyInnerProperties, SkipInDisplayNameChain))
 	FDataflowLinearGradientFromBoxFloatSampler Gradient;
 	
 	void Evaluate(UE::Dataflow::FContext& Context, const FDataflowOutput* Out) const override;
@@ -128,6 +133,7 @@ struct FDataflowRadialGradientFloatSampler : public FDataflowFloatSamplerBase
 
 	virtual ~FDataflowRadialGradientFloatSampler() override = default;
 	virtual void Sample(TArrayView<const FVector3f> Positions, TArrayView<float> OutValues) const override;
+	virtual FBox GetRenderBounds() const override;
 
 public:
 	// center defining the radial gradient
@@ -135,7 +141,7 @@ public:
 	FVector3f Center = { 0, 0, 0 };
 
 	// Radius of the gradient
-	UPROPERTY(EditAnywhere, Category = Gradient, meta = (DataflowInput))
+	UPROPERTY(EditAnywhere, Category = Gradient, meta = (DataflowInput, UIMin = 0.01, ClampMin = 0.01))
 	float Radius = 100.f;
 
 	// Value defined at center
@@ -167,7 +173,7 @@ private:
 	UPROPERTY(meta = (DataflowOutput))
 	FDataflowFloatSampler Sampler;
 
-	UPROPERTY(EditAnywhere, Category = Sampler, meta = (ShowOnlyInnerProperties))
+	UPROPERTY(EditAnywhere, Category = Sampler, meta = (ShowOnlyInnerProperties, SkipInDisplayNameChain))
 	FDataflowRadialGradientFloatSampler Gradient;;
 
 	void Evaluate(UE::Dataflow::FContext& Context, const FDataflowOutput* Out) const override;
diff --git a/Engine/Plugins/Experimental/Dataflow/Source/DataflowNodes/Private/Dataflow/SamplerNodes/DataflowMaxSamplerNode.cpp b/Engine/Plugins/Experimental/Dataflow/Source/DataflowNodes/Private/Dataflow/SamplerNodes/DataflowMaxSamplerNode.cpp
new file mode 100644
--- /dev/null
+++ b/Engine/Plugins/Experimental/Dataflow/Source/DataflowNodes/Private/Dataflow/SamplerNodes/DataflowMaxSamplerNode.cpp
@@ -0,0 +1,119 @@
+// Copyright Epic Games, Inc. All Rights Reserved.
+
+#include "Dataflow/SamplerNodes/DataflowMaxSamplerNode.h"
+#include "Dataflow/DataflowNodeFactory.h"
+
+#include UE_INLINE_GENERATED_CPP_BY_NAME(DataflowMaxSamplerNode)
+
+#define LOCTEXT_NAMESPACE "DataflowMaxSamplerNode"
+
+FBox FDataflowMaxFloatSampler::GetRenderBounds() const
+{
+	return FBox(FVector(-50.0), FVector(50.0));
+}
+
+void FDataflowMaxFloatSampler::Sample(TArrayView<const FVector3f> Positions, TArrayView<float> OutValues) const
+{
+	for (int32 Index = 0; Index < OutValues.Num(); ++Index)
+	{
+		OutValues[Index] = -FLT_MAX;
+	}
+
+	if (Positions.Num() == OutValues.Num())
+	{
+		TArray<float> OneSamplerValues;
+
+		for (const TSharedPtr<const FDataflowFloatSamplerBase>& Sampler : Samplers)
+		{
+			OneSamplerValues.Init(0, OutValues.Num());
+			Sampler->Sample(Positions, OneSamplerValues);
+			for (int32 Index = 0; Index < OutValues.Num(); ++Index)
+			{
+				OutValues[Index] = FMath::Max(OutValues[Index], OneSamplerValues[Index]);
+			}
+		}
+	}
+}
+
+FBox FDataflowMaxVectorSampler::GetRenderBounds() const
+{
+	return FBox(FVector(-50.0), FVector(50.0));
+}
+
+void FDataflowMaxVectorSampler::Sample(TArrayView<const FVector3f> Positions, TArrayView<FVector3f> OutValues) const
+{
+	for (int32 Index = 0; Index < OutValues.Num(); ++Index)
+	{
+		OutValues[Index] = FVector3f(-FLT_MAX);
+	}
+
+	if (Positions.Num() == OutValues.Num())
+	{
+		TArray<FVector3f> OneSamplerValues;
+
+		for (const TSharedPtr<const FDataflowVectorSamplerBase>& Sampler : Samplers)
+		{
+			OneSamplerValues.Init(FVector3f::ZeroVector, OutValues.Num());
+			Sampler->Sample(Positions, OneSamplerValues);
+			for (int32 Index = 0; Index < OutValues.Num(); ++Index)
+			{
+				OutValues[Index].X = FMath::Max(OutValues[Index].X, OneSamplerValues[Index].X);
+				OutValues[Index].Y = FMath::Max(OutValues[Index].Y, OneSamplerValues[Index].Y);
+				OutValues[Index].Z = FMath::Max(OutValues[Index].Z, OneSamplerValues[Index].Z);
+			}
+		}
+	}
+}
+
+//////////////////////////////////////////////////////////////////////////////////////////////////////////////////////
+
+FDataflowMaxSamplerNode::FDataflowMaxSamplerNode(const UE::Dataflow::FNodeParameters& InParam, FGuid InGuid)
+	: FDataflowMultiInputSamplerNodeBase(InParam, InGuid)
+{
+	RegisterInitialConnections();
+}
+
+void FDataflowMaxSamplerNode::EvaluateSampler(UE::Dataflow::FContext& Context, FDataflowSamplerTypes::FStorageType& OutSampler) const
+{
+	// use the first sampler as a reference
+	const int32 NumSamplers = GetNumSamplerInputs();
+	if (NumSamplers == 0)
+	{
+		return;
+	}
+
+	const FDataflowSamplerTypes::FStorageType& FirstSampler = GetSamplerInput(Context, 0);
+	if (FirstSampler.TryGet<FDataflowFloatSampler>())
+	{
+		TSharedPtr<FDataflowMaxFloatSampler> OutImpl = MakeShared<FDataflowMaxFloatSampler>();
+		for (int32 Index = 0; Index < NumSamplers; ++Index)
+		{
+			if (const FDataflowFloatSampler* FloatSampler = GetSamplerInput(Context, Index).TryGet<FDataflowFloatSampler>())
+			{
+				if (TSharedPtr<const FDataflowFloatSamplerBase> SamplerImpl = FloatSampler->GetImpl())
+				{
+					OutImpl->Samplers.Add(SamplerImpl);
+				}
+			}
+		}
+		OutSampler = FDataflowFloatSampler(OutImpl);
+	}
+	else if (FirstSampler.TryGet<FDataflowVectorSampler>())
+	{
+		TSharedPtr<FDataflowMaxVectorSampler> OutImpl = MakeShared<FDataflowMaxVectorSampler>();
+		for (int32 Index = 0; Index < NumSamplers; ++Index)
+		{
+			if (const FDataflowVectorSampler* VectorSampler = GetSamplerInput(Context, Index).TryGet<FDataflowVectorSampler>())
+			{
+				if (TSharedPtr<const FDataflowVectorSamplerBase> SamplerImpl = VectorSampler->GetImpl())
+				{
+					OutImpl->Samplers.Add(SamplerImpl);
+				}
+			}
+		}
+		OutSampler = FDataflowVectorSampler(OutImpl);
+	}
+
+}
+
+#undef LOCTEXT_NAMESPACE 
diff --git a/Engine/Plugins/Experimental/Dataflow/Source/DataflowNodes/Private/Dataflow/SamplerNodes/DataflowMaxSamplerNode.h b/Engine/Plugins/Experimental/Dataflow/Source/DataflowNodes/Private/Dataflow/SamplerNodes/DataflowMaxSamplerNode.h
new file mode 100644
--- /dev/null
+++ b/Engine/Plugins/Experimental/Dataflow/Source/DataflowNodes/Private/Dataflow/SamplerNodes/DataflowMaxSamplerNode.h
@@ -0,0 +1,50 @@
+// Copyright Epic Games, Inc. All Rights Reserved.
+
+#pragma once 
+
+#include "Dataflow/SamplerNodes/DataflowSamplerTypes.h"
+#include "Dataflow/SamplerNodes/DataflowSamplerMultiInput.h"
+
+#include "DataflowMaxSamplerNode.generated.h"
+
+USTRUCT()
+struct FDataflowMaxFloatSampler : public FDataflowFloatSamplerBase
+{
+	GENERATED_BODY();
+	virtual ~FDataflowMaxFloatSampler() override = default;
+	virtual void Sample(TArrayView<const FVector3f> Positions, TArrayView<float> OutValues) const override;
+	virtual FBox GetRenderBounds() const override;
+
+public:
+	/** list of samplers to Max results from */
+	TArray<TSharedPtr<const FDataflowFloatSamplerBase>> Samplers;
+};
+
+USTRUCT()
+struct FDataflowMaxVectorSampler : public FDataflowVectorSamplerBase
+{
+	GENERATED_BODY();
+	virtual ~FDataflowMaxVectorSampler() override = default;
+	virtual void Sample(TArrayView<const FVector3f> Positions, TArrayView<FVector3f> OutValues) const override;
+	virtual FBox GetRenderBounds() const override;
+
+public:
+	/** list of samplers to Max results from */
+	TArray<TSharedPtr<const FDataflowVectorSamplerBase>> Samplers;
+};
+
+/**
+* Max sampler results together
+*/
+USTRUCT()
+struct FDataflowMaxSamplerNode : public FDataflowMultiInputSamplerNodeBase
+{
+	GENERATED_USTRUCT_BODY()
+
+	DATAFLOW_NODE_DEFINE_INTERNAL(FDataflowMaxSamplerNode, "Max Sampler", "Samplers", "Maximum")
+
+public:
+	FDataflowMaxSamplerNode(const UE::Dataflow::FNodeParameters& InParam, FGuid InGuid = FGuid::NewGuid());
+	virtual void EvaluateSampler(UE::Dataflow::FContext& Context, FDataflowSamplerTypes::FStorageType& OutSampler) const override;
+};
+
diff --git a/Engine/Plugins/Experimental/Dataflow/Source/DataflowNodes/Private/Dataflow/SamplerNodes/DataflowMinSamplerNode.cpp b/Engine/Plugins/Experimental/Dataflow/Source/DataflowNodes/Private/Dataflow/SamplerNodes/DataflowMinSamplerNode.cpp
new file mode 100644
--- /dev/null
+++ b/Engine/Plugins/Experimental/Dataflow/Source/DataflowNodes/Private/Dataflow/SamplerNodes/DataflowMinSamplerNode.cpp
@@ -0,0 +1,119 @@
+// Copyright Epic Games, Inc. All Rights Reserved.
+
+#include "Dataflow/SamplerNodes/DataflowMinSamplerNode.h"
+#include "Dataflow/DataflowNodeFactory.h"
+
+#include UE_INLINE_GENERATED_CPP_BY_NAME(DataflowMinSamplerNode)
+
+#define LOCTEXT_NAMESPACE "DataflowMinSamplerNode"
+
+FBox FDataflowMinFloatSampler::GetRenderBounds() const
+{
+	return FBox(FVector(-50.0), FVector(50.0));
+}
+
+void FDataflowMinFloatSampler::Sample(TArrayView<const FVector3f> Positions, TArrayView<float> OutValues) const
+{
+	for (int32 Index = 0; Index < OutValues.Num(); ++Index)
+	{
+		OutValues[Index] = FLT_MAX;
+	}
+
+	if (Positions.Num() == OutValues.Num())
+	{
+		TArray<float> OneSamplerValues;
+
+		for (const TSharedPtr<const FDataflowFloatSamplerBase>& Sampler : Samplers)
+		{
+			OneSamplerValues.Init(0, OutValues.Num());
+			Sampler->Sample(Positions, OneSamplerValues);
+			for (int32 Index = 0; Index < OutValues.Num(); ++Index)
+			{
+				OutValues[Index] = FMath::Min(OutValues[Index], OneSamplerValues[Index]);
+			}
+		}
+	}
+}
+
+FBox FDataflowMinVectorSampler::GetRenderBounds() const
+{
+	return FBox(FVector(-50.0), FVector(50.0));
+}
+
+void FDataflowMinVectorSampler::Sample(TArrayView<const FVector3f> Positions, TArrayView<FVector3f> OutValues) const
+{
+	for (int32 Index = 0; Index < OutValues.Num(); ++Index)
+	{
+		OutValues[Index] = FVector3f(FLT_MAX);
+	}
+
+	if (Positions.Num() == OutValues.Num())
+	{
+		TArray<FVector3f> OneSamplerValues;
+
+		for (const TSharedPtr<const FDataflowVectorSamplerBase>& Sampler : Samplers)
+		{
+			OneSamplerValues.Init(FVector3f::ZeroVector, OutValues.Num());
+			Sampler->Sample(Positions, OneSamplerValues);
+			for (int32 Index = 0; Index < OutValues.Num(); ++Index)
+			{
+				OutValues[Index].X = FMath::Min(OutValues[Index].X, OneSamplerValues[Index].X);
+				OutValues[Index].Y = FMath::Min(OutValues[Index].Y, OneSamplerValues[Index].Y);
+				OutValues[Index].Z = FMath::Min(OutValues[Index].Z, OneSamplerValues[Index].Z);
+			}
+		}
+	}
+}
+
+//////////////////////////////////////////////////////////////////////////////////////////////////////////////////////
+
+FDataflowMinSamplerNode::FDataflowMinSamplerNode(const UE::Dataflow::FNodeParameters& InParam, FGuid InGuid)
+	: FDataflowMultiInputSamplerNodeBase(InParam, InGuid)
+{
+	RegisterInitialConnections();
+}
+
+void FDataflowMinSamplerNode::EvaluateSampler(UE::Dataflow::FContext& Context, FDataflowSamplerTypes::FStorageType& OutSampler) const
+{
+	// use the first sampler as a reference
+	const int32 NumSamplers = GetNumSamplerInputs();
+	if (NumSamplers == 0)
+	{
+		return;
+	}
+
+	const FDataflowSamplerTypes::FStorageType& FirstSampler = GetSamplerInput(Context, 0);
+	if (FirstSampler.TryGet<FDataflowFloatSampler>())
+	{
+		TSharedPtr<FDataflowMinFloatSampler> OutImpl = MakeShared<FDataflowMinFloatSampler>();
+		for (int32 Index = 0; Index < NumSamplers; ++Index)
+		{
+			if (const FDataflowFloatSampler* FloatSampler = GetSamplerInput(Context, Index).TryGet<FDataflowFloatSampler>())
+			{
+				if (TSharedPtr<const FDataflowFloatSamplerBase> SamplerImpl = FloatSampler->GetImpl())
+				{
+					OutImpl->Samplers.Add(SamplerImpl);
+				}
+			}
+		}
+		OutSampler = FDataflowFloatSampler(OutImpl);
+	}
+	else if (FirstSampler.TryGet<FDataflowVectorSampler>())
+	{
+		TSharedPtr<FDataflowMinVectorSampler> OutImpl = MakeShared<FDataflowMinVectorSampler>();
+		for (int32 Index = 0; Index < NumSamplers; ++Index)
+		{
+			if (const FDataflowVectorSampler* VectorSampler = GetSamplerInput(Context, Index).TryGet<FDataflowVectorSampler>())
+			{
+				if (TSharedPtr<const FDataflowVectorSamplerBase> SamplerImpl = VectorSampler->GetImpl())
+				{
+					OutImpl->Samplers.Add(SamplerImpl);
+				}
+			}
+		}
+		OutSampler = FDataflowVectorSampler(OutImpl);
+	}
+
+}
+
+#undef LOCTEXT_NAMESPACE 
diff --git a/Engine/Plugins/Experimental/Dataflow/Source/DataflowNodes/Private/Dataflow/SamplerNodes/DataflowMinSamplerNode.h b/Engine/Plugins/Experimental/Dataflow/Source/DataflowNodes/Private/Dataflow/SamplerNodes/DataflowMinSamplerNode.h
new file mode 100644
--- /dev/null
+++ b/Engine/Plugins/Experimental/Dataflow/Source/DataflowNodes/Private/Dataflow/SamplerNodes/DataflowMinSamplerNode.h
@@ -0,0 +1,50 @@
+// Copyright Epic Games, Inc. All Rights Reserved.
+
+#pragma once 
+
+#include "Dataflow/SamplerNodes/DataflowSamplerTypes.h"
+#include "Dataflow/SamplerNodes/DataflowSamplerMultiInput.h"
+
+#include "DataflowMinSamplerNode.generated.h"
+
+USTRUCT()
+struct FDataflowMinFloatSampler : public FDataflowFloatSamplerBase
+{
+	GENERATED_BODY();
+	virtual ~FDataflowMinFloatSampler() override = default;
+	virtual void Sample(TArrayView<const FVector3f> Positions, TArrayView<float> OutValues) const override;
+	virtual FBox GetRenderBounds() const override;
+
+public:
+	/** list of samplers to Min results from */
+	TArray<TSharedPtr<const FDataflowFloatSamplerBase>> Samplers;
+};
+
+USTRUCT()
+struct FDataflowMinVectorSampler : public FDataflowVectorSamplerBase
+{
+	GENERATED_BODY();
+	virtual ~FDataflowMinVectorSampler() override = default;
+	virtual void Sample(TArrayView<const FVector3f> Positions, TArrayView<FVector3f> OutValues) const override;
+	virtual FBox GetRenderBounds() const override;
+
+public:
+	/** list of samplers to Min results from */
+	TArray<TSharedPtr<const FDataflowVectorSamplerBase>> Samplers;
+};
+
+/**
+* Min sampler results together
+*/
+USTRUCT()
+struct FDataflowMinSamplerNode : public FDataflowMultiInputSamplerNodeBase
+{
+	GENERATED_USTRUCT_BODY()
+
+	DATAFLOW_NODE_DEFINE_INTERNAL(FDataflowMinSamplerNode, "Min Sampler", "Samplers", "Minimum")
+
+public:
+	FDataflowMinSamplerNode(const UE::Dataflow::FNodeParameters& InParam, FGuid InGuid = FGuid::NewGuid());
+	virtual void EvaluateSampler(UE::Dataflow::FContext& Context, FDataflowSamplerTypes::FStorageType& OutSampler) const override;
+};
+
diff --git a/Engine/Plugins/Experimental/Dataflow/Source/DataflowNodes/Private/Dataflow/SamplerNodes/DataflowMultiplySamplerNode.cpp b/Engine/Plugins/Experimental/Dataflow/Source/DataflowNodes/Private/Dataflow/SamplerNodes/DataflowMultiplySamplerNode.cpp
new file mode 100644
--- /dev/null
+++ b/Engine/Plugins/Experimental/Dataflow/Source/DataflowNodes/Private/Dataflow/SamplerNodes/DataflowMultiplySamplerNode.cpp
@@ -0,0 +1,119 @@
+// Copyright Epic Games, Inc. All Rights Reserved.
+
+#include "Dataflow/SamplerNodes/DataflowMultiplySamplerNode.h"
+#include "Dataflow/DataflowNodeFactory.h"
+
+#include UE_INLINE_GENERATED_CPP_BY_NAME(DataflowMultiplySamplerNode)
+
+#define LOCTEXT_NAMESPACE "DataflowMultiplySamplerNode"
+
+FBox FDataflowMultiplyFloatSampler::GetRenderBounds() const
+{
+	return FBox(FVector(-50.0), FVector(50.0));
+}
+
+void FDataflowMultiplyFloatSampler::Sample(TArrayView<const FVector3f> Positions, TArrayView<float> OutValues) const
+{
+	for (int32 Index = 0; Index < OutValues.Num(); ++Index)
+	{
+		OutValues[Index] = 1.f;
+	}
+
+	if (Positions.Num() == OutValues.Num())
+	{
+		TArray<float> OneSamplerValues;
+
+		for (const TSharedPtr<const FDataflowFloatSamplerBase>& Sampler : Samplers)
+		{
+			OneSamplerValues.Init(0, OutValues.Num());
+			Sampler->Sample(Positions, OneSamplerValues);
+			for (int32 Index = 0; Index < OutValues.Num(); ++Index)
+			{
+				OutValues[Index] *= OneSamplerValues[Index];
+			}
+		}
+	}
+}
+
+FBox FDataflowMultiplyVectorSampler::GetRenderBounds() const
+{
+	return FBox(FVector(-50.0), FVector(50.0));
+}
+
+void FDataflowMultiplyVectorSampler::Sample(TArrayView<const FVector3f> Positions, TArrayView<FVector3f> OutValues) const
+{
+	for (int32 Index = 0; Index < OutValues.Num(); ++Index)
+	{
+		OutValues[Index] = FVector3f::OneVector;
+	}
+
+	if (Positions.Num() == OutValues.Num())
+	{
+		TArray<FVector3f> OneSamplerValues;
+
+		for (const TSharedPtr<const FDataflowVectorSamplerBase>& Sampler : Samplers)
+		{
+			OneSamplerValues.Init(FVector3f::ZeroVector, OutValues.Num());
+			Sampler->Sample(Positions, OneSamplerValues);
+			for (int32 Index = 0; Index < OutValues.Num(); ++Index)
+			{
+				OutValues[Index].X *= OneSamplerValues[Index].X;
+				OutValues[Index].Y *= OneSamplerValues[Index].Y;
+				OutValues[Index].Z *= OneSamplerValues[Index].Z;
+			}
+		}
+	}
+}
+
+//////////////////////////////////////////////////////////////////////////////////////////////////////////////////////
+
+FDataflowMultiplySamplerNode::FDataflowMultiplySamplerNode(const UE::Dataflow::FNodeParameters& InParam, FGuid InGuid)
+	: FDataflowMultiInputSamplerNodeBase(InParam, InGuid)
+{
+	RegisterInitialConnections();
+}
+
+void FDataflowMultiplySamplerNode::EvaluateSampler(UE::Dataflow::FContext& Context, FDataflowSamplerTypes::FStorageType& OutSampler) const
+{
+	// use the first sampler as a reference
+	const int32 NumSamplers = GetNumSamplerInputs();
+	if (NumSamplers == 0)
+	{
+		return;
+	}
+
+	const FDataflowSamplerTypes::FStorageType& FirstSampler = GetSamplerInput(Context, 0);
+	if (FirstSampler.TryGet<FDataflowFloatSampler>())
+	{
+		TSharedPtr<FDataflowMultiplyFloatSampler> OutImpl = MakeShared<FDataflowMultiplyFloatSampler>();
+		for (int32 Index = 0; Index < NumSamplers; ++Index)
+		{
+			if (const FDataflowFloatSampler* FloatSampler = GetSamplerInput(Context, Index).TryGet<FDataflowFloatSampler>())
+			{
+				if (TSharedPtr<const FDataflowFloatSamplerBase> SamplerImpl = FloatSampler->GetImpl())
+				{
+					OutImpl->Samplers.Add(SamplerImpl);
+				}
+			}
+		}
+		OutSampler = FDataflowFloatSampler(OutImpl);
+	}
+	else if (FirstSampler.TryGet<FDataflowVectorSampler>())
+	{
+		TSharedPtr<FDataflowMultiplyVectorSampler> OutImpl = MakeShared<FDataflowMultiplyVectorSampler>();
+		for (int32 Index = 0; Index < NumSamplers; ++Index)
+		{
+			if (const FDataflowVectorSampler* VectorSampler = GetSamplerInput(Context, Index).TryGet<FDataflowVectorSampler>())
+			{
+				if (TSharedPtr<const FDataflowVectorSamplerBase> SamplerImpl = VectorSampler->GetImpl())
+				{
+					OutImpl->Samplers.Add(SamplerImpl);
+				}
+			}
+		}
+		OutSampler = FDataflowVectorSampler(OutImpl);
+	}
+
+}
+
+#undef LOCTEXT_NAMESPACE 
diff --git a/Engine/Plugins/Experimental/Dataflow/Source/DataflowNodes/Private/Dataflow/SamplerNodes/DataflowMultiplySamplerNode.h b/Engine/Plugins/Experimental/Dataflow/Source/DataflowNodes/Private/Dataflow/SamplerNodes/DataflowMultiplySamplerNode.h
new file mode 100644
--- /dev/null
+++ b/Engine/Plugins/Experimental/Dataflow/Source/DataflowNodes/Private/Dataflow/SamplerNodes/DataflowMultiplySamplerNode.h
@@ -0,0 +1,49 @@
+// Copyright Epic Games, Inc. All Rights Reserved.
+
+#pragma once 
+#include "Dataflow/SamplerNodes/DataflowSamplerTypes.h"
+#include "Dataflow/SamplerNodes/DataflowSamplerMultiInput.h"
+
+#include "DataflowMultiplySamplerNode.generated.h"
+
+USTRUCT()
+struct FDataflowMultiplyFloatSampler : public FDataflowFloatSamplerBase
+{
+	GENERATED_BODY();
+	virtual ~FDataflowMultiplyFloatSampler() override = default;
+	virtual void Sample(TArrayView<const FVector3f> Positions, TArrayView<float> OutValues) const override;
+	virtual FBox GetRenderBounds() const override;
+
+public:
+	/** list of samplers to add results from */
+	TArray<TSharedPtr<const FDataflowFloatSamplerBase>> Samplers;
+};
+
+USTRUCT()
+struct FDataflowMultiplyVectorSampler : public FDataflowVectorSamplerBase
+{
+	GENERATED_BODY();
+	virtual ~FDataflowMultiplyVectorSampler() override = default;
+	virtual void Sample(TArrayView<const FVector3f> Positions, TArrayView<FVector3f> OutValues) const override;
+	virtual FBox GetRenderBounds() const override;
+
+public:
+	/** list of samplers to add results from */
+	TArray<TSharedPtr<const FDataflowVectorSamplerBase>> Samplers;
+};
+
+/**
+* Multiply sampler results together
+*/
+USTRUCT()
+struct FDataflowMultiplySamplerNode : public FDataflowMultiInputSamplerNodeBase
+{
+	GENERATED_USTRUCT_BODY()
+
+	DATAFLOW_NODE_DEFINE_INTERNAL(FDataflowMultiplySamplerNode, "Multiply Sampler", "Samplers", "* Multiply")
+
+public:
+	FDataflowMultiplySamplerNode(const UE::Dataflow::FNodeParameters& InParam, FGuid InGuid = FGuid::NewGuid());
+	virtual void EvaluateSampler(UE::Dataflow::FContext& Context, FDataflowSamplerTypes::FStorageType& OutSampler) const override;
+};
+
diff --git a/Engine/Plugins/Experimental/Dataflow/Source/DataflowNodes/Private/Dataflow/SamplerNodes/DataflowPerlinNoiseSamplerNode.cpp b/Engine/Plugins/Experimental/Dataflow/Source/DataflowNodes/Private/Dataflow/SamplerNodes/DataflowPerlinNoiseSamplerNode.cpp
new file mode 100644
--- /dev/null
+++ b/Engine/Plugins/Experimental/Dataflow/Source/DataflowNodes/Private/Dataflow/SamplerNodes/DataflowPerlinNoiseSamplerNode.cpp
@@ -0,0 +1,83 @@
+// Copyright Epic Games, Inc. All Rights Reserved.
+
+#include "Dataflow/SamplerNodes/DataflowPerlinNoiseSamplerNode.h"
+#include "Dataflow/SamplerNodes/DataflowSamplerFunctions.h"
+#include "Dataflow/DataflowNodeFactory.h"
+
+FBox FDataflowPerlinNoiseFloatSampler::GetRenderBounds() const
+{
+	return FBox(FVector(-50.0), FVector(50.0));
+}
+
+void FDataflowPerlinNoiseFloatSampler::Sample(TArrayView<const FVector3f> Positions, TArrayView<float> OutValues) const
+{
+	const int32 NumSamples = FMath::Min(Positions.Num(), OutValues.Num());
+	
+	for (int32 Index = 0; Index < NumSamples; ++Index)
+	{
+		OutValues[Index] = Amplitude * UE::Dataflow::Sampler::SPerlinNoise(Frequency * FVector(Positions[Index]) + Offset);
+	}	
+}
+
+
+FDataflowPerlinNoiseFloatSamplerNode::FDataflowPerlinNoiseFloatSamplerNode(const UE::Dataflow::FNodeParameters& InParam, FGuid InGuid)
+	: FDataflowNode(InParam, InGuid)
+{
+	RegisterInputConnection(&PerlinNoiseSampler.Amplitude, GET_MEMBER_NAME_CHECKED(FDataflowPerlinNoiseFloatSampler, Amplitude));
+	RegisterInputConnection(&PerlinNoiseSampler.Frequency, GET_MEMBER_NAME_CHECKED(FDataflowPerlinNoiseFloatSampler, Frequency));
+	RegisterInputConnection(&PerlinNoiseSampler.Offset, GET_MEMBER_NAME_CHECKED(FDataflowPerlinNoiseFloatSampler, Offset));
+	RegisterOutputConnection(&Sampler);
+}
+
+void FDataflowPerlinNoiseFloatSamplerNode::Evaluate(UE::Dataflow::FContext& Context, const FDataflowOutput* Out) const
+{
+	if (Out->IsA(&Sampler))
+	{
+		TSharedRef<FDataflowPerlinNoiseFloatSampler> Impl = MakeShared<FDataflowPerlinNoiseFloatSampler>(PerlinNoiseSampler);
+		Impl->Amplitude = GetValue(Context, &PerlinNoiseSampler.Amplitude);
+		Impl->Frequency = GetValue(Context, &PerlinNoiseSampler.Frequency);
+		Impl->Offset = GetValue(Context, &PerlinNoiseSampler.Offset);
+
+		FDataflowFloatSampler OutSampler(Impl.ToSharedPtr());
+		SetValue(Context, OutSampler, &Sampler);
+	}
+}
+
+FBox FDataflowPerlinNoiseVectorSampler::GetRenderBounds() const
+{
+	return FBox(FVector(-50.0), FVector(50.0));
+}
+
+void FDataflowPerlinNoiseVectorSampler::Sample(TArrayView<const FVector3f> Positions, TArrayView<FVector3f> OutValues) const
+{
+	const int32 NumSamples = FMath::Min(Positions.Num(), OutValues.Num());
+
+	for (int32 Index = 0; Index < NumSamples; ++Index)
+	{
+		OutValues[Index] = Amplitude * FVector3f(UE::Dataflow::Sampler::VPerlinNoise(Frequency * FVector(Positions[Index]) + Offset));
+	}
+}
+
+FDataflowPerlinNoiseVectorSamplerNode::FDataflowPerlinNoiseVectorSamplerNode(const UE::Dataflow::FNodeParameters& InParam, FGuid InGuid)
+	: FDataflowNode(InParam, InGuid)
+{
+	RegisterInputConnection(&PerlinNoiseSampler.Amplitude, GET_MEMBER_NAME_CHECKED(FDataflowPerlinNoiseVectorSampler, Amplitude));
+	RegisterInputConnection(&PerlinNoiseSampler.Frequency, GET_MEMBER_NAME_CHECKED(FDataflowPerlinNoiseVectorSampler, Frequency));
+	RegisterInputConnection(&PerlinNoiseSampler.Offset, GET_MEMBER_NAME_CHECKED(FDataflowPerlinNoiseVectorSampler, Offset));
+	RegisterOutputConnection(&Sampler);
+}
+
+void FDataflowPerlinNoiseVectorSamplerNode::Evaluate(UE::Dataflow::FContext& Context, const FDataflowOutput* Out) const
+{
+	if (Out->IsA(&Sampler))
+	{
+		TSharedRef<FDataflowPerlinNoiseVectorSampler> Impl = MakeShared<FDataflowPerlinNoiseVectorSampler>(PerlinNoiseSampler);
+		Impl->Amplitude = GetValue(Context, &PerlinNoiseSampler.Amplitude);
+		Impl->Frequency = GetValue(Context, &PerlinNoiseSampler.Frequency);
+		Impl->Offset = GetValue(Context, &PerlinNoiseSampler.Offset);
+
+		FDataflowVectorSampler OutSampler(Impl.ToSharedPtr());
+		SetValue(Context, OutSampler, &Sampler);
+	}
+}
+
diff --git a/Engine/Plugins/Experimental/Dataflow/Source/DataflowNodes/Private/Dataflow/SamplerNodes/DataflowPerlinNoiseSamplerNode.h b/Engine/Plugins/Experimental/Dataflow/Source/DataflowNodes/Private/Dataflow/SamplerNodes/DataflowPerlinNoiseSamplerNode.h
new file mode 100644
--- /dev/null
+++ b/Engine/Plugins/Experimental/Dataflow/Source/DataflowNodes/Private/Dataflow/SamplerNodes/DataflowPerlinNoiseSamplerNode.h
@@ -0,0 +1,97 @@
+// Copyright Epic Games, Inc. All Rights Reserved.
+
+#pragma once 
+
+#include "Dataflow/DataflowNode.h"
+#include "Dataflow/SamplerNodes/DataflowSamplerTypes.h"
+
+#include "DataflowPerlinNoiseSamplerNode.generated.h"
+
+USTRUCT()
+struct FDataflowPerlinNoiseFloatSampler : public FDataflowFloatSamplerBase
+{
+	GENERATED_BODY();
+
+	virtual ~FDataflowPerlinNoiseFloatSampler() override = default;
+	virtual void Sample(TArrayView<const FVector3f> Positions, TArrayView<float> OutValues) const override;
+	virtual FBox GetRenderBounds() const override;
+
+public:
+	UPROPERTY(EditAnywhere, Category = "Perlin Noise", meta = (DataflowInput, UIMin = 0, ClampMin = 0))
+	float Amplitude = 1.f;
+
+	// 
+	UPROPERTY(EditAnywhere, Category = "Perlin Noise", meta = (DataflowInput, UIMin = 0, ClampMin = 0))
+	float Frequency = 1.f;
+
+	// 
+	UPROPERTY(EditAnywhere, Category = "Perlin Noise", meta = (DataflowInput, UIMin = 0, ClampMin = 0))
+	FVector Offset = FVector(0.f);
+};
+
+/**
+ * Output a Perlin Noise float sampler
+ */
+USTRUCT()
+struct FDataflowPerlinNoiseFloatSamplerNode : public FDataflowNode
+{
+	GENERATED_BODY()
+	DATAFLOW_NODE_DEFINE_INTERNAL(FDataflowPerlinNoiseFloatSamplerNode, "Perlin Noise Float Sampler", "Samplers", "")
+
+public:
+	FDataflowPerlinNoiseFloatSamplerNode(const UE::Dataflow::FNodeParameters& InParam, FGuid InGuid = FGuid::NewGuid());
+
+private:
+	UPROPERTY(meta = (DataflowOutput))
+	FDataflowFloatSampler Sampler;
+
+	UPROPERTY(EditAnywhere, Category = Sampler, meta = (ShowOnlyInnerProperties, SkipInDisplayNameChain))
+	FDataflowPerlinNoiseFloatSampler PerlinNoiseSampler;
+
+	void Evaluate(UE::Dataflow::FContext& Context, const FDataflowOutput* Out) const override;
+};
+
+USTRUCT()
+struct FDataflowPerlinNoiseVectorSampler : public FDataflowVectorSamplerBase
+{
+	GENERATED_BODY();
+
+	virtual ~FDataflowPerlinNoiseVectorSampler() override = default;
+	virtual void Sample(TArrayView<const FVector3f> Positions, TArrayView<FVector3f> OutValues) const override;
+	virtual FBox GetRenderBounds() const override;
+
+public:
+	UPROPERTY(EditAnywhere, Category = "Perlin Noise", meta = (DataflowInput, UIMin = 0, ClampMin = 0))
+	float Amplitude = 1.f;
+
+	// 
+	UPROPERTY(EditAnywhere, Category = "Perlin Noise", meta = (DataflowInput, UIMin = 0, ClampMin = 0))
+	float Frequency = 1.f;
+
+	// 
+	UPROPERTY(EditAnywhere, Category = "Perlin Noise", meta = (DataflowInput, UIMin = 0, ClampMin = 0))
+	FVector Offset = FVector(0.f);
+};
+
+/**
+ * Output a Perlin Noise vector sampler
+ */
+USTRUCT()
+struct FDataflowPerlinNoiseVectorSamplerNode : public FDataflowNode
+{
+	GENERATED_BODY()
+	DATAFLOW_NODE_DEFINE_INTERNAL(FDataflowPerlinNoiseVectorSamplerNode, "Perlin Noise Vector Sampler", "Samplers", "")
+
+public:
+	FDataflowPerlinNoiseVectorSamplerNode(const UE::Dataflow::FNodeParameters& InParam, FGuid InGuid = FGuid::NewGuid());
+
+private:
+	UPROPERTY(meta = (DataflowOutput))
+	FDataflowVectorSampler Sampler;
+
+	UPROPERTY(EditAnywhere, Category = Sampler, meta = (ShowOnlyInnerProperties, SkipInDisplayNameChain))
+	FDataflowPerlinNoiseVectorSampler PerlinNoiseSampler;
+
+	void Evaluate(UE::Dataflow::FContext& Context, const FDataflowOutput* Out) const override;
+};
+
diff --git a/Engine/Plugins/Experimental/Dataflow/Source/DataflowNodes/Private/Dataflow/SamplerNodes/DataflowRandomSamplerNode.cpp b/Engine/Plugins/Experimental/Dataflow/Source/DataflowNodes/Private/Dataflow/SamplerNodes/DataflowRandomSamplerNode.cpp
new file mode 100644
--- /dev/null
+++ b/Engine/Plugins/Experimental/Dataflow/Source/DataflowNodes/Private/Dataflow/SamplerNodes/DataflowRandomSamplerNode.cpp
@@ -0,0 +1,82 @@
+// Copyright Epic Games, Inc. All Rights Reserved.
+
+#include "Dataflow/SamplerNodes/DataflowRandomSamplerNode.h"
+#include "Dataflow/SamplerNodes/DataflowSamplerFunctions.h"
+#include "Dataflow/DataflowNodeFactory.h"
+
+FBox FDataflowRandomFloatSampler::GetRenderBounds() const
+{
+	return FBox(FVector(-50.0), FVector(50.0));
+}
+
+void FDataflowRandomFloatSampler::Sample(TArrayView<const FVector3f> Positions, TArrayView<float> OutValues) const
+{
+	const int32 NumSamples = FMath::Min(Positions.Num(), OutValues.Num());
+	
+	for (int32 Index = 0; Index < NumSamples; ++Index)
+	{
+		OutValues[Index] = UE::Dataflow::Sampler::FRandom(FVector(Positions[Index]), Min, Max, Seed);
+	}
+}
+
+FDataflowRandomFloatSamplerNode::FDataflowRandomFloatSamplerNode(const UE::Dataflow::FNodeParameters& InParam, FGuid InGuid)
+	: FDataflowNode(InParam, InGuid)
+{
+	RegisterInputConnection(&RandomSampler.Seed, GET_MEMBER_NAME_CHECKED(FDataflowRandomFloatSampler, Seed));
+	RegisterInputConnection(&RandomSampler.Min, GET_MEMBER_NAME_CHECKED(FDataflowRandomFloatSampler, Min));
+	RegisterInputConnection(&RandomSampler.Max, GET_MEMBER_NAME_CHECKED(FDataflowRandomFloatSampler, Max));
+	RegisterOutputConnection(&Sampler);
+}
+
+void FDataflowRandomFloatSamplerNode::Evaluate(UE::Dataflow::FContext& Context, const FDataflowOutput* Out) const
+{
+	if (Out->IsA(&Sampler))
+	{
+		TSharedRef<FDataflowRandomFloatSampler> Impl = MakeShared<FDataflowRandomFloatSampler>(RandomSampler);
+		Impl->Seed = GetValue(Context, &RandomSampler.Seed);
+		Impl->Min = GetValue(Context, &RandomSampler.Min);
+		Impl->Max = GetValue(Context, &RandomSampler.Max);
+
+		FDataflowFloatSampler OutSampler(Impl.ToSharedPtr());
+		SetValue(Context, OutSampler, &Sampler);
+	}
+}
+
+FBox FDataflowRandomVectorSampler::GetRenderBounds() const
+{
+	return FBox(FVector(-50.0), FVector(50.0));
+}
+
+void FDataflowRandomVectorSampler::Sample(TArrayView<const FVector3f> Positions, TArrayView<FVector3f> OutValues) const
+{
+	const int32 NumSamples = FMath::Min(Positions.Num(), OutValues.Num());
+
+	for (int32 Index = 0; Index < NumSamples; ++Index)
+	{
+		OutValues[Index] = FVector3f(UE::Dataflow::Sampler::VRandom(FVector(Positions[Index]), Min, Max, Seed));
+	}
+}
+
+FDataflowRandomVectorSamplerNode::FDataflowRandomVectorSamplerNode(const UE::Dataflow::FNodeParameters& InParam, FGuid InGuid)
+	: FDataflowNode(InParam, InGuid)
+{
+	RegisterInputConnection(&RandomSampler.Seed, GET_MEMBER_NAME_CHECKED(FDataflowRandomVectorSampler, Seed));
+	RegisterInputConnection(&RandomSampler.Min, GET_MEMBER_NAME_CHECKED(FDataflowRandomVectorSampler, Min));
+	RegisterInputConnection(&RandomSampler.Max, GET_MEMBER_NAME_CHECKED(FDataflowRandomVectorSampler, Max));
+	RegisterOutputConnection(&Sampler);
+}
+
+void FDataflowRandomVectorSamplerNode::Evaluate(UE::Dataflow::FContext& Context, const FDataflowOutput* Out) const
+{
+	if (Out->IsA(&Sampler))
+	{
+		TSharedRef<FDataflowRandomVectorSampler> Impl = MakeShared<FDataflowRandomVectorSampler>(RandomSampler);
+		Impl->Seed = GetValue(Context, &RandomSampler.Seed);
+		Impl->Min = GetValue(Context, &RandomSampler.Min);
+		Impl->Max = GetValue(Context, &RandomSampler.Max);
+
+		FDataflowVectorSampler OutSampler(Impl.ToSharedPtr());
+		SetValue(Context, OutSampler, &Sampler);
+	}
+}
+
diff --git a/Engine/Plugins/Experimental/Dataflow/Source/DataflowNodes/Private/Dataflow/SamplerNodes/DataflowRandomSamplerNode.h b/Engine/Plugins/Experimental/Dataflow/Source/DataflowNodes/Private/Dataflow/SamplerNodes/DataflowRandomSamplerNode.h
new file mode 100644
--- /dev/null
+++ b/Engine/Plugins/Experimental/Dataflow/Source/DataflowNodes/Private/Dataflow/SamplerNodes/DataflowRandomSamplerNode.h
@@ -0,0 +1,99 @@
+// Copyright Epic Games, Inc. All Rights Reserved.
+
+#pragma once 
+
+#include "Dataflow/DataflowNode.h"
+#include "Dataflow/SamplerNodes/DataflowSamplerTypes.h"
+
+#include "DataflowRandomSamplerNode.generated.h"
+
+USTRUCT()
+struct FDataflowRandomFloatSampler : public FDataflowFloatSamplerBase
+{
+	GENERATED_BODY();
+
+	virtual ~FDataflowRandomFloatSampler() override = default;
+	virtual void Sample(TArrayView<const FVector3f> Positions, TArrayView<float> OutValues) const override;
+	virtual FBox GetRenderBounds() const override;
+
+public:
+	// Random seed
+	UPROPERTY(EditAnywhere, Category = "Random", meta = (DataflowInput, UIMin = 0, ClampMin = 0))
+	int32 Seed = 0.f;
+	
+	// Random Min
+	UPROPERTY(EditAnywhere, Category = "Random", meta = (DataflowInput))
+	float Min = 0.f;
+	
+	// Random Max
+	UPROPERTY(EditAnywhere, Category = "Random", meta = (DataflowInput))
+	float Max = 1.f;
+};
+
+/**
+ * Output a Random float sampler
+ */
+USTRUCT()
+struct FDataflowRandomFloatSamplerNode : public FDataflowNode
+{
+	GENERATED_BODY()
+	DATAFLOW_NODE_DEFINE_INTERNAL(FDataflowRandomFloatSamplerNode, "Random Float Sampler", "Samplers", "")
+
+public:
+	FDataflowRandomFloatSamplerNode(const UE::Dataflow::FNodeParameters& InParam, FGuid InGuid = FGuid::NewGuid());
+
+private:
+	UPROPERTY(meta = (DataflowOutput))
+	FDataflowFloatSampler Sampler;
+
+	UPROPERTY(EditAnywhere, Category = Sampler, meta = (ShowOnlyInnerProperties, SkipInDisplayNameChain))
+	FDataflowRandomFloatSampler RandomSampler;
+
+	void Evaluate(UE::Dataflow::FContext& Context, const FDataflowOutput* Out) const override;
+};
+
+USTRUCT()
+struct FDataflowRandomVectorSampler : public FDataflowVectorSamplerBase
+{
+	GENERATED_BODY();
+
+	virtual ~FDataflowRandomVectorSampler() override = default;
+	virtual void Sample(TArrayView<const FVector3f> Positions, TArrayView<FVector3f> OutValues) const override;
+	virtual FBox GetRenderBounds() const override;
+
+public:
+	// Random seed
+	UPROPERTY(EditAnywhere, Category = "Random", meta = (DataflowInput, UIMin = 0, ClampMin = 0))
+	int32 Seed = 0.f;
+
+	// Random Min
+	UPROPERTY(EditAnywhere, Category = "Random", meta = (DataflowInput))
+	float Min = 0.f;
+
+	// Random Max
+	UPROPERTY(EditAnywhere, Category = "Random", meta = (DataflowInput))
+	float Max = 1.f;
+}; 
+
+/**
+ * Output a random vector sampler
+ */
+USTRUCT()
+struct FDataflowRandomVectorSamplerNode : public FDataflowNode
+{
+	GENERATED_BODY()
+	DATAFLOW_NODE_DEFINE_INTERNAL(FDataflowRandomVectorSamplerNode, "Random Vector Sampler", "Samplers", "")
+
+public:
+	FDataflowRandomVectorSamplerNode(const UE::Dataflow::FNodeParameters& InParam, FGuid InGuid = FGuid::NewGuid());
+
+private:
+	UPROPERTY(meta = (DataflowOutput))
+	FDataflowVectorSampler Sampler;
+
+	UPROPERTY(EditAnywhere, Category = Sampler, meta = (ShowOnlyInnerProperties, SkipInDisplayNameChain))
+	FDataflowRandomVectorSampler RandomSampler;
+
+	void Evaluate(UE::Dataflow::FContext& Context, const FDataflowOutput* Out) const override;
+};
+
diff --git a/Engine/Plugins/Experimental/Dataflow/Source/DataflowNodes/Private/Dataflow/SamplerNodes/DataflowSamplerFunctions.cpp b/Engine/Plugins/Experimental/Dataflow/Source/DataflowNodes/Private/Dataflow/SamplerNodes/DataflowSamplerFunctions.cpp
new file mode 100644
--- /dev/null
+++ b/Engine/Plugins/Experimental/Dataflow/Source/DataflowNodes/Private/Dataflow/SamplerNodes/DataflowSamplerFunctions.cpp
@@ -0,0 +1,133 @@
+// Copyright Epic Games, Inc. All Rights Reserved.
+
+#include "Dataflow/SamplerNodes/DataflowSamplerFunctions.h"
+#include "Math/RandomStream.h"
+
+namespace UE::Dataflow::Sampler
+{
+	float FRandom(const FVector InPosition, float InRandomMin, float InRandomMax, const int32 InSeed)
+	{
+		const int32 SeedX = int32(InPosition.X);
+		const int32 SeedY = int32(InPosition.Y);
+		const int32 SeedZ = int32(InPosition.Z);
+	
+		const int64 SeedAll = int64(SeedX + SeedY + SeedZ + InSeed);
+	
+		FRandomStream RandStream((int32)SeedAll);
+
+		if (InRandomMin > InRandomMax)
+		{
+			InRandomMax = InRandomMin + 0.1f;
+		}
+		
+		return RandStream.FRandRange(InRandomMin, InRandomMax);
+	}
+
+	FVector VRandom(const FVector InPosition, float InRandomMin, float InRandomMax, const int32 InSeed)
+	{
+		const int32 SeedX = int32(InPosition.X);
+		const int32 SeedY = int32(InPosition.Y);
+		const int32 SeedZ = int32(InPosition.Z);
+
+		const int64 SeedAll = int64(SeedX + SeedY + SeedZ + InSeed);
+
+		FRandomStream RandStream((int32)SeedAll);
+
+		if (InRandomMin > InRandomMax)
+		{
+			InRandomMax = InRandomMin + 0.1f;
+		}
+
+		return FVector(RandStream.FRandRange(InRandomMin, InRandomMax), 
+			RandStream.FRandRange(InRandomMin, InRandomMax), 
+			RandStream.FRandRange(InRandomMin, InRandomMax));
+	}
+
+	float SPerlinNoise(const FVector& InPosition)
+	{
+		return FMath::PerlinNoise3D(InPosition);
+	}
+
+	FVector VPerlinNoise(const FVector& InPosition)
+	{
+		return FVector(FMath::PerlinNoise3D(InPosition), 
+			FMath::PerlinNoise3D(InPosition + FVector(12.34)), 
+			FMath::PerlinNoise3D(InPosition + FVector(112.34f)));
+	}
+
+	float fBm(const FVector InPoisition, const float InAmplitude, const float InFrequency, const FVector& InOffset, const int32 InOctaves, const float InLacunarity, const float InGain)
+	{
+		float Sum = 0.f;
+		float Amp = InAmplitude;
+		FVector Position = InPoisition;
+
+		for (int32 Idx = 0; Idx < InOctaves; ++Idx)
+		{
+			Sum += Amp * SPerlinNoise(Position * InFrequency + InOffset);
+
+			Amp *= InGain;
+			Position *= InLacunarity;
+		}
+
+		return Sum;
+	}
+
+	FVector VfBm(const FVector InPoisition, const float InAmplitude, const float InFrequency, const FVector& InOffset, const int32 InOctaves, const float InLacunarity, const float InGain)
+	{
+		FVector Sum = FVector(0.f);
+		float Amp = InAmplitude;
+		FVector Position = InPoisition;
+
+		for (int32 Idx = 0; Idx < InOctaves; ++Idx)
+		{
+			Sum += Amp * VPerlinNoise(Position * InFrequency + InOffset);
+
+			Amp *= InGain;
+			Position *= InLacunarity;
+		}
+
+		return Sum;
+	}
+
+	float Turbulence(const FVector InPoisition, const float InAmplitude, const float InFrequency, const FVector& InOffset, const int32 InOctaves, const float InLacunarity, const float InGain)
+	{
+		float Sum = 0.f;
+		float Amp = InAmplitude;
+		FVector Position = InPoisition;
+
+		for (int32 Idx = 0; Idx < InOctaves; ++Idx)
+		{
+			Sum += Amp * FMath::Abs(SPerlinNoise(Position * InFrequency + InOffset));
+
+			Amp *= InGain;
+			Position *= InLacunarity;
+		}
+
+		return Sum;
+	}
+
+	float Remap(float Value, const float OriginalMin, const float OriginalMax, const float NewMin, const float NewMax, const float Power, const bool bClamp)
+	{
+		if ((OriginalMax - OriginalMin) > SMALL_NUMBER && (NewMax - NewMin) > SMALL_NUMBER)
+		{
+			if (bClamp)
+			{
+				Value = FMath::Clamp(Value, OriginalMin, OriginalMax);
+			}
+			float D = (Value - OriginalMin) / (OriginalMax - OriginalMin);
+			if (D >= 0)
+			{
+				D = FMath::Pow(D, Power);
+				float NewValue = NewMin + D * (NewMax - NewMin);
+
+				return NewValue;
+			}
+			else
+			{
+				return NewMin;
+			}
+		}
+
+		return NewMin;
+	}
+}
diff --git a/Engine/Plugins/Experimental/Dataflow/Source/DataflowNodes/Private/Dataflow/SamplerNodes/DataflowSamplerFunctions.h b/Engine/Plugins/Experimental/Dataflow/Source/DataflowNodes/Private/Dataflow/SamplerNodes/DataflowSamplerFunctions.h
new file mode 100644
--- /dev/null
+++ b/Engine/Plugins/Experimental/Dataflow/Source/DataflowNodes/Private/Dataflow/SamplerNodes/DataflowSamplerFunctions.h
@@ -0,0 +1,25 @@
+// Copyright Epic Games, Inc. All Rights Reserved.
+
+#pragma once 
+
+#include "HAL/Platform.h"
+#include "Math/MathFwd.h"
+
+namespace UE::Dataflow::Sampler
+{
+	float FRandom(const FVector InPosition, float InRandomMin, float InRandomMax, const int32 InSeed);
+
+	FVector VRandom(const FVector InPosition, float InRandomMin, float InRandomMax, const int32 InSeed);
+
+	float SPerlinNoise(const FVector& InPosition);
+
+	FVector VPerlinNoise(const FVector& InPosition);
+
+	float fBm(const FVector InPoisition, const float InAmplitude, const float InFrequency, const FVector& InOffset, const int32 InOctaves, const float InLacunarity, const float InGain);
+
+	FVector VfBm(const FVector InPoisition, const float InAmplitude, const float InFrequency, const FVector& InOffset, const int32 InOctaves, const float InLacunarity, const float InGain);
+
+	float Turbulence(const FVector InPoisition, const float InAmplitude, const float InFrequency, const FVector& InOffset, const int32 InOctaves, const float InLacunarity, const float InGain);
+
+	float Remap(float Value, const float OriginalMin, const float OriginalMax, const float NewMin, const float NewMax, const float Power, const bool bClamp);
+}
diff --git a/Engine/Plugins/Experimental/Dataflow/Source/DataflowNodes/Private/Dataflow/SamplerNodes/DataflowSamplerMultiInput.cpp b/Engine/Plugins/Experimental/Dataflow/Source/DataflowNodes/Private/Dataflow/SamplerNodes/DataflowSamplerMultiInput.cpp
new file mode 100644
--- /dev/null
+++ b/Engine/Plugins/Experimental/Dataflow/Source/DataflowNodes/Private/Dataflow/SamplerNodes/DataflowSamplerMultiInput.cpp
@@ -0,0 +1,151 @@
+// Copyright Epic Games, Inc. All Rights Reserved.
+
+#include "Dataflow/SamplerNodes/DataflowSamplerMultiInput.h"
+
+const FName FDataflowMultiInputSamplerNodeBase::DependencyTypeGroup = "Main";
+FDataflowMultiInputSamplerNodeBase::FDataflowMultiInputSamplerNodeBase(const UE::Dataflow::FNodeParameters& Param, FGuid InGuid)
+	: FDataflowNode(Param, InGuid)
+{
+}
+
+void FDataflowMultiInputSamplerNodeBase::RegisterInitialConnections()
+{
+	// Add two sets of pins to start.
+	for (int32 Index = 0; Index < NumInitialInputs; ++Index)
+	{
+		AddPins();
+	}
+	RegisterOutputConnection(&Sampler)
+		.SetPassthroughInput(GetConnectionReference(0))
+		.SetTypeDependencyGroup(DependencyTypeGroup);
+	check(NumRequiredDataflowInputs + NumInitialInputs == GetNumInputs()); // Update NumRequiredDataflowInputs when adding more inputs. This is used by Serialize
+}
+
+bool FDataflowMultiInputSamplerNodeBase::CanAddPin() const
+{
+	return true;
+}
+
+bool FDataflowMultiInputSamplerNodeBase::CanRemovePin() const
+{
+	return InputSamplers.Num() > NumInitialInputs;
+}
+
+TArray<UE::Dataflow::FPin> FDataflowMultiInputSamplerNodeBase::AddPins()
+{
+	const int32 Index = InputSamplers.AddDefaulted();
+	FDataflowInput& Input = RegisterInputArrayConnection(GetConnectionReference(Index));
+	if (Index > 0)
+	{
+		// Set concrete type the same as Index0.
+		FDataflowInput* const Input0 = FindInput(GetConnectionReference(0));
+		check(Input0);
+		SetConnectionConcreteType(&Input, Input0->GetType(), DependencyTypeGroup);
+	}
+	else
+	{
+		Input.SetTypeDependencyGroup(DependencyTypeGroup);
+	}
+	return { { UE::Dataflow::FPin::EDirection::INPUT, Input.GetType(), Input.GetName() } };
+}
+
+TArray<UE::Dataflow::FPin> FDataflowMultiInputSamplerNodeBase::GetPinsToRemove() const
+{
+	const int32 Index = InputSamplers.Num() - 1;
+	check(InputSamplers.IsValidIndex(Index));
+	if (const FDataflowInput* const Input = FindInput(GetConnectionReference(Index)))
+	{
+		return { { UE::Dataflow::FPin::EDirection::INPUT, Input->GetType(), Input->GetName() } };
+	}
+	return FDataflowNode::GetPinsToRemove();
+}
+
+void FDataflowMultiInputSamplerNodeBase::OnPinRemoved(const UE::Dataflow::FPin& Pin)
+{
+	const int32 Index = InputSamplers.Num() - 1;
+	check(InputSamplers.IsValidIndex(Index));
+#if DO_CHECK
+	const FDataflowInput* const Input = FindInput(GetConnectionReference(Index));
+	check(Input);
+	check(Input->GetName() == Pin.Name);
+	check(Input->GetType() == Pin.Type);
+#endif
+	InputSamplers.SetNum(Index);
+	return FDataflowNode::OnPinRemoved(Pin);
+}
+
+void FDataflowMultiInputSamplerNodeBase::PostSerialize(const FArchive& Ar)
+{
+	if (Ar.IsLoading())
+	{
+		check(InputSamplers.Num() >= NumInitialInputs);
+		for (int32 Index = 0; Index < NumInitialInputs; ++Index)
+		{
+			check(FindInput(GetConnectionReference(Index)));
+		}
+		for (int32 Index = NumInitialInputs; Index < InputSamplers.Num(); ++Index)
+		{
+			FDataflowInput& Input = FindOrRegisterInputArrayConnection(GetConnectionReference(Index));
+			// reset the type to allow the type group to be properly set as well 
+			SetConnectionConcreteType(&Input, Input.GetType(), DependencyTypeGroup);
+		}
+		if (Ar.IsTransacting())
+		{
+			const int32 OrigNumRegisteredInputs = GetNumInputs() - NumRequiredDataflowInputs;
+			const int32 OrigNumInputs = InputSamplers.Num();
+			if (OrigNumRegisteredInputs > OrigNumInputs)
+			{
+				// Inputs have been removed.
+				// Temporarily expand Inputs so we can get connection references.
+				InputSamplers.SetNum(OrigNumRegisteredInputs);
+				for (int32 Index = OrigNumInputs; Index < InputSamplers.Num(); ++Index)
+				{
+					UnregisterInputConnection(GetConnectionReference(Index));
+				}
+				InputSamplers.SetNum(OrigNumInputs);
+			}
+		}
+		else
+		{
+			// Index + all Inputs
+			ensureAlways(InputSamplers.Num() + NumRequiredDataflowInputs == GetNumInputs());
+		}
+	}
+}
+
+UE::Dataflow::TConnectionReference<FDataflowSamplerTypes> FDataflowMultiInputSamplerNodeBase::GetConnectionReference(int32 Index) const
+{
+	return { &InputSamplers[Index], Index, &InputSamplers };
+}
+
+int32 FDataflowMultiInputSamplerNodeBase::GetNumSamplerInputs() const
+{
+	return InputSamplers.Num();
+}
+
+FDataflowSamplerTypes::FStorageType FDataflowMultiInputSamplerNodeBase::GetSamplerInput(UE::Dataflow::FContext& Context, int32 Index) const
+{
+	static const FDataflowSamplerTypes::FStorageType EmptyStorage;
+	if (InputSamplers.IsValidIndex(Index))
+	{
+		const UE::Dataflow::TConnectionReference<FDataflowSamplerTypes> InputReference = GetConnectionReference(Index);
+		return GetValue(Context, InputReference);
+	}
+	return EmptyStorage;
+}
+
+void FDataflowMultiInputSamplerNodeBase::EvaluateSampler(UE::Dataflow::FContext& Context, FDataflowSamplerTypes::FStorageType& OutSampler) const
+{
+	// need to be implemented by the derived class
+	check(false);
+}
+
+void FDataflowMultiInputSamplerNodeBase::Evaluate(UE::Dataflow::FContext& Context, const FDataflowOutput* Out) const
+{
+	if (Out->IsA(&Sampler))
+	{
+		FDataflowSamplerTypes::FStorageType OutSampler;
+		EvaluateSampler(Context, OutSampler);
+		SetValue(Context, OutSampler, &Sampler);
+	}
+}
diff --git a/Engine/Plugins/Experimental/Dataflow/Source/DataflowNodes/Private/Dataflow/SamplerNodes/DataflowSamplerMultiInput.h b/Engine/Plugins/Experimental/Dataflow/Source/DataflowNodes/Private/Dataflow/SamplerNodes/DataflowSamplerMultiInput.h
new file mode 100644
--- /dev/null
+++ b/Engine/Plugins/Experimental/Dataflow/Source/DataflowNodes/Private/Dataflow/SamplerNodes/DataflowSamplerMultiInput.h
@@ -0,0 +1,49 @@
+// Copyright Epic Games, Inc. All Rights Reserved.
+
+#pragma once 
+
+#include "Dataflow/SamplerNodes/DataflowSamplerTypes.h"
+#include "Dataflow/DataflowNode.h"
+#include "Dataflow/DataflowAnyType.h"
+
+#include "DataflowSamplerMultiInput.generated.h"
+
+#define UE_API DATAFLOWNODES_API
+
+USTRUCT()
+struct FDataflowMultiInputSamplerNodeBase : public FDataflowNode
+{
+	GENERATED_USTRUCT_BODY()
+
+public:
+	FDataflowMultiInputSamplerNodeBase() = default;
+	FDataflowMultiInputSamplerNodeBase(const UE::Dataflow::FNodeParameters& Param, FGuid InGuid = FGuid::NewGuid());
+
+protected:
+	void RegisterInitialConnections();
+	int32 GetNumSamplerInputs() const;
+	FDataflowSamplerTypes::FStorageType GetSamplerInput(UE::Dataflow::FContext& Context, int32 Index) const;
+	virtual void EvaluateSampler(UE::Dataflow::FContext& Context, FDataflowSamplerTypes::FStorageType& OutSampler) const;
+	virtual void PostSerialize(const FArchive& Ar) override;
+
+private:
+	UPROPERTY()
+	TArray<FDataflowSamplerTypes> InputSamplers;
+
+	UPROPERTY(meta = (DataflowOutput, DataflowPassthrough = "InputSamplers[0]"))
+	FDataflowSamplerTypes Sampler;
+
+	virtual TArray<UE::Dataflow::FPin> AddPins() override;
+	virtual bool CanAddPin() const override;
+	virtual bool CanRemovePin() const override;
+	virtual TArray<UE::Dataflow::FPin> GetPinsToRemove() const override;
+	virtual void OnPinRemoved(const UE::Dataflow::FPin& Pin) override;
+	void Evaluate(UE::Dataflow::FContext& Context, const FDataflowOutput* Out) const;
+	UE::Dataflow::TConnectionReference<FDataflowSamplerTypes> GetConnectionReference(int32 Index) const;
+
+	static constexpr int32 NumRequiredDataflowInputs = 0;
+	static constexpr int32 NumInitialInputs = 2;
+	static const FName DependencyTypeGroup;
+};
+
+#undef UE_API
diff --git a/Engine/Plugins/Experimental/Dataflow/Source/DataflowNodes/Private/Dataflow/SamplerNodes/DataflowSamplerTypes.cpp b/Engine/Plugins/Experimental/Dataflow/Source/DataflowNodes/Private/Dataflow/SamplerNodes/DataflowSamplerTypes.cpp
--- a/Engine/Plugins/Experimental/Dataflow/Source/DataflowNodes/Private/Dataflow/SamplerNodes/DataflowSamplerTypes.cpp
+++ b/Engine/Plugins/Experimental/Dataflow/Source/DataflowNodes/Private/Dataflow/SamplerNodes/DataflowSamplerTypes.cpp
@@ -2,6 +2,73 @@
 
 #include "Dataflow/SamplerNodes/DataflowSamplerTypes.h"
 #include "Dataflow/DataflowAnyTypeRegistry.h"
+#include "Dataflow/SamplerNodes/DataflowGradientSamplerNode.h"
+#include "Dataflow/SamplerNodes/DataflowTransformSamplerNode.h"
+#include "Dataflow/SamplerNodes/DataflowUniformSamplerNode.h"
+#include "Dataflow/SamplerNodes/DataflowRandomSamplerNode.h"
+#include "Dataflow/SamplerNodes/DataflowPerlinNoiseSamplerNode.h"
+#include "Dataflow/SamplerNodes/DataflowfBmSamplerNode.h"
+#include "Dataflow/SamplerNodes/DataflowTurbulenceSamplerNode.h"
+#include "Dataflow/SamplerNodes/DataflowAbsSamplerNode.h"
+#include "Dataflow/SamplerNodes/DataflowSignSamplerNode.h"
+
+#include "Dataflow/SamplerNodes/DataflowAddSamplerNode.h"
+#include "Dataflow/SamplerNodes/DataflowMultiplySamplerNode.h"
+#include "Dataflow/SamplerNodes/DataflowMinSamplerNode.h"
+#include "Dataflow/SamplerNodes/DataflowMaxSamplerNode.h"
+
+#include "Dataflow/DataflowNodeFactory.h"
+#include "Dataflow/DataflowCategoryRegistry.h"
+#include "Dataflow/DataflowNodeColorsRegistry.h"
+
+namespace UE::Dataflow
+{
+	void RegisterSamplerTypes()
+	{
+		UE_DATAFLOW_REGISTER_ANYTYPE(FDataflowSamplerTypes);
+	}
+};
+
+namespace UE::Dataflow
+{
+	void RegisterSamplerNodes()
+	{
+		static const FLinearColor CDefaultNodeBodyTintColor = FLinearColor(0.f, 0.f, 0.f, 0.5f);
+		static const float CDefaultWireThickness = 2.f;
+
+		DATAFLOW_NODE_REGISTER_CREATION_FACTORY(FDataflowLinearGradientFloatSamplerNode);
+		DATAFLOW_NODE_REGISTER_CREATION_FACTORY(FDataflowLinearGradientFromBoxFloatSamplerNode);
+		DATAFLOW_NODE_REGISTER_CREATION_FACTORY(FDataflowRadialGradientFloatSamplerNode);
+
+		DATAFLOW_NODE_REGISTER_CREATION_FACTORY(FDataflowTransformFloatSamplerNode);
+
+		DATAFLOW_NODE_REGISTER_CREATION_FACTORY(FDataflowUniformFloatSamplerNode);
+		DATAFLOW_NODE_REGISTER_CREATION_FACTORY(FDataflowUniformVectorSamplerNode);
+
+		DATAFLOW_NODE_REGISTER_CREATION_FACTORY(FDataflowRandomFloatSamplerNode);
+		DATAFLOW_NODE_REGISTER_CREATION_FACTORY(FDataflowRandomVectorSamplerNode);
+
+		DATAFLOW_NODE_REGISTER_CREATION_FACTORY(FDataflowPerlinNoiseFloatSamplerNode);
+		DATAFLOW_NODE_REGISTER_CREATION_FACTORY(FDataflowPerlinNoiseVectorSamplerNode);
+
+		DATAFLOW_NODE_REGISTER_CREATION_FACTORY(FDataflowfBmFloatSamplerNode);
+		DATAFLOW_NODE_REGISTER_CREATION_FACTORY(FDataflowfBmVectorSamplerNode);
+
+		DATAFLOW_NODE_REGISTER_CREATION_FACTORY(FDataflowTurbulenceFloatSamplerNode);
+
+		DATAFLOW_NODE_REGISTER_CREATION_FACTORY(FDataflowAbsSamplerNode);
+		DATAFLOW_NODE_REGISTER_CREATION_FACTORY(FDataflowSignSamplerNode);
+
+		DATAFLOW_NODE_REGISTER_CREATION_FACTORY(FDataflowAddSamplerNode);
+		DATAFLOW_NODE_REGISTER_CREATION_FACTORY(FDataflowMultiplySamplerNode);
+		DATAFLOW_NODE_REGISTER_CREATION_FACTORY(FDataflowMinSamplerNode);
+		DATAFLOW_NODE_REGISTER_CREATION_FACTORY(FDataflowMaxSamplerNode);
+
+		DATAFLOW_NODE_REGISTER_CREATION_FACTORY_NODE_COLORS_BY_CATEGORY("Samplers", FLinearColor(1.f, 0.f, 0.93f), CDefaultNodeBodyTintColor);
+		DATAFLOW_NODE_REGISTER_CREATION_FACTORY_PIN_SETTINGS_BY_TYPE("FDataflowFloatSampler", FLinearColor(1.f, 0.3f, 0.95f), CDefaultWireThickness);
+		DATAFLOW_NODE_REGISTER_CREATION_FACTORY_PIN_SETTINGS_BY_TYPE("FDataflowVectorSampler", FLinearColor(0.4f, 0.f, 0.35f), CDefaultWireThickness);
+	}
+}
 
 FDataflowFloatSampler::FDataflowFloatSampler(const TSharedPtr<const FDataflowFloatSamplerBase>& InImpl)
 	: Impl(InImpl)
@@ -23,3 +90,32 @@ FBox FDataflowFloatSampler::GetRenderBounds() const
 	}
 	return FDataflowFloatSamplerBase::GetRenderBounds();
 }
+
+/////////////////////////////////////////////////////////////////////////////////////////
+
+FDataflowVectorSampler::FDataflowVectorSampler(const TSharedPtr<const FDataflowVectorSamplerBase>& InImpl)
+	: Impl(InImpl)
+{
+}
+
+void FDataflowVectorSampler::Sample(TArrayView<const FVector3f> Positions, TArrayView<FVector3f> OutValues) const
+{
+	if (Impl)
+	{
+		Impl->Sample(Positions, OutValues);
+	}
+}
+
+FBox FDataflowVectorSampler::GetRenderBounds() const
+{
+	if (Impl)
+	{
+		return Impl->GetRenderBounds();
+	}
+	return FDataflowVectorSamplerBase::GetRenderBounds();
+}
+
+
+
+
+
diff --git a/Engine/Plugins/Experimental/Dataflow/Source/DataflowNodes/Private/Dataflow/SamplerNodes/DataflowSamplerTypes.h b/Engine/Plugins/Experimental/Dataflow/Source/DataflowNodes/Private/Dataflow/SamplerNodes/DataflowSamplerTypes.h
deleted file mode 100644
diff --git a/Engine/Plugins/Experimental/Dataflow/Source/DataflowNodes/Private/Dataflow/SamplerNodes/DataflowSignSamplerNode.cpp b/Engine/Plugins/Experimental/Dataflow/Source/DataflowNodes/Private/Dataflow/SamplerNodes/DataflowSignSamplerNode.cpp
new file mode 100644
--- /dev/null
+++ b/Engine/Plugins/Experimental/Dataflow/Source/DataflowNodes/Private/Dataflow/SamplerNodes/DataflowSignSamplerNode.cpp
@@ -0,0 +1,111 @@
+// Copyright Epic Games, Inc. All Rights Reserved.
+
+#include "Dataflow/SamplerNodes/DataflowSignSamplerNode.h"
+#include "Dataflow/SamplerNodes/DataflowSamplerFunctions.h"
+#include "Dataflow/DataflowNodeFactory.h"
+
+FBox FDataflowSignFloatSampler::GetRenderBounds() const
+{
+	return FBox(FVector(-50.0), FVector(50.0));
+}
+
+void FDataflowSignFloatSampler::Sample(TArrayView<const FVector3f> Positions, TArrayView<float> OutValues) const
+{
+	if (Sampler.IsValid() && Positions.Num() == OutValues.Num())
+	{
+		TArray<float> OutValueFromInputSampler;
+		OutValueFromInputSampler.SetNumUninitialized(OutValues.Num());
+	
+		Sampler->Sample(Positions, OutValueFromInputSampler);
+	
+		for (int32 Idx = 0; Idx < OutValues.Num(); ++Idx)
+		{
+			OutValues[Idx] = FMath::Sign(OutValueFromInputSampler[Idx]);
+		}
+	}
+	else
+	{
+		for (int32 Idx = 0; Idx < OutValues.Num(); ++Idx)
+		{
+			OutValues[Idx] = 0.f;
+		}
+	}
+}
+
+FBox FDataflowSignVectorSampler::GetRenderBounds() const
+{
+	return FBox(FVector(-50.0), FVector(50.0));
+}
+
+void FDataflowSignVectorSampler::Sample(TArrayView<const FVector3f> Positions, TArrayView<FVector3f> OutValues) const
+{
+	if (Sampler.IsValid() && Positions.Num() == OutValues.Num())
+	{
+		TArray<FVector3f> OutValueFromInputSampler;
+		OutValueFromInputSampler.SetNumUninitialized(OutValues.Num());
+	
+		Sampler->Sample(Positions, OutValueFromInputSampler);
+	
+		for (int32 Idx = 0; Idx < OutValues.Num(); ++Idx)
+		{
+			OutValues[Idx] = FVector3f(FMath::Sign(OutValueFromInputSampler[Idx].X),
+				FMath::Sign(OutValueFromInputSampler[Idx].Y),
+				FMath::Sign(OutValueFromInputSampler[Idx].Z));
+		}
+	}
+	else
+	{
+		for (int32 Idx = 0; Idx < OutValues.Num(); ++Idx)
+		{
+			OutValues[Idx] = FVector3f::ZeroVector;
+		}
+	}
+}
+
+FDataflowSignSamplerNode::FDataflowSignSamplerNode(const UE::Dataflow::FNodeParameters& InParam, FGuid InGuid)
+	: FDataflowNode(InParam, InGuid)
+{
+	static const FName MainTypeGroup("Main");
+	
+	RegisterInputConnection(&Sampler)
+		.SetTypeDependencyGroup(MainTypeGroup);
+	RegisterOutputConnection(&Sampler)
+		.SetPassthroughInput(&Sampler)
+		.SetTypeDependencyGroup(MainTypeGroup);
+}
+
+void FDataflowSignSamplerNode::Evaluate(UE::Dataflow::FContext& Context, const FDataflowOutput* Out) const
+{
+	if (Out->IsA(&Sampler))
+	{
+		const FDataflowFloatSampler* FloatSampler = GetValue(Context, &Sampler).TryGet <FDataflowFloatSampler>();
+		if (FloatSampler)
+		{
+			TSharedRef<FDataflowSignFloatSampler> Impl = MakeShared<FDataflowSignFloatSampler>();
+
+			Impl->Sampler = FloatSampler->GetImpl();
+
+			FDataflowFloatSampler OutSampler(Impl.ToSharedPtr());
+
+			SetValue(Context, OutSampler, &Sampler);
+			return;
+		}
+
+		const FDataflowVectorSampler* VectorSampler = GetValue(Context, &Sampler).TryGet<FDataflowVectorSampler>();
+		if (VectorSampler)
+		{
+			TSharedRef<FDataflowSignVectorSampler> Impl = MakeShared<FDataflowSignVectorSampler>();
+		
+			Impl->Sampler = VectorSampler->GetImpl();
+		
+			FDataflowVectorSampler OutSampler(Impl.ToSharedPtr());
+
+			SetValue(Context, OutSampler, &Sampler);
+			return;
+		}
+
+		SafeForwardInput(Context, &Sampler, &Sampler);
+	}
+}
+
+
diff --git a/Engine/Plugins/Experimental/Dataflow/Source/DataflowNodes/Private/Dataflow/SamplerNodes/DataflowSignSamplerNode.h b/Engine/Plugins/Experimental/Dataflow/Source/DataflowNodes/Private/Dataflow/SamplerNodes/DataflowSignSamplerNode.h
new file mode 100644
--- /dev/null
+++ b/Engine/Plugins/Experimental/Dataflow/Source/DataflowNodes/Private/Dataflow/SamplerNodes/DataflowSignSamplerNode.h
@@ -0,0 +1,58 @@
+// Copyright Epic Games, Inc. All Rights Reserved.
+
+#pragma once 
+
+#include "Dataflow/DataflowNode.h"
+#include "Dataflow/SamplerNodes/DataflowSamplerTypes.h"
+
+#include "DataflowSignSamplerNode.generated.h"
+
+/***********************************************************************************************************************
+ * No Input Sampler nodes
+ ***********************************************************************************************************************/
+USTRUCT()
+struct FDataflowSignFloatSampler : public FDataflowFloatSamplerBase
+{
+	GENERATED_BODY();
+
+	virtual ~FDataflowSignFloatSampler() override = default;
+	virtual void Sample(TArrayView<const FVector3f> Positions, TArrayView<float> OutValues) const override;
+	virtual FBox GetRenderBounds() const override;
+
+public:
+	/**  */
+	TSharedPtr<const FDataflowFloatSamplerBase> Sampler;
+};
+
+USTRUCT()
+struct FDataflowSignVectorSampler : public FDataflowVectorSamplerBase
+{
+	GENERATED_BODY();
+
+	virtual ~FDataflowSignVectorSampler() override = default;
+	virtual void Sample(TArrayView<const FVector3f> Positions, TArrayView<FVector3f> OutValues) const override;
+	virtual FBox GetRenderBounds() const override;
+
+public:
+	/**  */
+	TSharedPtr<const FDataflowVectorSamplerBase> Sampler;
+};
+
+/**
+ * Output an Sign float sampler
+ */
+USTRUCT()
+struct FDataflowSignSamplerNode : public FDataflowNode
+{
+	GENERATED_BODY()
+	DATAFLOW_NODE_DEFINE_INTERNAL(FDataflowSignSamplerNode, "Sign Sampler", "Samplers", "")
+
+public:
+	FDataflowSignSamplerNode(const UE::Dataflow::FNodeParameters& InParam, FGuid InGuid = FGuid::NewGuid());
+
+private:
+	UPROPERTY(meta = (DataflowInput, DataflowOutput, DataflowPassthrough = Sampler))
+	FDataflowSamplerTypes Sampler;
+	
+	void Evaluate(UE::Dataflow::FContext& Context, const FDataflowOutput* Out) const override;
+};
diff --git a/Engine/Plugins/Experimental/Dataflow/Source/DataflowNodes/Private/Dataflow/SamplerNodes/DataflowTransformSamplerNode.cpp b/Engine/Plugins/Experimental/Dataflow/Source/DataflowNodes/Private/Dataflow/SamplerNodes/DataflowTransformSamplerNode.cpp
--- a/Engine/Plugins/Experimental/Dataflow/Source/DataflowNodes/Private/Dataflow/SamplerNodes/DataflowTransformSamplerNode.cpp
+++ b/Engine/Plugins/Experimental/Dataflow/Source/DataflowNodes/Private/Dataflow/SamplerNodes/DataflowTransformSamplerNode.cpp
@@ -13,6 +13,11 @@ namespace UE::Dataflow
 
 ////////////////////////////////////////////////////////////////////////////////////////////////////////////////////////
 
+FBox FDataflowTransformFloatSampler::GetRenderBounds() const
+{
+	return FBox(FVector(-50.0), FVector(50.0));
+}
+
 void FDataflowTransformFloatSampler::Sample(TArrayView<const FVector3f> Positions, TArrayView<float> OutValues) const
 {
 	if (Sampler.IsValid())
@@ -49,8 +54,8 @@ void FDataflowTransformFloatSamplerNode::Evaluate(UE::Dataflow::FContext& Contex
 		TSharedRef<FDataflowTransformFloatSampler> Impl = MakeShared<FDataflowTransformFloatSampler>(TransformSampler);
 		Impl->Transform = GetValue(Context, &TransformSampler.Transform);
 		Impl->Sampler = GetValue(Context, &Sampler).GetImpl();
-	
+
 		FDataflowFloatSampler OutSampler(Impl.ToSharedPtr());
 		SetValue(Context, OutSampler, &Sampler);
 	}
-}
\ No newline at end of file
+}
diff --git a/Engine/Plugins/Experimental/Dataflow/Source/DataflowNodes/Private/Dataflow/SamplerNodes/DataflowTransformSamplerNode.h b/Engine/Plugins/Experimental/Dataflow/Source/DataflowNodes/Private/Dataflow/SamplerNodes/DataflowTransformSamplerNode.h
--- a/Engine/Plugins/Experimental/Dataflow/Source/DataflowNodes/Private/Dataflow/SamplerNodes/DataflowTransformSamplerNode.h
+++ b/Engine/Plugins/Experimental/Dataflow/Source/DataflowNodes/Private/Dataflow/SamplerNodes/DataflowTransformSamplerNode.h
@@ -14,12 +14,14 @@ struct FDataflowTransformFloatSampler : public FDataflowFloatSamplerBase
 
 	virtual ~FDataflowTransformFloatSampler() override = default;
 	virtual void Sample(TArrayView<const FVector3f> Positions, TArrayView<float> OutValues) const override;
+	virtual FBox GetRenderBounds() const override;
 
 public:
 	/** Transform to transform the points into a different space */
 	UPROPERTY(EditAnywhere, Category = Transform, meta = (DataflowInput))
 	FTransform Transform;
 
+	/** Use inverse transform */
 	UPROPERTY(EditAnywhere, Category = Transform)
 	bool bInverseTransform = false;
 
@@ -43,7 +45,7 @@ private:
 	UPROPERTY(meta = (DataflowInput, DataflowOutput, DataflowPassthrough = Sampler))
 	FDataflowFloatSampler Sampler;
 
-	UPROPERTY(EditAnywhere, Category = Sampler, meta = (ShowOnlyInnerProperties))
+	UPROPERTY(EditAnywhere, Category = Sampler, meta = (ShowOnlyInnerProperties, SkipInDisplayNameChain))
 	FDataflowTransformFloatSampler TransformSampler;
 
 	void Evaluate(UE::Dataflow::FContext& Context, const FDataflowOutput* Out) const override;
diff --git a/Engine/Plugins/Experimental/Dataflow/Source/DataflowNodes/Private/Dataflow/SamplerNodes/DataflowTurbulenceSamplerNode.cpp b/Engine/Plugins/Experimental/Dataflow/Source/DataflowNodes/Private/Dataflow/SamplerNodes/DataflowTurbulenceSamplerNode.cpp
new file mode 100644
--- /dev/null
+++ b/Engine/Plugins/Experimental/Dataflow/Source/DataflowNodes/Private/Dataflow/SamplerNodes/DataflowTurbulenceSamplerNode.cpp
@@ -0,0 +1,57 @@
+// Copyright Epic Games, Inc. All Rights Reserved.
+
+#include "Dataflow/SamplerNodes/DataflowTurbulenceSamplerNode.h"
+#include "Dataflow/SamplerNodes/DataflowSamplerFunctions.h"
+#include "Dataflow/DataflowNodeFactory.h"
+
+FBox FDataflowTurbulenceFloatSampler::GetRenderBounds() const
+{
+	return FBox(FVector(-50.0), FVector(50.0));
+}
+
+void FDataflowTurbulenceFloatSampler::Sample(TArrayView<const FVector3f> Positions, TArrayView<float> OutValues) const
+{
+	const int32 NumSamples = FMath::Min(Positions.Num(), OutValues.Num());
+
+	for (int32 Index = 0; Index < NumSamples; ++Index)
+	{
+		OutValues[Index] = UE::Dataflow::Sampler::Turbulence(FVector(Positions[Index]),
+			Amplitude,
+			Frequency,
+			Offset,
+			Octaves,
+			Lacunarity,
+			Gain);
+	}
+}
+
+
+FDataflowTurbulenceFloatSamplerNode::FDataflowTurbulenceFloatSamplerNode(const UE::Dataflow::FNodeParameters& InParam, FGuid InGuid)
+	: FDataflowNode(InParam, InGuid)
+{
+	RegisterInputConnection(&TurbulenceSampler.Amplitude, GET_MEMBER_NAME_CHECKED(FDataflowTurbulenceFloatSampler, Amplitude));
+	RegisterInputConnection(&TurbulenceSampler.Frequency, GET_MEMBER_NAME_CHECKED(FDataflowTurbulenceFloatSampler, Frequency));
+	RegisterInputConnection(&TurbulenceSampler.Offset, GET_MEMBER_NAME_CHECKED(FDataflowTurbulenceFloatSampler, Offset));
+	RegisterInputConnection(&TurbulenceSampler.Octaves, GET_MEMBER_NAME_CHECKED(FDataflowTurbulenceFloatSampler, Octaves));
+	RegisterInputConnection(&TurbulenceSampler.Lacunarity, GET_MEMBER_NAME_CHECKED(FDataflowTurbulenceFloatSampler, Lacunarity));
+	RegisterInputConnection(&TurbulenceSampler.Gain, GET_MEMBER_NAME_CHECKED(FDataflowTurbulenceFloatSampler, Gain));
+	RegisterOutputConnection(&Sampler);
+}
+
+void FDataflowTurbulenceFloatSamplerNode::Evaluate(UE::Dataflow::FContext& Context, const FDataflowOutput* Out) const
+{
+	if (Out->IsA(&Sampler))
+	{
+		TSharedRef<FDataflowTurbulenceFloatSampler> Impl = MakeShared<FDataflowTurbulenceFloatSampler>(TurbulenceSampler);
+		Impl->Amplitude = GetValue(Context, &TurbulenceSampler.Amplitude);
+		Impl->Frequency = GetValue(Context, &TurbulenceSampler.Frequency);
+		Impl->Offset = GetValue(Context, &TurbulenceSampler.Offset);
+		Impl->Octaves = GetValue(Context, &TurbulenceSampler.Octaves);
+		Impl->Lacunarity = GetValue(Context, &TurbulenceSampler.Lacunarity);
+		Impl->Gain = GetValue(Context, &TurbulenceSampler.Gain);
+
+		FDataflowFloatSampler OutSampler(Impl.ToSharedPtr());
+		SetValue(Context, OutSampler, &Sampler);
+	}
+}
+
diff --git a/Engine/Plugins/Experimental/Dataflow/Source/DataflowNodes/Private/Dataflow/SamplerNodes/DataflowTurbulenceSamplerNode.h b/Engine/Plugins/Experimental/Dataflow/Source/DataflowNodes/Private/Dataflow/SamplerNodes/DataflowTurbulenceSamplerNode.h
new file mode 100644
--- /dev/null
+++ b/Engine/Plugins/Experimental/Dataflow/Source/DataflowNodes/Private/Dataflow/SamplerNodes/DataflowTurbulenceSamplerNode.h
@@ -0,0 +1,66 @@
+// Copyright Epic Games, Inc. All Rights Reserved.
+
+#pragma once 
+
+#include "Dataflow/DataflowNode.h"
+#include "Dataflow/SamplerNodes/DataflowSamplerTypes.h"
+
+#include "DataflowTurbulenceSamplerNode.generated.h"
+
+USTRUCT()
+struct FDataflowTurbulenceFloatSampler : public FDataflowFloatSamplerBase
+{
+	GENERATED_BODY();
+
+	virtual ~FDataflowTurbulenceFloatSampler() override = default;
+	virtual void Sample(TArrayView<const FVector3f> Positions, TArrayView<float> OutValues) const override;
+	virtual FBox GetRenderBounds() const override;
+
+public:
+	// Amplitude of the noise
+	UPROPERTY(EditAnywhere, Category = "fBm", meta = (DataflowInput, UIMin = 0, ClampMin = 0))
+	float Amplitude = 1.f;
+
+	// Frequency of the noise
+	UPROPERTY(EditAnywhere, Category = "fBm", meta = (DataflowInput, UIMin = 0, ClampMin = 0))
+	float Frequency = 1.f;
+
+	// Offset for the noise
+	UPROPERTY(EditAnywhere, Category = "fBm", meta = (DataflowInput, UIMin = 0, ClampMin = 0))
+	FVector Offset = FVector(0.f);
+
+	// Number of layers of the summed instances of noise()
+	UPROPERTY(EditAnywhere, Category = "fBm", meta = (DataflowInput, UIMin = 0, ClampMin = 0))
+	int32 Octaves = 1;
+
+	// Controls the uniform irregularity exponent of the series
+	UPROPERTY(EditAnywhere, Category = "fBm", meta = (DataflowInput, UIMin = 0, ClampMin = 0))
+	float Lacunarity = 2.f;
+
+	// Gain
+	UPROPERTY(EditAnywhere, Category = "fBm", meta = (DataflowInput, UIMin = 0, ClampMin = 0))
+	float Gain = 0.5f;
+};
+
+/**
+ * Output a Turbulence float sampler
+ */
+USTRUCT()
+struct FDataflowTurbulenceFloatSamplerNode : public FDataflowNode
+{
+	GENERATED_BODY()
+	DATAFLOW_NODE_DEFINE_INTERNAL(FDataflowTurbulenceFloatSamplerNode, "Turbulence Float Sampler", "Samplers", "")
+
+public:
+	FDataflowTurbulenceFloatSamplerNode(const UE::Dataflow::FNodeParameters& InParam, FGuid InGuid = FGuid::NewGuid());
+
+private:
+	UPROPERTY(meta = (DataflowOutput))
+	FDataflowFloatSampler Sampler;
+
+	UPROPERTY(EditAnywhere, Category = Sampler, meta = (ShowOnlyInnerProperties, SkipInDisplayNameChain))
+	FDataflowTurbulenceFloatSampler TurbulenceSampler;
+
+	void Evaluate(UE::Dataflow::FContext& Context, const FDataflowOutput* Out) const override;
+};
+
diff --git a/Engine/Plugins/Experimental/Dataflow/Source/DataflowNodes/Private/Dataflow/SamplerNodes/DataflowUniformSamplerNode.cpp b/Engine/Plugins/Experimental/Dataflow/Source/DataflowNodes/Private/Dataflow/SamplerNodes/DataflowUniformSamplerNode.cpp
new file mode 100644
--- /dev/null
+++ b/Engine/Plugins/Experimental/Dataflow/Source/DataflowNodes/Private/Dataflow/SamplerNodes/DataflowUniformSamplerNode.cpp
@@ -0,0 +1,70 @@
+// Copyright Epic Games, Inc. All Rights Reserved.
+
+#include "Dataflow/SamplerNodes/DataflowUniformSamplerNode.h"
+#include "Dataflow/DataflowNodeFactory.h"
+
+FBox FDataflowUniformFloatSampler::GetRenderBounds() const
+{
+	return FBox(FVector(-50.0), FVector(50.0));
+}
+
+void FDataflowUniformFloatSampler::Sample(TArrayView<const FVector3f> Positions, TArrayView<float> OutValues) const
+{
+	for (float& OutValue : OutValues)
+	{
+		OutValue = Value;
+	}
+}
+
+
+FDataflowUniformFloatSamplerNode::FDataflowUniformFloatSamplerNode(const UE::Dataflow::FNodeParameters& InParam, FGuid InGuid)
+	: FDataflowNode(InParam, InGuid)
+{
+	RegisterInputConnection(&UniformSampler.Value, GET_MEMBER_NAME_CHECKED(FDataflowUniformFloatSampler, Value));
+	RegisterOutputConnection(&Sampler);
+}
+
+void FDataflowUniformFloatSamplerNode::Evaluate(UE::Dataflow::FContext& Context, const FDataflowOutput* Out) const
+{
+	if (Out->IsA(&Sampler))
+	{
+		TSharedRef<FDataflowUniformFloatSampler> Impl = MakeShared<FDataflowUniformFloatSampler>(UniformSampler);
+		Impl->Value = GetValue(Context, &UniformSampler.Value);
+
+		FDataflowFloatSampler OutSampler(Impl.ToSharedPtr());
+		SetValue(Context, OutSampler, &Sampler);
+	}
+}
+
+FBox FDataflowUniformVectorSampler::GetRenderBounds() const
+{
+	return FBox(FVector(-50.0), FVector(50.0));
+}
+
+void FDataflowUniformVectorSampler::Sample(TArrayView<const FVector3f> Positions, TArrayView<FVector3f> OutValues) const
+{
+	for (FVector3f& OutValue : OutValues)
+	{
+		OutValue = Value;
+	}
+}
+
+FDataflowUniformVectorSamplerNode::FDataflowUniformVectorSamplerNode(const UE::Dataflow::FNodeParameters& InParam, FGuid InGuid)
+	: FDataflowNode(InParam, InGuid)
+{
+	RegisterInputConnection(&UniformSampler.Value, GET_MEMBER_NAME_CHECKED(FDataflowUniformFloatSampler, Value));
+	RegisterOutputConnection(&Sampler);
+}
+
+void FDataflowUniformVectorSamplerNode::Evaluate(UE::Dataflow::FContext& Context, const FDataflowOutput* Out) const
+{
+	if (Out->IsA(&Sampler))
+	{
+		TSharedRef<FDataflowUniformVectorSampler> Impl = MakeShared<FDataflowUniformVectorSampler>(UniformSampler);
+		Impl->Value = GetValue(Context, &UniformSampler.Value);
+
+		FDataflowVectorSampler OutSampler(Impl.ToSharedPtr());
+		SetValue(Context, OutSampler, &Sampler);
+	}
+}
+
diff --git a/Engine/Plugins/Experimental/Dataflow/Source/DataflowNodes/Private/Dataflow/SamplerNodes/DataflowUniformSamplerNode.h b/Engine/Plugins/Experimental/Dataflow/Source/DataflowNodes/Private/Dataflow/SamplerNodes/DataflowUniformSamplerNode.h
new file mode 100644
--- /dev/null
+++ b/Engine/Plugins/Experimental/Dataflow/Source/DataflowNodes/Private/Dataflow/SamplerNodes/DataflowUniformSamplerNode.h
@@ -0,0 +1,83 @@
+// Copyright Epic Games, Inc. All Rights Reserved.
+
+#pragma once 
+
+#include "Dataflow/DataflowNode.h"
+#include "Dataflow/SamplerNodes/DataflowSamplerTypes.h"
+
+#include "DataflowUniformSamplerNode.generated.h"
+
+USTRUCT()
+struct FDataflowUniformFloatSampler : public FDataflowFloatSamplerBase
+{
+	GENERATED_BODY();
+
+	virtual ~FDataflowUniformFloatSampler() override = default;
+	virtual void Sample(TArrayView<const FVector3f> Positions, TArrayView<float> OutValues) const override;
+	virtual FBox GetRenderBounds() const override;
+
+public:
+	// Uniform float value to set on all sampled points
+	UPROPERTY(EditAnywhere, Category = "Uniform", meta = (DataflowInput))
+	float Value = 0.f;
+};
+
+/**
+ * Output a uniform float sampler
+ */
+USTRUCT()
+struct FDataflowUniformFloatSamplerNode : public FDataflowNode
+{
+	GENERATED_BODY()
+	DATAFLOW_NODE_DEFINE_INTERNAL(FDataflowUniformFloatSamplerNode, "Uniform Float Sampler", "Samplers", "")
+
+public:
+	FDataflowUniformFloatSamplerNode(const UE::Dataflow::FNodeParameters& InParam, FGuid InGuid = FGuid::NewGuid());
+
+private:
+	UPROPERTY(meta = (DataflowOutput))
+	FDataflowFloatSampler Sampler;
+
+	UPROPERTY(EditAnywhere, Category = Sampler, meta = (ShowOnlyInnerProperties, SkipInDisplayNameChain))
+	FDataflowUniformFloatSampler UniformSampler;
+
+	void Evaluate(UE::Dataflow::FContext& Context, const FDataflowOutput* Out) const override;
+};
+
+USTRUCT()
+struct FDataflowUniformVectorSampler : public FDataflowVectorSamplerBase
+{
+	GENERATED_BODY();
+
+	virtual ~FDataflowUniformVectorSampler() override = default;
+	virtual void Sample(TArrayView<const FVector3f> Positions, TArrayView<FVector3f> OutValues) const override;
+	virtual FBox GetRenderBounds() const override;
+
+public:
+	// Uniform vector value to set on all sampled points
+	UPROPERTY(EditAnywhere, Category = "Uniform", meta = (DataflowInput))
+	FVector3f Value = FVector3f(0.f);
+};
+
+/**
+ * Output a uniform vector sampler
+ */
+USTRUCT()
+struct FDataflowUniformVectorSamplerNode : public FDataflowNode
+{
+	GENERATED_BODY()
+	DATAFLOW_NODE_DEFINE_INTERNAL(FDataflowUniformVectorSamplerNode, "Uniform Vector Sampler", "Samplers", "")
+
+public:
+	FDataflowUniformVectorSamplerNode(const UE::Dataflow::FNodeParameters& InParam, FGuid InGuid = FGuid::NewGuid());
+
+private:
+	UPROPERTY(meta = (DataflowOutput))
+	FDataflowVectorSampler Sampler;
+
+	UPROPERTY(EditAnywhere, Category = Sampler, meta = (ShowOnlyInnerProperties, SkipInDisplayNameChain))
+	FDataflowUniformVectorSampler UniformSampler;
+
+	void Evaluate(UE::Dataflow::FContext& Context, const FDataflowOutput* Out) const override;
+};
+
diff --git a/Engine/Plugins/Experimental/Dataflow/Source/DataflowNodes/Private/Dataflow/SamplerNodes/DataflowfBmSamplerNode.cpp b/Engine/Plugins/Experimental/Dataflow/Source/DataflowNodes/Private/Dataflow/SamplerNodes/DataflowfBmSamplerNode.cpp
new file mode 100644
--- /dev/null
+++ b/Engine/Plugins/Experimental/Dataflow/Source/DataflowNodes/Private/Dataflow/SamplerNodes/DataflowfBmSamplerNode.cpp
@@ -0,0 +1,123 @@
+// Copyright Epic Games, Inc. All Rights Reserved.
+
+#include "Dataflow/SamplerNodes/DataflowfBmSamplerNode.h"
+#include "Dataflow/SamplerNodes/DataflowSamplerFunctions.h"
+#include "Dataflow/DataflowNodeFactory.h"
+
+FBox FDataflowfBmFloatSampler::GetRenderBounds() const
+{
+	return FBox(FVector(-50.0), FVector(50.0));
+}
+
+void FDataflowfBmFloatSampler::Sample(TArrayView<const FVector3f> Positions, TArrayView<float> OutValues) const
+{
+	if (Positions.Num() == OutValues.Num())
+	{
+		for (int32 Index = 0; Index < OutValues.Num(); ++Index)
+		{
+			OutValues[Index] = UE::Dataflow::Sampler::fBm(FVector(Positions[Index]),
+				Amplitude,
+				Frequency,
+				Offset,
+				Octaves,
+				Lacunarity,
+				Gain);
+		}
+	}
+	else	
+	{
+		for (int32 Index = 0; Index < OutValues.Num(); ++Index)
+		{
+			OutValues[Index] = 0;
+		}
+	}
+}
+
+
+FDataflowfBmFloatSamplerNode::FDataflowfBmFloatSamplerNode(const UE::Dataflow::FNodeParameters& InParam, FGuid InGuid)
+	: FDataflowNode(InParam, InGuid)
+{
+	RegisterInputConnection(&fBmSampler.Amplitude, GET_MEMBER_NAME_CHECKED(FDataflowfBmFloatSampler, Amplitude));
+	RegisterInputConnection(&fBmSampler.Frequency, GET_MEMBER_NAME_CHECKED(FDataflowfBmFloatSampler, Frequency));
+	RegisterInputConnection(&fBmSampler.Offset, GET_MEMBER_NAME_CHECKED(FDataflowfBmFloatSampler, Offset));
+	RegisterInputConnection(&fBmSampler.Octaves, GET_MEMBER_NAME_CHECKED(FDataflowfBmFloatSampler, Octaves));
+	RegisterInputConnection(&fBmSampler.Lacunarity, GET_MEMBER_NAME_CHECKED(FDataflowfBmFloatSampler, Lacunarity));
+	RegisterInputConnection(&fBmSampler.Gain, GET_MEMBER_NAME_CHECKED(FDataflowfBmFloatSampler, Gain));
+	RegisterOutputConnection(&Sampler);
+}
+
+void FDataflowfBmFloatSamplerNode::Evaluate(UE::Dataflow::FContext& Context, const FDataflowOutput* Out) const
+{
+	if (Out->IsA(&Sampler))
+	{
+		TSharedRef<FDataflowfBmFloatSampler> Impl = MakeShared<FDataflowfBmFloatSampler>(fBmSampler);
+		Impl->Amplitude = GetValue(Context, &fBmSampler.Amplitude);
+		Impl->Frequency = GetValue(Context, &fBmSampler.Frequency);
+		Impl->Offset = GetValue(Context, &fBmSampler.Offset);
+		Impl->Octaves = GetValue(Context, &fBmSampler.Octaves);
+		Impl->Lacunarity = GetValue(Context, &fBmSampler.Lacunarity);
+		Impl->Gain = GetValue(Context, &fBmSampler.Gain);
+
+		FDataflowFloatSampler OutSampler(Impl.ToSharedPtr());
+		SetValue(Context, OutSampler, &Sampler);
+	}
+}
+
+FBox FDataflowfBmVectorSampler::GetRenderBounds() const
+{
+	return FBox(FVector(-50.0), FVector(50.0));
+}
+
+void FDataflowfBmVectorSampler::Sample(TArrayView<const FVector3f> Positions, TArrayView<FVector3f> OutValues) const
+{
+	if (Positions.Num() == OutValues.Num())
+	{
+		for (int32 Index = 0; Index < OutValues.Num(); ++Index)
+		{
+			OutValues[Index] = FVector3f(UE::Dataflow::Sampler::VfBm(FVector(Positions[Index]),
+				Amplitude,
+				Frequency,
+				Offset,
+				Octaves,
+				Lacunarity,
+				Gain));
+		}
+	}
+	else
+	{
+		for (int32 Index = 0; Index < OutValues.Num(); ++Index)
+		{
+			OutValues[Index] = FVector3f(0.f);
+		}
+	}
+}
+
+FDataflowfBmVectorSamplerNode::FDataflowfBmVectorSamplerNode(const UE::Dataflow::FNodeParameters& InParam, FGuid InGuid)
+	: FDataflowNode(InParam, InGuid)
+{
+	RegisterInputConnection(&fBmSampler.Amplitude, GET_MEMBER_NAME_CHECKED(FDataflowfBmFloatSampler, Amplitude));
+	RegisterInputConnection(&fBmSampler.Frequency, GET_MEMBER_NAME_CHECKED(FDataflowfBmFloatSampler, Frequency));
+	RegisterInputConnection(&fBmSampler.Offset, GET_MEMBER_NAME_CHECKED(FDataflowfBmFloatSampler, Offset));
+	RegisterInputConnection(&fBmSampler.Octaves, GET_MEMBER_NAME_CHECKED(FDataflowfBmFloatSampler, Octaves));
+	RegisterInputConnection(&fBmSampler.Lacunarity, GET_MEMBER_NAME_CHECKED(FDataflowfBmFloatSampler, Lacunarity));
+	RegisterInputConnection(&fBmSampler.Gain, GET_MEMBER_NAME_CHECKED(FDataflowfBmFloatSampler, Gain));
+	RegisterOutputConnection(&Sampler);
+}
+
+void FDataflowfBmVectorSamplerNode::Evaluate(UE::Dataflow::FContext& Context, const FDataflowOutput* Out) const
+{
+	if (Out->IsA(&Sampler))
+	{
+		TSharedRef<FDataflowfBmVectorSampler> Impl = MakeShared<FDataflowfBmVectorSampler>(fBmSampler);
+		Impl->Amplitude = GetValue(Context, &fBmSampler.Amplitude);
+		Impl->Frequency = GetValue(Context, &fBmSampler.Frequency);
+		Impl->Offset = GetValue(Context, &fBmSampler.Offset);
+		Impl->Octaves = GetValue(Context, &fBmSampler.Octaves);
+		Impl->Lacunarity = GetValue(Context, &fBmSampler.Lacunarity);
+		Impl->Gain = GetValue(Context, &fBmSampler.Gain);
+
+		FDataflowVectorSampler OutSampler(Impl.ToSharedPtr());
+		SetValue(Context, OutSampler, &Sampler);
+	}
+}
+
diff --git a/Engine/Plugins/Experimental/Dataflow/Source/DataflowNodes/Private/Dataflow/SamplerNodes/DataflowfBmSamplerNode.h b/Engine/Plugins/Experimental/Dataflow/Source/DataflowNodes/Private/Dataflow/SamplerNodes/DataflowfBmSamplerNode.h
new file mode 100644
--- /dev/null
+++ b/Engine/Plugins/Experimental/Dataflow/Source/DataflowNodes/Private/Dataflow/SamplerNodes/DataflowfBmSamplerNode.h
@@ -0,0 +1,126 @@
+// Copyright Epic Games, Inc. All Rights Reserved.
+
+#pragma once 
+
+#include "Dataflow/DataflowNode.h"
+#include "Dataflow/SamplerNodes/DataflowSamplerTypes.h"
+
+#include "DataflowfBmSamplerNode.generated.h"
+
+/***********************************************************************************************************************
+ * No Input Sampler nodes
+ ***********************************************************************************************************************/
+USTRUCT()
+struct FDataflowfBmFloatSampler : public FDataflowFloatSamplerBase
+{
+	GENERATED_BODY();
+
+	virtual ~FDataflowfBmFloatSampler() override = default;
+	virtual void Sample(TArrayView<const FVector3f> Positions, TArrayView<float> OutValues) const override;
+	virtual FBox GetRenderBounds() const override;
+
+public:
+	// 
+	UPROPERTY(EditAnywhere, Category = "fBm", meta = (DataflowInput, UIMin = 0, ClampMin = 0))
+	float Amplitude = 1.f;
+	
+	// 
+	UPROPERTY(EditAnywhere, Category = "fBm", meta = (DataflowInput, UIMin = 0, ClampMin = 0))
+	float Frequency = 1.f;
+	
+	// 
+	UPROPERTY(EditAnywhere, Category = "fBm", meta = (DataflowInput, UIMin = 0, ClampMin = 0))
+	FVector Offset = FVector(0.f);
+	
+	// Number of layers of the summed instances of noise()
+	UPROPERTY(EditAnywhere, Category = "fBm", meta = (DataflowInput, UIMin = 0, ClampMin = 0))
+	int32 Octaves = 1;
+	
+	// Controls the uniform irregularity exponent of the series
+	UPROPERTY(EditAnywhere, Category = "fBm", meta = (DataflowInput, UIMin = 0, ClampMin = 0))
+	float Lacunarity = 2.f;
+	
+	// Gain
+	UPROPERTY(EditAnywhere, Category = "fBm", meta = (DataflowInput, UIMin = 0, ClampMin = 0))
+	float Gain = 0.5f;
+};
+
+/**
+ * Output a Perlin Noise float sampler
+ */
+USTRUCT()
+struct FDataflowfBmFloatSamplerNode : public FDataflowNode
+{
+	GENERATED_BODY()
+	DATAFLOW_NODE_DEFINE_INTERNAL(FDataflowfBmFloatSamplerNode, "fBm Float Sampler", "Samplers", "")
+
+public:
+	FDataflowfBmFloatSamplerNode(const UE::Dataflow::FNodeParameters& InParam, FGuid InGuid = FGuid::NewGuid());
+
+private:
+	UPROPERTY(meta = (DataflowOutput))
+	FDataflowFloatSampler Sampler;
+
+	UPROPERTY(EditAnywhere, Category = Sampler, meta = (ShowOnlyInnerProperties, SkipInDisplayNameChain))
+	FDataflowfBmFloatSampler fBmSampler;
+
+	void Evaluate(UE::Dataflow::FContext& Context, const FDataflowOutput* Out) const override;
+};
+
+USTRUCT()
+struct FDataflowfBmVectorSampler : public FDataflowVectorSamplerBase
+{
+	GENERATED_BODY();
+
+	virtual ~FDataflowfBmVectorSampler() override = default;
+	virtual void Sample(TArrayView<const FVector3f> Positions, TArrayView<FVector3f> OutValues) const override;
+	virtual FBox GetRenderBounds() const override;
+
+public:
+	// 
+	UPROPERTY(EditAnywhere, Category = "fBm", meta = (DataflowInput, UIMin = 0, ClampMin = 0))
+	float Amplitude = 1.f;
+
+	// 
+	UPROPERTY(EditAnywhere, Category = "fBm", meta = (DataflowInput, UIMin = 0, ClampMin = 0))
+	float Frequency = 1.f;
+
+	// 
+	UPROPERTY(EditAnywhere, Category = "fBm", meta = (DataflowInput, UIMin = 0, ClampMin = 0))
+	FVector Offset = FVector(0.f);
+
+	// Number of layers of the summed instances of noise()
+	UPROPERTY(EditAnywhere, Category = "fBm", meta = (DataflowInput, UIMin = 0, ClampMin = 0))
+	int32 Octaves = 1;
+
+	// Controls the uniform irregularity exponent of the series
+	UPROPERTY(EditAnywhere, Category = "fBm", meta = (DataflowInput, UIMin = 0, ClampMin = 0))
+	float Lacunarity = 2.f;
+
+	// Gain
+	UPROPERTY(EditAnywhere, Category = "fBm", meta = (DataflowInput, UIMin = 0, ClampMin = 0))
+	float Gain = 0.5f;
+};
+
+/**
+ * Output a fBm vector sampler
+ */
+USTRUCT()
+struct FDataflowfBmVectorSamplerNode : public FDataflowNode
+{
+	GENERATED_BODY()
+	DATAFLOW_NODE_DEFINE_INTERNAL(FDataflowfBmVectorSamplerNode, "fBm Vector Sampler", "Samplers", "")
+
+public:
+	FDataflowfBmVectorSamplerNode(const UE::Dataflow::FNodeParameters& InParam, FGuid InGuid = FGuid::NewGuid());
+
+private:
+	UPROPERTY(meta = (DataflowOutput))
+	FDataflowVectorSampler Sampler;
+
+	UPROPERTY(EditAnywhere, Category = Sampler, meta = (ShowOnlyInnerProperties, SkipInDisplayNameChain))
+	FDataflowfBmVectorSampler fBmSampler;
+
+	void Evaluate(UE::Dataflow::FContext& Context, const FDataflowOutput* Out) const override;
+};
+
diff --git a/Engine/Plugins/Experimental/Dataflow/Source/DataflowNodes/Public/Dataflow/SamplerNodes/DataflowSamplerTypes.h b/Engine/Plugins/Experimental/Dataflow/Source/DataflowNodes/Public/Dataflow/SamplerNodes/DataflowSamplerTypes.h
new file mode 100644
--- /dev/null
+++ b/Engine/Plugins/Experimental/Dataflow/Source/DataflowNodes/Public/Dataflow/SamplerNodes/DataflowSamplerTypes.h
@@ -0,0 +1,167 @@
+// Copyright Epic Games, Inc. All Rights Reserved.
+
+#pragma once 
+
+#include "Dataflow/DataflowAnyType.h"
+#include "StructUtils/InstancedStruct.h"
+#include "Misc/TVariant.h"
+
+#include "DataflowSamplerTypes.generated.h"
+
+#define UE_API DATAFLOWNODES_API
+
+/**
+* A float sampler, sample at a specific 3D position  and return a single float value per position 
+*/
+USTRUCT()
+struct FDataflowFloatSamplerBase
+{
+	GENERATED_BODY();
+
+public:
+	virtual ~FDataflowFloatSamplerBase() = default;
+	virtual void Sample(TArrayView<const FVector3f> Positions, TArrayView<float> OutValues) const {};
+	virtual FBox GetRenderBounds() const { return FBox(EForceInit::ForceInitToZero); };
+};
+
+USTRUCT()
+struct FDataflowFloatSampler : public FDataflowFloatSamplerBase
+{
+	GENERATED_BODY()
+
+public:
+	FDataflowFloatSampler() = default;
+	FDataflowFloatSampler(const FDataflowFloatSampler&) = default;
+	FDataflowFloatSampler& operator=(const FDataflowFloatSampler&) = default;
+	UE_API FDataflowFloatSampler(const TSharedPtr<const FDataflowFloatSamplerBase>& InImpl);
+
+	UE_API virtual void Sample(TArrayView<const FVector3f> Positions, TArrayView<float> OutValues) const override;
+	UE_API virtual FBox GetRenderBounds() const override;
+	
+	const TSharedPtr<const FDataflowFloatSamplerBase>& GetImpl() const { return Impl; }
+
+private:
+	TSharedPtr<const FDataflowFloatSamplerBase> Impl;
+};
+
+/////////////////////////////////////////////////////////////////////////////////////////
+
+/**
+* A vector sampler, sample at a specific 3D position  and return a FVector3f value per position
+*/
+USTRUCT()
+struct FDataflowVectorSamplerBase
+{
+	GENERATED_BODY();
+public:
+	virtual ~FDataflowVectorSamplerBase() = default;
+	virtual void Sample(TArrayView<const FVector3f> Positions, TArrayView<FVector3f> OutValues) const {};
+	virtual FBox GetRenderBounds() const { return FBox(EForceInit::ForceInitToZero); };
+};
+
+USTRUCT()
+struct FDataflowVectorSampler : public FDataflowVectorSamplerBase
+{
+	GENERATED_BODY()
+
+public:
+	FDataflowVectorSampler() = default;
+	FDataflowVectorSampler(const FDataflowVectorSampler&) = default;
+	FDataflowVectorSampler& operator=(const FDataflowVectorSampler&) = default;
+	UE_API FDataflowVectorSampler(const TSharedPtr<const FDataflowVectorSamplerBase>& InImpl);
+
+	UE_API virtual void Sample(TArrayView<const FVector3f> Positions, TArrayView<FVector3f> OutValues) const override;
+	UE_API virtual FBox GetRenderBounds() const override;
+
+	const TSharedPtr<const FDataflowVectorSamplerBase>& GetImpl() const { return Impl; }
+
+private:
+	TSharedPtr<const FDataflowVectorSamplerBase> Impl;
+};
+
+/////////////////////////////////////////////////////////////////////////////////////////
+
+USTRUCT()
+struct FDataflowSamplerTypesStorage
+{
+	GENERATED_BODY()
+
+	FDataflowSamplerTypesStorage() = default;
+	FDataflowSamplerTypesStorage(const FDataflowSamplerTypesStorage& Other) = default;
+	FDataflowSamplerTypesStorage& operator=(const FDataflowSamplerTypesStorage& Other) = default;
+	FDataflowSamplerTypesStorage(FDataflowSamplerTypesStorage&& Other) = default;
+	FDataflowSamplerTypesStorage& operator=(FDataflowSamplerTypesStorage&& Other) = default;
+	FDataflowSamplerTypesStorage(const FDataflowFloatSampler& Sampler)
+		: Storage(TInPlaceType<FDataflowFloatSampler>(), Sampler)
+	{
+	}
+	FDataflowSamplerTypesStorage(const FDataflowVectorSampler& Sampler)
+		: Storage(TInPlaceType<FDataflowVectorSampler>(), Sampler)
+	{
+	}
+
+	template<typename T>
+	const T* TryGet() const
+	{
+		return Storage.TryGet<T>();
+	}
+	
+private:
+	// Currently we don't support serialization for this storage
+	TVariant<FDataflowFloatSampler, FDataflowVectorSampler> Storage;
+};
+template <>
+struct FDataflowConverter<FDataflowSamplerTypesStorage>
+{
+	template <typename TFromType>
+	static void From(const TFromType& From, FDataflowSamplerTypesStorage& To)
+	{
+		To = TFromType(From);
+	}
+	template <typename TToType>
+	static void To(const FDataflowSamplerTypesStorage& From, TToType& To)
+	{
+		if (const TToType* Ptr = From.TryGet<TToType>())
+		{
+			To = *Ptr;
+		}
+	}
+};
+
+/**
+* Sampler types 
+*/
+UE_DATAFLOW_POLICY_DECLARE_TYPENAME(FDataflowFloatSampler)
+UE_DATAFLOW_POLICY_DECLARE_TYPENAME(FDataflowVectorSampler)
+UE_DATAFLOW_POLICY_DECLARE_TYPENAME(FDataflowSamplerTypesStorage)
+
+struct FDataflowSamplerTypePolicy :
+	public TDataflowMultiTypePolicy<FDataflowFloatSampler, FDataflowVectorSampler>
+{
+};
+
+USTRUCT()
+struct FDataflowSamplerTypes : public FDataflowAnyType
+{
+	using FPolicyType = FDataflowSamplerTypePolicy;
+	using FStorageType = FDataflowSamplerTypesStorage;
+
+	GENERATED_BODY()
+
+	// This needs to be a UPROPERTY() because the anytype system needs to inspect it
+	UPROPERTY()
+	FDataflowSamplerTypesStorage Value;
+};
+namespace UE::Dataflow
+{
+	void RegisterSamplerTypes();
+};
+
+/////////////////////////////////////////////////////////////////////////////////////////
+
+namespace UE::Dataflow
+{
+	void RegisterSamplerNodes();
+}
+
+#undef UE_API
