Dataflow : paint weightmap tool : fix tooltips to be individual on multi-button bars

#Synced-CL 50953304

diff --git a/Engine/Plugins/Experimental/Dataflow/Source/DataflowEditor/Private/DataflowEditorTools/DataflowEditorVertexAttributePaintToolDetailCustomization.cpp b/Engine/Plugins/Experimental/Dataflow/Source/DataflowEditor/Private/DataflowEditorTools/DataflowEditorVertexAttributePaintToolDetailCustomization.cpp
--- a/Engine/Plugins/Experimental/Dataflow/Source/DataflowEditor/Private/DataflowEditorTools/DataflowEditorVertexAttributePaintToolDetailCustomization.cpp
+++ b/Engine/Plugins/Experimental/Dataflow/Source/DataflowEditor/Private/DataflowEditorTools/DataflowEditorVertexAttributePaintToolDetailCustomization.cpp
@@ -132,10 +132,6 @@ namespace UE::Dataflow::Editor
 			.Padding(2.0f)
 			[
 				SNew(SSegmentedControl<EDataflowEditorToolEditMode>)
-				.ToolTipText(LOCTEXT("EditingModeTooltip",
-					"Brush: edit values by painting on mesh.\n"
-					"Mesh: select vertices/edges/faces to edit values directly.\n"
-				))
 				.Value_Lambda([this]()
 					{
 						return ToolProperties.IsValid() ? ToolProperties->EditingMode : EDataflowEditorToolEditMode::Brush;
@@ -153,8 +149,10 @@ namespace UE::Dataflow::Editor
 					})
 				+ SSegmentedControl<EDataflowEditorToolEditMode>::Slot(EDataflowEditorToolEditMode::Brush)
 				.Text(LOCTEXT("BrushEditMode", "Brush"))
+				.ToolTip(LOCTEXT("BrushEditModeTooltip", "Brush: Edit values by painting on mesh."))
 				+ SSegmentedControl<EDataflowEditorToolEditMode>::Slot(EDataflowEditorToolEditMode::Mesh)
 				.Text(LOCTEXT("MeshEditMode", "Mesh"))
+				.ToolTip(LOCTEXT("MeshEditModeTooltip", "Mesh: Select vertices/edges/faces to edit values directly."))
 			]
 		];
 	}
@@ -171,11 +169,6 @@ namespace UE::Dataflow::Editor
 			.HAlign(HAlign_Center)
 			[
 				SNew(SSegmentedControl<EDataflowEditorToolColorMode>)
-				.ToolTipText(LOCTEXT("WeightColorTooltip",
-					"Adjust the value display in the viewport.\n\n"
-					"Greyscale: Displays values by blending from black (0) to white (1).\n"
-					"Ramp: Values 0 and 1 use the min and max colors. values in between 0 and 1 use the ramp colors.\n"
-					"Full Material: Displays normal mesh materials with textures.\n"))
 				.Value_Lambda([this]()
 					{
 						return ToolProperties.IsValid() ? ToolProperties->DisplayProperties.ColorMode : EDataflowEditorToolColorMode::Greyscale;
@@ -189,13 +182,13 @@ namespace UE::Dataflow::Editor
 					})
 				+ SSegmentedControl<EDataflowEditorToolColorMode>::Slot(EDataflowEditorToolColorMode::Greyscale)
 				.Text(LOCTEXT("GreyscaleMode", "Greyscale"))
-				.ToolTipWidget(MakeTooltipWithShortcut(LOCTEXT("GreyscaleModeTooltip", "Greyscale Mode"), UDataflowEditorVertexAttributePaintTool::GreyscaleDisplayInputChord))
+				.ToolTipWidget(MakeTooltipWithShortcut(LOCTEXT("GreyscaleModeTooltip", "Displays values by blending from black (0) to white (1)."), UDataflowEditorVertexAttributePaintTool::GreyscaleDisplayInputChord))
 				+ SSegmentedControl<EDataflowEditorToolColorMode>::Slot(EDataflowEditorToolColorMode::Ramp)
 				.Text(LOCTEXT("RampMode", "Ramp"))
-				.ToolTipWidget(MakeTooltipWithShortcut(LOCTEXT("RampModeModeTooltip", "Color Ramp Mode"), UDataflowEditorVertexAttributePaintTool::ColorMapDisplayInputChord))
+				.ToolTipWidget(MakeTooltipWithShortcut(LOCTEXT("RampModeModeTooltip", "Values are mapped to the custom color ramp with 0 being the left most color of the ramp and 1 the right most color."), UDataflowEditorVertexAttributePaintTool::ColorMapDisplayInputChord))
 				+ SSegmentedControl<EDataflowEditorToolColorMode>::Slot(EDataflowEditorToolColorMode::FullMaterial)
 				.Text(LOCTEXT("MaterialMode", "Full Material"))
-				.ToolTipWidget(MakeTooltipWithShortcut(LOCTEXT("MaterialModeTooltip", "Material Mode"), UDataflowEditorVertexAttributePaintTool::MaterialColorDisplayInputChord))
+				.ToolTipWidget(MakeTooltipWithShortcut(LOCTEXT("MaterialModeTooltip", "Displays normal mesh materials with textures."), UDataflowEditorVertexAttributePaintTool::MaterialColorDisplayInputChord))
 			]
 		];
 	}
@@ -669,12 +662,6 @@ namespace UE::Dataflow::Editor
 			.Padding(2.0f)
 			[
 				SNew(SSegmentedControl<EDataflowEditorToolEditOperation>)
-				.ToolTipText(LOCTEXT("BrushTooltip",
-					"Select the operation to apply when using the brush.\n"
-					"Add: Applies the current value plus the value to the new value.\n"
-					"Replace: applies the current value minus the value to the new value.\n"
-					"Multiply: applies the current value multiplied by the value to the new value.\n"
-					"Relax: applies the average of the connected (by edge) vertex value to the new vertex value."))
 				.Value_Lambda([this]()
 					{
 						return ToolProperties.IsValid() ? ToolProperties->BrushProperties.BrushMode : EDataflowEditorToolEditOperation::Add;
@@ -688,12 +675,16 @@ namespace UE::Dataflow::Editor
 					})
 				+ SSegmentedControl<EDataflowEditorToolEditOperation>::Slot(EDataflowEditorToolEditOperation::Add)
 				.Text(LOCTEXT("BrushAddMode", "Add"))
+				.ToolTip(LOCTEXT("BrushAddTooltip", "Add: Applies the current value plus the value to the new value."))
 				+ SSegmentedControl<EDataflowEditorToolEditOperation>::Slot(EDataflowEditorToolEditOperation::Replace)
 				.Text(LOCTEXT("BrushReplaceMode", "Replace"))
+				.ToolTip(LOCTEXT("BrushReplaceTooltip", "Replace: Applies the current value minus the value to the new value."))
 				+ SSegmentedControl<EDataflowEditorToolEditOperation>::Slot(EDataflowEditorToolEditOperation::Multiply)
 				.Text(LOCTEXT("BrushMultiplyMode", "Multiply"))
+				.ToolTip(LOCTEXT("BrushMultiplyTooltip", "Multiply: Applies the current value multiplied by the value to the new value."))
 				+ SSegmentedControl<EDataflowEditorToolEditOperation>::Slot(EDataflowEditorToolEditOperation::Relax)
 				.Text(LOCTEXT("BrushRelaxMode", "Relax"))
+				.ToolTip(LOCTEXT("BrushRelaxTooltip", "Relax: Applies the average of the connected (by edge) vertex value to the new vertex value."))
 			]
 		];
 	}
@@ -710,9 +701,6 @@ namespace UE::Dataflow::Editor
 			.Padding(2.0f)
 			[
 				SNew(SSegmentedControl<EMeshVertexPaintBrushAreaType>)
-				.ToolTipText(LOCTEXT("BrushFalloffModeTooltip",
-					"Surface: falloff is based on the distance along the surface from the brush center to nearby connected vertices.\n"
-					"Volume: falloff is based on the straight-line distance from the brush center to surrounding vertices.\n"))
 				.Value_Lambda([this]()
 					{
 						return ToolProperties.IsValid() ? ToolProperties->BrushProperties.BrushAreaMode : EMeshVertexPaintBrushAreaType::Connected;
@@ -726,8 +714,10 @@ namespace UE::Dataflow::Editor
 					})
 				+ SSegmentedControl<EMeshVertexPaintBrushAreaType>::Slot(EMeshVertexPaintBrushAreaType::Connected)
 				.Text(LOCTEXT("SurfaceMode", "Surface"))
+				.ToolTip(LOCTEXT("SurfaceModeTooltip", "Surface: Falloff is based on the distance along the surface from the brush center to nearby connected vertices."))
 				+ SSegmentedControl<EMeshVertexPaintBrushAreaType>::Slot(EMeshVertexPaintBrushAreaType::Volumetric)
 				.Text(LOCTEXT("VolumeMode", "Volume"))
+				.ToolTip(LOCTEXT("VolumeModeTooltip", "Volume: Falloff is based on the straight-line distance from the brush center to surrounding vertices"))
 			]
 		];
 	}
@@ -1631,10 +1621,6 @@ namespace UE::Dataflow::Editor
 
 		TSharedRef<SSegmentedControl<EDataflowEditorToolEditOperation>> AddMultiplyButtonsWidget =
 			SNew(SSegmentedControl<EDataflowEditorToolEditOperation>)
-			.ToolTipText(LOCTEXT("InteractiveEditModeTooltip",
-				"Add: applies the current value plus the flood value to the new value.\n"
-				"Multiply: applies the current value multiplied by the flood value to the new value.\n"
-				"This operation applies interactively while dragging the slider. It operates on the currently selected vertices."))
 			.Value(EDataflowEditorToolEditOperation::Add)
 			.OnValueChanged_Lambda([ResetValue](EDataflowEditorToolEditOperation Mode)
 				{
@@ -1642,8 +1628,10 @@ namespace UE::Dataflow::Editor
 				})
 			+ SSegmentedControl<EDataflowEditorToolEditOperation>::Slot(EDataflowEditorToolEditOperation::Add)
 			.Text(LOCTEXT("BrushAddMode", "Add"))
+			.ToolTip(LOCTEXT("BrushAddModeTooltip", "Add: Applies interactively the current value plus the flood value to the new value when dragging the slider."))
 			+ SSegmentedControl<EDataflowEditorToolEditOperation>::Slot(EDataflowEditorToolEditOperation::Multiply)
 			.Text(LOCTEXT("BrushMultiplyMode", "Multiply"))
+			.ToolTip(LOCTEXT("BrushMultiplyModeTooltip", "Multiply: Applies interactively the current value multiplied by the flood value to the new value when dragging the slider."))
 			;
 
 		auto GetMinValue = [AddMultiplyButtonsWidget]() -> float
@@ -1911,13 +1899,9 @@ namespace UE::Dataflow::Editor
 	TSharedRef<SWidget> FVertexAttributePaintToolDetailCustomization::MakeMirrorAxisWidget() const
 	{
 		return SNew(SSegmentedControl<EAxis::Type>)
-			.ToolTipText(LOCTEXT("MirrorAxisTooltip",
-				"X: copies weights across the YZ plane.\n"
-				"Y: copies weights across the XZ plane.\n"
-				"Z: copies weights across the XY plane."))
 			.Value_Lambda([this]()
 				{
-					return Tool.IsValid()? Tool->GetMirrorAxis(): EAxis::X;
+					return Tool.IsValid() ? Tool->GetMirrorAxis() : EAxis::X;
 				})
 			.OnValueChanged_Lambda([this](EAxis::Type Mode)
 				{
@@ -1928,10 +1912,14 @@ namespace UE::Dataflow::Editor
 				})
 			+ SSegmentedControl<EAxis::Type>::Slot(EAxis::X)
 			.Text(LOCTEXT("MirrorXLabel", "X"))
+			.ToolTip(LOCTEXT("MirrorXLabelTooltip", "X: copies weights across the YZ plane."))
 			+ SSegmentedControl<EAxis::Type>::Slot(EAxis::Y)
 			.Text(LOCTEXT("MirrorYLabel", "Y"))
+			.ToolTip(LOCTEXT("MirrorYLabelTooltip", "Y: copies weights across the XZ plane."))
 			+ SSegmentedControl<EAxis::Type>::Slot(EAxis::Z)
-			.Text(LOCTEXT("MirrorZLabel", "Z"));
+			.Text(LOCTEXT("MirrorZLabel", "Z"))
+			.ToolTip(LOCTEXT("MirrorZLabelTooltip", "Z: copies weights across the XY plane."))
+			;
 	}
 
 	TSharedRef<SWidget> FVertexAttributePaintToolDetailCustomization::MakeMirrorShowAndAxisWidget() const
