Dataflow : split DataflowVertexAttributeEditableNode into a separate file

#Synced-CL 50953304

diff --git a/Engine/Plugins/ChaosClothAssetDataflowNodes/Source/ChaosClothAssetDataflowNodes/Public/ChaosClothAsset/WeightMapNode.h b/Engine/Plugins/ChaosClothAssetDataflowNodes/Source/ChaosClothAssetDataflowNodes/Public/ChaosClothAsset/WeightMapNode.h
--- a/Engine/Plugins/ChaosClothAssetDataflowNodes/Source/ChaosClothAssetDataflowNodes/Public/ChaosClothAsset/WeightMapNode.h
+++ b/Engine/Plugins/ChaosClothAssetDataflowNodes/Source/ChaosClothAssetDataflowNodes/Public/ChaosClothAsset/WeightMapNode.h
@@ -6,7 +6,7 @@
 #include "Dataflow/DataflowFunctionProperty.h"
 #include "GeometryCollection/ManagedArrayCollection.h"
 #include "ChaosClothAsset/ConnectableValue.h"
-#include "Dataflow/DataflowCollectionAddScalarVertexPropertyNode.h"
+#include "Dataflow/DataflowVertexAttributeEditableNode.h"
 
 #include "WeightMapNode.generated.h"
 
@@ -24,8 +24,6 @@ enum class EChaosClothAssetWeightMapMeshTarget : uint8
 	Render
 };
 
-
-
 /** Painted weight map attributes node. */
 PRAGMA_DISABLE_DEPRECATION_WARNINGS  // For Name implicit operators.
 USTRUCT(Meta = (DataflowCloth))
diff --git a/Engine/Plugins/Experimental/Dataflow/Source/DataflowNodes/Private/Dataflow/DataflowCollectionAddScalarVertexPropertyNode.cpp b/Engine/Plugins/Experimental/Dataflow/Source/DataflowNodes/Private/Dataflow/DataflowCollectionAddScalarVertexPropertyNode.cpp
--- a/Engine/Plugins/Experimental/Dataflow/Source/DataflowNodes/Private/Dataflow/DataflowCollectionAddScalarVertexPropertyNode.cpp
+++ b/Engine/Plugins/Experimental/Dataflow/Source/DataflowNodes/Private/Dataflow/DataflowCollectionAddScalarVertexPropertyNode.cpp
@@ -88,74 +88,6 @@ TArray<UE::Dataflow::FRenderingParameter> FDataflowAddScalarVertexPropertyCallba
 }
 
 
-//
-// Deprecated class
-//
-
-PRAGMA_DISABLE_DEPRECATION_WARNINGS
-DataflowAddScalarVertexPropertyCallbackRegistry& DataflowAddScalarVertexPropertyCallbackRegistry::Get()
-{
-	return TLazySingleton<DataflowAddScalarVertexPropertyCallbackRegistry>::Get();
-}
-void DataflowAddScalarVertexPropertyCallbackRegistry::TearDown()
-{
-	TLazySingleton<DataflowAddScalarVertexPropertyCallbackRegistry>::TearDown();
-}
-void DataflowAddScalarVertexPropertyCallbackRegistry::RegisterCallbacks(TUniquePtr<IDataflowAddScalarVertexPropertyCallbacks>&& Callbacks)
-{
-	AllCallbacks.Add(Callbacks->GetName(), MoveTemp(Callbacks));
-}
-void DataflowAddScalarVertexPropertyCallbackRegistry::DeregisterCallbacks(const FName& CallbacksName)
-{
-	AllCallbacks.Remove(CallbacksName);
-}
-TArray<FName> DataflowAddScalarVertexPropertyCallbackRegistry::GetTargetGroupNames() const
-{
-	TArray<FName> UniqueNames;
-
-	for (const TPair<FName, TUniquePtr<IDataflowAddScalarVertexPropertyCallbacks>>& CallbacksEntry : AllCallbacks)
-	{
-		for (const FName& GroupName : CallbacksEntry.Value->GetTargetGroupNames())
-		{
-			UniqueNames.AddUnique(GroupName);
-		}
-	}
-	return UniqueNames;
-}
-TArray<UE::Dataflow::FRenderingParameter> DataflowAddScalarVertexPropertyCallbackRegistry::GetRenderingParameters() const
-{
-	TArray<UE::Dataflow::FRenderingParameter> UniqueParameters;
-
-	for (const TPair<FName, TUniquePtr<IDataflowAddScalarVertexPropertyCallbacks>>& CallbacksEntry : AllCallbacks)
-	{
-		for (const UE::Dataflow::FRenderingParameter& RenderingParameter : CallbacksEntry.Value->GetRenderingParameters())
-		{
-			UniqueParameters.AddUnique(RenderingParameter);
-		}
-	}
-	return UniqueParameters;
-}
-PRAGMA_ENABLE_DEPRECATION_WARNINGS
-
-////////////////////////////////////////////////////////////////////////////////////////////////////////////////////////////////////////////////////////////
-
-FDataflowVertexAttributeEditableNode::FDataflowVertexAttributeEditableNode(const UE::Dataflow::FNodeParameters& InParam, FGuid InGuid)
-	: FDataflowNode(InParam, InGuid)
-{}
-
-const TArray<float>& FDataflowVertexAttributeEditableNode::GetStoredAttributeValues() const 
-{ 
-	static const TArray<float> ConstEmptyArray;
-	ensure(false);
-	return ConstEmptyArray;
-}
-
-void FDataflowVertexAttributeEditableNode::GetExtraVertexMapping(UE::Dataflow::FContext& Context, FName SelectedViewMode, TArray<int32>& OutMappingToWeight, TArray<TArray<int32>>& OutMappingFromWeight) const
-{
-	OutMappingToWeight.Reset();
-	OutMappingFromWeight.Reset();
-}
-
 ////////////////////////////////////////////////////////////////////////////////////////////////////////////////////////////////////////////////////////////
 // 
 // Object encapsulating a change to the vertex attribute node's values. Used for Undo/Redo.
diff --git a/Engine/Plugins/Experimental/Dataflow/Source/DataflowNodes/Private/Dataflow/DataflowVertexAttributeEditableNode.cpp b/Engine/Plugins/Experimental/Dataflow/Source/DataflowNodes/Private/Dataflow/DataflowVertexAttributeEditableNode.cpp
new file mode 100644
--- /dev/null
+++ b/Engine/Plugins/Experimental/Dataflow/Source/DataflowNodes/Private/Dataflow/DataflowVertexAttributeEditableNode.cpp
@@ -0,0 +1,28 @@
+// Copyright Epic Games, Inc. All Rights Reserved.
+
+#include "Dataflow/DataflowVertexAttributeEditableNode.h"
+
+#include "Dataflow/DataflowInputOutput.h"
+#include "Dataflow/DataflowTools.h"
+#include "InteractiveToolChange.h"
+#include "Misc/LazySingleton.h"
+
+#include UE_INLINE_GENERATED_CPP_BY_NAME(DataflowVertexAttributeEditableNode)
+
+FDataflowVertexAttributeEditableNode::FDataflowVertexAttributeEditableNode(const UE::Dataflow::FNodeParameters& InParam, FGuid InGuid)
+	: FDataflowNode(InParam, InGuid)
+{}
+
+const TArray<float>& FDataflowVertexAttributeEditableNode::GetStoredAttributeValues() const 
+{ 
+	static const TArray<float> ConstEmptyArray;
+	ensure(false);
+	return ConstEmptyArray;
+}
+
+void FDataflowVertexAttributeEditableNode::GetExtraVertexMapping(UE::Dataflow::FContext& Context, FName SelectedViewMode, TArray<int32>& OutMappingToWeight, TArray<TArray<int32>>& OutMappingFromWeight) const
+{
+	OutMappingToWeight.Reset();
+	OutMappingFromWeight.Reset();
+}
+
diff --git a/Engine/Plugins/Experimental/Dataflow/Source/DataflowNodes/Public/Dataflow/DataflowCollectionAddScalarVertexPropertyNode.h b/Engine/Plugins/Experimental/Dataflow/Source/DataflowNodes/Public/Dataflow/DataflowCollectionAddScalarVertexPropertyNode.h
--- a/Engine/Plugins/Experimental/Dataflow/Source/DataflowNodes/Public/Dataflow/DataflowCollectionAddScalarVertexPropertyNode.h
+++ b/Engine/Plugins/Experimental/Dataflow/Source/DataflowNodes/Public/Dataflow/DataflowCollectionAddScalarVertexPropertyNode.h
@@ -3,6 +3,7 @@
 #pragma once 
 
 #include "Dataflow/DataflowNode.h"
+#include "Dataflow/DataflowVertexAttributeEditableNode.h"
 #include "GeometryCollection/GeometryCollection.h"
 #include "GeometryCollection/ManagedArrayCollection.h"
 #include "Dataflow/DataflowCollectionAttributeKeyNodes.h"
@@ -42,19 +43,6 @@ private:
 	TMap<FName, TUniquePtr<IDataflowAddScalarVertexPropertyCallbacks>> AllCallbacks;
 };
 
-class UE_DEPRECATED(5.6, "Please use FDataflowAddScalarVertexPropertyCallbackRegistry") DataflowAddScalarVertexPropertyCallbackRegistry
-{
-public:
-	DATAFLOWNODES_API static DataflowAddScalarVertexPropertyCallbackRegistry& Get();
-	DATAFLOWNODES_API static void TearDown();
-	DATAFLOWNODES_API void RegisterCallbacks(TUniquePtr<IDataflowAddScalarVertexPropertyCallbacks> && Callbacks);
-	DATAFLOWNODES_API void DeregisterCallbacks(const FName & CallbacksName);
-	DATAFLOWNODES_API TArray<FName> GetTargetGroupNames() const;
-	DATAFLOWNODES_API TArray<UE::Dataflow::FRenderingParameter> GetRenderingParameters() const;
-private:
-	TMap<FName, TUniquePtr<IDataflowAddScalarVertexPropertyCallbacks>> AllCallbacks;
-};
-
 /*
 * Custom type so that we can use property type customization
 */
@@ -80,36 +68,6 @@ enum class EDataflowWeightMapOverrideType : uint8
 	ReplaceChanged,
 };
 
-USTRUCT()
-struct FDataflowVertexAttributeEditableNode : public FDataflowNode
-{
-	GENERATED_USTRUCT_BODY()
-
-	FDataflowVertexAttributeEditableNode() {};
-
-	UE_API FDataflowVertexAttributeEditableNode(const UE::Dataflow::FNodeParameters& InParam, FGuid InGuid = FGuid::NewGuid());
-
-public:
-	class FEditNodeToolChange;
-	UE_API TUniquePtr<class FToolCommandChange> MakeEditNodeToolChange();
-
-	/** Get the list of view mode supported by this node when being edited */
-	virtual void GetSupportedViewModes(UE::Dataflow::FContext& Context, TArray<FName>& OutViewModeNames) const { ensure(false); }
-
-	/** Get the absolute vertex attribute values */
-	virtual void GetVertexAttributeValues(UE::Dataflow::FContext& Context, TArray<float>& OutValues) const { ensure(false); }
-
-	/** Set the absolute vertex attribute values - InWeightIndices can be used to sepcify which indices are to be set */
-	virtual void SetVertexAttributeValues(UE::Dataflow::FContext& Context, const TArray<float>& InValues, const TArray<int32>& InWeightIndices) { ensure(false); }
-
-	/** Provide a extra indirection vertex mapping - optional will return empty by default */
-	UE_API virtual void GetExtraVertexMapping(UE::Dataflow::FContext& Context, FName SelectedViewMode, TArray<int32>& OutMappingToWeight, TArray<TArray<int32>>& OutMappingFromWeight) const;
-
-private:
-	UE_API virtual const TArray<float>& GetStoredAttributeValues() const;
-	virtual void SwapStoredAttributeValuesWith(TArray<float>& OtherValues) { ensure(false); }
-};
-
 /** Scalar vertex properties. */
 USTRUCT(Meta = (DataflowCollection))
 struct FDataflowCollectionAddScalarVertexPropertyNode : public FDataflowVertexAttributeEditableNode
diff --git a/Engine/Plugins/Experimental/Dataflow/Source/DataflowNodes/Public/Dataflow/DataflowVertexAttributeEditableNode.h b/Engine/Plugins/Experimental/Dataflow/Source/DataflowNodes/Public/Dataflow/DataflowVertexAttributeEditableNode.h
new file mode 100644
--- /dev/null
+++ b/Engine/Plugins/Experimental/Dataflow/Source/DataflowNodes/Public/Dataflow/DataflowVertexAttributeEditableNode.h
@@ -0,0 +1,44 @@
+// Copyright Epic Games, Inc. All Rights Reserved.
+
+#pragma once 
+
+#include "Dataflow/DataflowNode.h"
+
+#include "DataflowVertexAttributeEditableNode.generated.h"
+
+#define UE_API DATAFLOWNODES_API
+
+/**
+* Base class for node using the vertex attribute paint tool
+*/
+USTRUCT()
+struct FDataflowVertexAttributeEditableNode : public FDataflowNode
+{
+	GENERATED_USTRUCT_BODY()
+
+	FDataflowVertexAttributeEditableNode() {};
+
+	UE_API FDataflowVertexAttributeEditableNode(const UE::Dataflow::FNodeParameters& InParam, FGuid InGuid = FGuid::NewGuid());
+
+public:
+	class FEditNodeToolChange;
+	UE_API TUniquePtr<class FToolCommandChange> MakeEditNodeToolChange();
+
+	/** Get the list of view mode supported by this node when being edited */
+	virtual void GetSupportedViewModes(UE::Dataflow::FContext& Context, TArray<FName>& OutViewModeNames) const { ensure(false); }
+
+	/** Get the absolute vertex attribute values */
+	virtual void GetVertexAttributeValues(UE::Dataflow::FContext& Context, TArray<float>& OutValues) const { ensure(false); }
+
+	/** Set the absolute vertex attribute values - InWeightIndices can be used to sepcify which indices are to be set */
+	virtual void SetVertexAttributeValues(UE::Dataflow::FContext& Context, const TArray<float>& InValues, const TArray<int32>& InWeightIndices) { ensure(false); }
+
+	/** Provide a extra indirection vertex mapping - optional will return empty by default */
+	UE_API virtual void GetExtraVertexMapping(UE::Dataflow::FContext& Context, FName SelectedViewMode, TArray<int32>& OutMappingToWeight, TArray<TArray<int32>>& OutMappingFromWeight) const;
+
+private:
+	UE_API virtual const TArray<float>& GetStoredAttributeValues() const;
+	virtual void SwapStoredAttributeValuesWith(TArray<float>& OtherValues) { ensure(false); }
+};
+
+#undef UE_API
