Dataflow : make sure we create a transient cache when no cache asset has been assigned by the user

#Synced-CL 50953304

diff --git a/Engine/Plugins/Experimental/Dataflow/Source/DataflowEditor/Private/Dataflow/DataflowSimulationScene.cpp b/Engine/Plugins/Experimental/Dataflow/Source/DataflowEditor/Private/Dataflow/DataflowSimulationScene.cpp
--- a/Engine/Plugins/Experimental/Dataflow/Source/DataflowEditor/Private/Dataflow/DataflowSimulationScene.cpp
+++ b/Engine/Plugins/Experimental/Dataflow/Source/DataflowEditor/Private/Dataflow/DataflowSimulationScene.cpp
@@ -54,6 +54,18 @@ namespace DataflowSimulationScene::Private
 			Notification->SetCompletionState(SNotificationItem::ECompletionState::CS_Fail);
 		}
 	}
+
+	UChaosCacheCollection* GetTransientCacheCollection()
+	{
+		static const FName BaseName = TEXT("TempDataflowCacheCache");
+
+		const UChaosCacheCollection* DefaultCacheAsset = LoadObject<UChaosCacheCollection>(nullptr, TEXT("/Dataflow/CC_DefaultChaosCache.CC_DefaultChaosCache"));
+		if (DefaultCacheAsset)
+		{
+			return Cast<UChaosCacheCollection>(StaticDuplicateObject(DefaultCacheAsset, GetTransientPackage(), BaseName));
+		}
+		return NewObject<UChaosCacheCollection>(GetTransientPackage(), BaseName);
+	}
 }
 //
 // Simulation Scene
@@ -80,8 +92,11 @@ FDataflowSimulationScene::FDataflowSimulationScene(FPreviewScene::ConstructionVa
 		if(const UDataflow* DataflowAsset = GetEditorContent()->GetDataflowAsset())
 		{
 			SceneDescription->CacheParams = DataflowAsset->PreviewCacheParams;
-			SceneDescription->CacheAsset = DataflowAsset->PreviewCacheAsset ? Cast<UChaosCacheCollection>(DataflowAsset->PreviewCacheAsset.LoadSynchronous()) :
-					LoadObject<UChaosCacheCollection>(nullptr, TEXT("/Dataflow/CC_DefaultChaosCache.CC_DefaultChaosCache"));
+			SceneDescription->CacheAsset = Cast<UChaosCacheCollection>(DataflowAsset->PreviewCacheAsset.LoadSynchronous());
+			if (SceneDescription->CacheAsset == nullptr)
+			{
+				SceneDescription->CacheAsset = DataflowSimulationScene::Private::GetTransientCacheCollection();
+			}
 			SceneDescription->SetBlueprintClass(DataflowAsset->PreviewBlueprintClass);
 			SceneDescription->BlueprintTransform = DataflowAsset->PreviewBlueprintTransform; 
 			SceneDescription->GeometryCacheAsset = Cast<UGeometryCache>(DataflowAsset->PreviewGeometryCacheAsset.LoadSynchronous());
@@ -554,7 +569,11 @@ PRAGMA_ENABLE_DEPRECATION_WARNINGS
 		{
 #if WITH_EDITORONLY_DATA
 			DataflowAsset->PreviewCacheParams = SceneDescription->CacheParams;
-			DataflowAsset->PreviewCacheAsset = SceneDescription->CacheAsset;
+			DataflowAsset->PreviewCacheAsset = {};
+			if (SceneDescription->CacheAsset && SceneDescription->CacheAsset->IsAsset())
+			{
+				DataflowAsset->PreviewCacheAsset = SceneDescription->CacheAsset;
+			}
 			DataflowAsset->PreviewBlueprintClass = SceneDescription->GetBlueprintClass();
 			DataflowAsset->PreviewBlueprintTransform = SceneDescription->BlueprintTransform;
 			DataflowAsset->PreviewGeometryCacheAsset = SceneDescription->GeometryCacheAsset;
