Dataflow editor : add a visual clue on the timeline when the cache is active instead of the simulation

#Synced-CL 50953304

diff --git a/Engine/Plugins/Experimental/Dataflow/Source/DataflowEditor/Private/Widgets/SDataflowSimulationTimeline.cpp b/Engine/Plugins/Experimental/Dataflow/Source/DataflowEditor/Private/Widgets/SDataflowSimulationTimeline.cpp
--- a/Engine/Plugins/Experimental/Dataflow/Source/DataflowEditor/Private/Widgets/SDataflowSimulationTimeline.cpp
+++ b/Engine/Plugins/Experimental/Dataflow/Source/DataflowEditor/Private/Widgets/SDataflowSimulationTimeline.cpp
@@ -109,6 +109,18 @@ void SDataflowSimulationTimeline::Construct(const FArguments& InArgs, const TSha
 	
 	const FMargin ResizeBarPadding(4.f, 0, 0, 0);
 	
+	auto GetCacheStatusBarColor = [WeakSimulationBinding]() -> FLinearColor
+		{
+			bool bCacheActive = false;
+			if (TSharedPtr<FDataflowSimulationBinding> SimulationBinding = WeakSimulationBinding.Pin())
+			{
+				bCacheActive = SimulationBinding->IsSimulationLocked();
+			}
+			return bCacheActive
+				? FLinearColor(0.0f, 4.0f, 0.0f, 1.0f)
+				: FLinearColor(.50f, .50f, .50f, 1.0f);
+		};
+
 	ChildSlot
 		[
 			SNew(SOverlay)
@@ -121,7 +133,22 @@ void SDataflowSimulationTimeline::Construct(const FArguments& InArgs, const TSha
 						+ SGridPanel::Slot(Column1, Row0, SGridPanel::Layer(10))
 						.Padding(ResizeBarPadding)
 						[
-							SNew(SBorder)
+							SNew(SVerticalBox)
+							+ SVerticalBox::Slot()
+							.MinHeight(4)
+							.MaxHeight(4)
+							[
+								SNew(SBorder)
+								.BorderImage(FAppStyle::GetBrush("ToolPanel.GroupBorder"))
+								//.BorderBackgroundColor(FLinearColor(.50f, .50f, .50f, 1.0f))
+									//FLinearColor(0.0f, 4.0f, 0.0f, 1.0f)
+								.BorderBackgroundColor_Lambda(GetCacheStatusBarColor)
+								.Padding(0.f)
+							]
+							+ SVerticalBox::Slot()
+							.AutoHeight()
+							[
+								SNew(SBorder)
 								.BorderImage(FAppStyle::GetBrush("ToolPanel.GroupBorder"))
 								.BorderBackgroundColor(FLinearColor(.50f, .50f, .50f, 1.0f))
 								.Padding(0.f)
@@ -129,6 +156,7 @@ void SDataflowSimulationTimeline::Construct(const FArguments& InArgs, const TSha
 								[
 									TimeSlider.ToSharedRef()
 								]
+							]
 						]
 						+ SGridPanel::Slot(Column1, Row1, SGridPanel::Layer(10))
 						.Padding(ResizeBarPadding)
