Add Dataflow tool node
- this node can be used for node that have a supporting tool and adds support for snapshots

#Synced-CL 50953304

diff --git a/Engine/Plugins/Experimental/Dataflow/Source/DataflowNodes/Private/Dataflow/DataflowToolNode.cpp b/Engine/Plugins/Experimental/Dataflow/Source/DataflowNodes/Private/Dataflow/DataflowToolNode.cpp
new file mode 100644
--- /dev/null
+++ b/Engine/Plugins/Experimental/Dataflow/Source/DataflowNodes/Private/Dataflow/DataflowToolNode.cpp
@@ -0,0 +1,87 @@
+// Copyright Epic Games, Inc. All Rights Reserved.
+
+#include "Dataflow/DataflowToolNode.h"
+
+//////////////////////////////////////////////////////////////////////////////////////////////////////////////////////////
+
+uint64 FDataflowToolNodeSnapshot::GetAllocatedSize() const
+{
+	return Data.GetAllocatedSize();
+}
+
+//////////////////////////////////////////////////////////////////////////////////////////////////////////////////////////
+
+FDataflowToolNode::FDataflowToolNode(const UE::Dataflow::FNodeParameters& InParam, FGuid InGuid)
+	: FDataflowNode(InParam, InGuid)
+{}
+
+FDataflowToolNodeSnapshot& FDataflowToolNode::AddSnapshot()
+{
+	constexpr int32 MaxNumberOfSnapshots = 5;
+	if (Snapshots.Num() >= MaxNumberOfSnapshots)
+	{
+		Snapshots.RemoveAt(0);
+	}
+	FDataflowToolNodeSnapshot& NewSnapshot = Snapshots.AddDefaulted_GetRef();
+	NewSnapshot.Name = TEXT("Tool Snapshot");
+	return NewSnapshot;
+}
+
+bool FDataflowToolNode::RemoveSnapshot(int32 Index)
+{
+	if (Snapshots.IsValidIndex(Index))
+	{
+		Snapshots.RemoveAt(Index);
+		return true;
+	}
+	return false;
+}
+
+void FDataflowToolNode::SwapSnapshots(TArray<FDataflowToolNodeSnapshot>& InOutSnapshots)
+{
+	Swap(InOutSnapshots, Snapshots);
+}
+
+TConstArrayView<FDataflowToolNodeSnapshot> FDataflowToolNode::GetSnapshots() const
+{
+	return Snapshots;
+}
+
+const FDataflowToolNodeSnapshot* FDataflowToolNode::GetLastSnapshot() const
+{
+	return Snapshots.IsEmpty() ? nullptr : &Snapshots.Last();
+}
+
+//////////////////////////////////////////////////////////////////////////////////////////////////////////////////////////
+
+FDataflowToolNode::FSnapshotToolChange::FSnapshotToolChange::FSnapshotToolChange(FDataflowToolNode& Node)
+	: WeakNode(StaticCastWeakPtr<FDataflowToolNode>(Node.AsWeak()))
+	, SavedSnapshots(Node.GetSnapshots())
+{}
+
+FString FDataflowToolNode::FSnapshotToolChange::ToString() const
+{
+	return TEXT("FDataflowToolNode::FSnapshotToolChange");
+}
+
+void FDataflowToolNode::FSnapshotToolChange::Apply(UObject* Object)
+{
+	if (TSharedPtr<FDataflowToolNode> Node = WeakNode.Pin())
+	{
+		SwapApplyRevert(Object, *Node);
+	}
+}
+
+void FDataflowToolNode::FSnapshotToolChange::Revert(UObject* Object)
+{
+	if (TSharedPtr<FDataflowToolNode> Node = WeakNode.Pin())
+	{
+		SwapApplyRevert(Object, *Node);
+	}
+}
+
+void FDataflowToolNode::FSnapshotToolChange::SwapApplyRevert(UObject* Object, FDataflowToolNode& Node)
+{
+	Node.SwapSnapshots(SavedSnapshots);
+	Node.Invalidate();
+}
\ No newline at end of file
diff --git a/Engine/Plugins/Experimental/Dataflow/Source/DataflowNodes/Public/Dataflow/DataflowToolNode.h b/Engine/Plugins/Experimental/Dataflow/Source/DataflowNodes/Public/Dataflow/DataflowToolNode.h
new file mode 100644
--- /dev/null
+++ b/Engine/Plugins/Experimental/Dataflow/Source/DataflowNodes/Public/Dataflow/DataflowToolNode.h
@@ -0,0 +1,85 @@
+// Copyright Epic Games, Inc. All Rights Reserved.
+
+#pragma once
+
+#include "Containers/Array.h"
+#include "Dataflow/DataflowNode.h"
+#include "GeometryCollection/ManagedArrayCollection.h"
+#include "InteractiveToolChange.h"
+#include "UObject/NameTypes.h"
+
+
+#include "DataflowToolNode.generated.h"
+
+#define UE_API DATAFLOWNODES_API
+
+/**
+* Represents data for tool node to store when the user press accept or explictly store a snapshot
+*
+*/
+USTRUCT()
+struct FDataflowToolNodeSnapshot
+{
+	GENERATED_BODY()
+public:
+	/** Returns the approximate amount of memory taken by the snapshot in bytes */
+	uint64 GetAllocatedSize() const;
+
+public:
+	/** Display name of the snapshot */
+	UPROPERTY(EditAnywhere, Category = Snapshots)
+	FString Name;
+
+	UPROPERTY(VisibleAnywhere, Category = Snapshots)
+	FDateTime Date = FDateTime::Now();
+
+	/** collection that stores the data of the snapshot */
+	UPROPERTY()
+	FManagedArrayCollection Data;
+};
+
+/**
+* Base class for a node that can bring up a tool 
+*/
+USTRUCT()
+struct FDataflowToolNode : public FDataflowNode
+{
+	GENERATED_USTRUCT_BODY()
+
+	FDataflowToolNode() = default;
+
+	UE_API FDataflowToolNode(const UE::Dataflow::FNodeParameters& InParam, FGuid InGuid = FGuid::NewGuid());
+
+public:
+	UE_API TConstArrayView<FDataflowToolNodeSnapshot> GetSnapshots() const;
+	UE_API const FDataflowToolNodeSnapshot* GetLastSnapshot() const;
+
+protected:
+	UE_API FDataflowToolNodeSnapshot& AddSnapshot();
+	UE_API bool RemoveSnapshot(int32 Index);
+
+	class FSnapshotToolChange : public FToolCommandChange
+	{
+	public:
+		UE_API FSnapshotToolChange(FDataflowToolNode& Node);
+
+	protected:
+		UE_API virtual FString ToString() const override;
+		UE_API virtual void Apply(UObject* Object) override;
+		UE_API virtual void Revert(UObject* Object) override;
+
+		UE_API virtual void SwapApplyRevert(UObject* Object, FDataflowToolNode& Node);
+
+	private:
+		TWeakPtr<FDataflowToolNode> WeakNode;
+		TArray<FDataflowToolNodeSnapshot> SavedSnapshots;
+	};
+
+private:
+	UE_API void SwapSnapshots(TArray<FDataflowToolNodeSnapshot>& InOutSnapshots);
+
+	UPROPERTY(EditAnywhere, Category = Snapshots, meta = (EditFixedSize, TitleProperty = "{Name} {Date}"))
+	TArray<FDataflowToolNodeSnapshot> Snapshots;
+};
+
+#undef UE_API
