Dataflow : enable new evaluation system by default in Dataflow editors and bind it to the async evaluation option

#Synced-CL 50953304

diff --git a/Engine/Plugins/Experimental/Dataflow/Source/DataflowEditor/Private/Dataflow/DataflowEditorToolkit.cpp b/Engine/Plugins/Experimental/Dataflow/Source/DataflowEditor/Private/Dataflow/DataflowEditorToolkit.cpp
--- a/Engine/Plugins/Experimental/Dataflow/Source/DataflowEditor/Private/Dataflow/DataflowEditorToolkit.cpp
+++ b/Engine/Plugins/Experimental/Dataflow/Source/DataflowEditor/Private/Dataflow/DataflowEditorToolkit.cpp
@@ -91,9 +91,9 @@ bool bDataflowForceGraphEditable = false;
 FAutoConsoleVariableRef CVARbDataflowForceGraphEditable(TEXT("p.Dataflow.Editor.ForceGraphEditable"), bDataflowForceGraphEditable,
 	TEXT("when enabled all graph will be editable and the Unlock/lock editing buttons will not be visible.[def:false]"));
 
-bool bDataflowEnableNewEvaluation = false;
+bool bDataflowEnableNewEvaluation = true;
 FAutoConsoleVariableRef CVARDataflowEnableNewEvaluation(TEXT("p.Dataflow.Editor.EnableNewEvaluation"), bDataflowEnableNewEvaluation,
-	TEXT("When enable, use the new evaluation system using evaluators directly [def:false]"));
+	TEXT("When enable, use the new evaluation system using evaluators directly - it is bound to the async evaluation option in the dataflow editor [def:true]"));
 
 const FName FDataflowEditorToolkit::GraphCanvasTabId(TEXT("DataflowEditor_GraphCanvas"));
 const FName FDataflowEditorToolkit::SubGraphCanvasTabId(TEXT("DataflowEditor_SubGraphTab"));
@@ -844,29 +844,31 @@ bool FDataflowEditorToolkit::IsPerfDataEnabled() const
 
 void FDataflowEditorToolkit::ToggleAsyncEvaluation()
 {
-	if (TObjectPtr<UDataflowBaseContent>& EditorContent = GetEditorContent())
-	{
-		if (TSharedPtr<UE::Dataflow::FEngineContext> DataflowContext = EditorContent->GetDataflowContext())
-		{
-			if (DataflowContext->IsThreaded())
-			{
-				DataflowContext->CancelAsyncEvaluation();
-			}
-			DataflowContext->SetThreaded(!DataflowContext->IsThreaded());
-		}
-	}
+	bDataflowEnableNewEvaluation = !bDataflowEnableNewEvaluation;
+	//if (TObjectPtr<UDataflowBaseContent>& EditorContent = GetEditorContent())
+	//{
+	//	if (TSharedPtr<UE::Dataflow::FEngineContext> DataflowContext = EditorContent->GetDataflowContext())
+	//	{
+	//		if (DataflowContext->IsThreaded())
+	//		{
+	//			DataflowContext->CancelAsyncEvaluation();
+	//		}
+	//		DataflowContext->SetThreaded(!DataflowContext->IsThreaded());
+	//	}
+	//}
 }
 
 bool FDataflowEditorToolkit::IsAsyncEvaluationEnabled() const
 {
-	if (const TObjectPtr<UDataflowBaseContent>& EditorContent = GetEditorContent())
-	{
-		if (TSharedPtr<UE::Dataflow::FEngineContext> DataflowContext = EditorContent->GetDataflowContext())
-		{
-			return DataflowContext->IsThreaded();
-		}
-	}
-	return false;
+	//if (const TObjectPtr<UDataflowBaseContent>& EditorContent = GetEditorContent())
+	//{
+	//	if (TSharedPtr<UE::Dataflow::FEngineContext> DataflowContext = EditorContent->GetDataflowContext())
+	//	{
+	//		return DataflowContext->IsThreaded();
+	//	}
+	//}
+	//return false;
+	return bDataflowEnableNewEvaluation;
 }
 
 void FDataflowEditorToolkit::ClearGraphCache()
