Dataflow : add option to add and remove optional pin from the pin contextual menu

#Synced-CL 51104518

diff --git a/Engine/Plugins/Experimental/Dataflow/Source/DataflowEditor/Private/Dataflow/DataflowSchema.cpp b/Engine/Plugins/Experimental/Dataflow/Source/DataflowEditor/Private/Dataflow/DataflowSchema.cpp
--- a/Engine/Plugins/Experimental/Dataflow/Source/DataflowEditor/Private/Dataflow/DataflowSchema.cpp
+++ b/Engine/Plugins/Experimental/Dataflow/Source/DataflowEditor/Private/Dataflow/DataflowSchema.cpp
@@ -167,6 +167,16 @@ void UDataflowSchema::GetContextMenuActions(class UToolMenu* Menu, class UGraphN
 			{
 				FToolMenuSection& Section = Menu->FindOrAddSection("Pin Actions");
 				Section.AddMenuEntry(FGraphEditorCommands::Get().PromoteToVariable);
+
+				if (TSharedPtr<const FDataflowNode> DataflowNode = UDataflowEdNode::GetDataflowNodeFromEdNode(EdNode))
+				{
+					if (DataflowNode->CanAddPin())
+					{
+						Section.AddSeparator(TEXT("PinActionSeperator_OptionPins"));
+						Section.AddMenuEntry(FDataflowEditorCommands::Get().AddOptionPin, LOCTEXT("DataflowContextPinMenu_AddOptionPin", "Add Option Pin"));
+						Section.AddMenuEntry(FDataflowEditorCommands::Get().RemoveOptionPin, LOCTEXT("DataflowContextPinMenu_RemoveOptionPin", "Remove Last Option Pin"));
+					}
+				}
 			}
 		}
 	}
