Dataflow:
- Added rendering to Samplers

#Synced-CL 51232910

diff --git a/Engine/Plugins/Experimental/Dataflow/Source/DataflowEditor/Private/Dataflow/DataflowEditorModule.cpp b/Engine/Plugins/Experimental/Dataflow/Source/DataflowEditor/Private/Dataflow/DataflowEditorModule.cpp
--- a/Engine/Plugins/Experimental/Dataflow/Source/DataflowEditor/Private/Dataflow/DataflowEditorModule.cpp
+++ b/Engine/Plugins/Experimental/Dataflow/Source/DataflowEditor/Private/Dataflow/DataflowEditorModule.cpp
@@ -47,6 +47,7 @@
 #include "DataflowRendering/DataflowSphereRenderableType.h"
 #include "DataflowRendering/DataflowMaterialRenderableType.h"
 #include "DataflowRendering/DataflowTextureRenderableType.h"
+#include "DataflowRendering/DataflowSamplerRenderableType.h"
 
 #include "ContentBrowserModule.h"
 #include "Modules/ModuleManager.h"
@@ -344,6 +345,7 @@ void FDataflowEditorModule::StartupModule()
 	UE::Dataflow::Private::RegisterSphereRenderableTypes();
 	UE::Dataflow::Private::RegisterMaterialRenderableTypes();
 	UE::Dataflow::Private::RegisterTextureRenderableTypes();
+	UE::Dataflow::Private::RegisterSamplerRenderableTypes();
 
 	// register an browser extension to show when an asset is bound to a dataflow asset
 	if (FContentBrowserModule* ContentBrowserModule = FModuleManager::Get().GetModulePtr<FContentBrowserModule>(TEXT("ContentBrowser")))
diff --git a/Engine/Plugins/Experimental/Dataflow/Source/DataflowEditor/Private/DataflowRendering/DataflowSamplerRenderableType.cpp b/Engine/Plugins/Experimental/Dataflow/Source/DataflowEditor/Private/DataflowRendering/DataflowSamplerRenderableType.cpp
new file mode 100644
--- /dev/null
+++ b/Engine/Plugins/Experimental/Dataflow/Source/DataflowEditor/Private/DataflowRendering/DataflowSamplerRenderableType.cpp
@@ -0,0 +1,313 @@
+// Copyright Epic Games, Inc. All Rights Reserved.
+
+#include "DataflowRendering/DataflowSamplerRenderableType.h"
+
+#include "Drawing/PointSetComponent.h"
+#include "Components/InstancedStaticMeshComponent.h"
+#include "Drawing/LineSetComponent.h"
+#include "Components/BoxComponent.h"
+
+#include "Dataflow/DataflowRenderingViewMode.h"
+
+#include "DataflowRendering/DataflowRenderableComponents.h"
+#include "DataflowRendering/DataflowRenderableType.h"
+#include "DataflowRendering/DataflowRenderableTypeInstance.h"
+#include "DataflowRendering/DataflowRenderableTypeRegistry.h"
+#include "DataflowRendering/DataflowRenderableTypeUtils.h"
+
+#include "Materials/MaterialInterface.h"
+#include "Materials/Material.h"
+
+#include "GeometryCollection/Facades/PointsFacade.h"
+
+#include "UObject/ObjectPtr.h"
+#include "Dataflow/SamplerNodes/DataflowSamplerTypes.h"
+
+UDataflowFloatSamplerRenderSettings::UDataflowFloatSamplerRenderSettings(const FObjectInitializer& ObjectInitializer) : Super(ObjectInitializer)
+{
+	constexpr bool bOnlyRGB = true;
+	ColorRamp.SetColorAtTime(0.0f, FLinearColor::Green, bOnlyRGB);
+	ColorRamp.SetColorAtTime(1.0f, FLinearColor::Blue, bOnlyRGB);
+}
+
+namespace UE::Dataflow::Private
+{
+	static TArray<FVector> GetPointsOnSlice(const FVector& InCenter, const FVector& InExtent, const EDataflowSlicePlane& InPlane, float InDelta, const float InOffset)
+	{
+		TArray<FVector> OutPoints;
+
+		FVector NewExtent;
+
+		if (InDelta < SMALL_NUMBER)
+		{
+			InDelta = 0.01f;
+		}
+
+		NewExtent.X = (FMath::TruncToFloat(InExtent.X / InDelta) + 1.f) * InDelta;
+		NewExtent.Y = (FMath::TruncToFloat(InExtent.Y / InDelta) + 1.f) * InDelta;
+		NewExtent.Z = (FMath::TruncToFloat(InExtent.Z / InDelta) + 1.f) * InDelta;
+
+		FVector Min = InCenter - 0.5 * NewExtent;
+		FVector Max = InCenter + 0.5 * NewExtent;
+
+		const int32 NumX = FMath::Min(1000, int32(NewExtent.X / InDelta));
+		const int32 NumY = FMath::Min(1000, int32(NewExtent.Y / InDelta));
+		const int32 NumZ = FMath::Min(1000, int32(NewExtent.Z / InDelta));
+
+		float X, Y, Z;
+		if (InPlane == EDataflowSlicePlane::XYPlane)
+		{
+			for (int32 IdxX = 0; IdxX < NumX; ++IdxX)
+			{
+				for (int32 IdxY = 0; IdxY < NumY; ++IdxY)
+				{
+					X = Min.X + float(IdxX) * InDelta;
+					Y = Min.Y + float(IdxY) * InDelta;
+					Z = InCenter.Z + InOffset * 0.5 * InExtent.Z;
+
+					OutPoints.Add(FVector(X, Y, Z));
+				}
+			}
+		}
+		else if (InPlane == EDataflowSlicePlane::YZPlane)
+		{
+			for (int32 IdxY = 0; IdxY < NumY; ++IdxY)
+			{
+				for (int32 IdxZ = 0; IdxZ < NumZ; ++IdxZ)
+				{
+					X = InCenter.X + InOffset * 0.5 * InExtent.X;
+					Y = Min.Y + float(IdxY) * InDelta;
+					Z = Min.Z + float(IdxZ) * InDelta;
+
+					OutPoints.Add(FVector(X, Y, Z));
+				}
+			}
+		}
+		else if (InPlane == EDataflowSlicePlane::ZXPlane)
+		{
+			for (int32 IdxZ = 0; IdxZ < NumZ; ++IdxZ)
+			{
+				for (int32 IdxX = 0; IdxX < NumX; ++IdxX)
+				{
+					X = Min.X + float(IdxX) * InDelta;
+					Y = InCenter.Y + InOffset * 0.5 * InExtent.Y;
+					Z = Min.Z + float(IdxZ) * InDelta;
+
+					OutPoints.Add(FVector(X, Y, Z));
+				}
+			}
+		}
+
+		return OutPoints;
+	}
+
+	///////////////////////////////////////////////////////////////////////////////////////////////////////////////
+
+	class FFloatSamplerSurfaceRenderableType : public IRenderableType
+	{
+		UE_DATAFLOW_IRENDERABLE_PRIMARY_TYPE(FDataflowFloatSampler, Sampler);
+		UE_DATAFLOW_IRENDERABLE_RENDER_GROUP(Sampler);
+		UE_DATAFLOW_IRENDERABLE_VIEW_MODE(FDataflowConstruction3DViewMode);
+		UE_DATAFLOW_IRENDERABLE_SETTINGS(UDataflowFloatSamplerRenderSettings);
+
+	public:
+		FFloatSamplerSurfaceRenderableType()
+		{
+			constexpr bool bDepthTested = false;
+			Material = (bDepthTested)
+				? LoadObject<UMaterial>(nullptr, TEXT("/MeshModelingToolsetExp/Materials/PointSetComponentMaterial"))
+				: LoadObject<UMaterial>(nullptr, TEXT("/MeshModelingToolsetExp/Materials/PointSetOverlaidComponentMaterial"));
+		}
+
+		virtual void GetPrimitiveComponents(const FRenderableTypeInstance& Instance, FRenderableComponents& OutComponents) const override
+		{
+			const UDataflowFloatSamplerRenderSettings* Settings = Instance.GetTypedRenderSettings<UDataflowFloatSamplerRenderSettings>();
+			if (Settings)
+			{
+				const FDataflowFloatSampler& InSampler = GetSampler(Instance);
+
+				FBox RenderBounds = InSampler.GetRenderBounds();
+				if (Settings->bOverrideBounds)
+				{
+					RenderBounds.Min = Settings->Center - Settings->Extent;
+					RenderBounds.Max = Settings->Center + Settings->Extent;
+				}
+
+				TArray<FVector> Points = GetPointsOnSlice(RenderBounds.GetCenter(), 2.f * RenderBounds.GetExtent(), Settings->Plane, Settings->Delta, Settings->Offset);
+
+				const int32 NumPoints = Points.Num();
+
+				TArray<float> Values; Values.SetNumUninitialized(NumPoints);
+				InSampler.Sample(TArray<FVector3f>(Points), Values);
+
+				const FName PointComponentName = Instance.GetComponentName(TEXT("FloatSampler"));
+
+				TArray<FRenderablePoint> RenderablePoints;
+				RenderablePoints.SetNumUninitialized(NumPoints);
+
+				for (int32 Idx = 0; Idx < NumPoints; ++Idx)
+				{
+					const float AttrValue = FMath::Clamp(Values[Idx], Settings->Min, Settings->Max);
+
+					FLinearColor LinearColor = FLinearColor::Green;
+					FColor Color = LinearColor.ToFColor(true).WithAlpha(255);
+
+					if (Settings->Max - Settings->Min > UE_SMALL_NUMBER)
+					{
+						float NormAttrValue = (AttrValue - Settings->Min) / (Settings->Max - Settings->Min);
+
+						LinearColor = Settings->ColorRamp.GetLinearColorValue(NormAttrValue);
+						Color = LinearColor.ToFColor(true).WithAlpha(255);
+					}
+
+					RenderablePoints[Idx] = FRenderablePoint(Points[Idx], Color, Settings->Size);
+				}
+
+				// Render box
+				static const FName BoxComponentName = TEXT("RenderBounds");
+
+				if (UBoxComponent* BoxComponent = OutComponents.AddNewComponent<UBoxComponent>(BoxComponentName))
+				{
+					BoxComponent->SetBoxExtent(RenderBounds.GetExtent());
+
+					FTransform Transform = FTransform::Identity;
+					Transform.SetTranslation(RenderBounds.GetCenter());
+
+					BoxComponent->SetWorldTransform(Transform);
+				}
+
+				// Render points
+				if (UPointSetComponent* PointComponent = OutComponents.AddNewComponent<UPointSetComponent>(PointComponentName))
+				{
+					PointComponent->ReservePoints(NumPoints);
+					PointComponent->AddPoints(RenderablePoints);
+
+					PointComponent->SetPointMaterial(Material);
+				}
+			}
+		}
+
+		UMaterialInterface* Material = nullptr;
+	};
+}
+
+///////////////////////////////////////////////////////////////////////////////////////////////////////////////
+
+UDataflowVectorSamplerRenderSettings::UDataflowVectorSamplerRenderSettings(const FObjectInitializer& ObjectInitializer) : Super(ObjectInitializer)
+{
+	constexpr bool bOnlyRGB = true;
+	ColorRamp.SetColorAtTime(0.0f, FLinearColor::Green, bOnlyRGB);
+	ColorRamp.SetColorAtTime(1.0f, FLinearColor::Blue, bOnlyRGB);
+}
+
+namespace UE::Dataflow::Private
+{
+	class FVectorSamplerSurfaceRenderableType : public IRenderableType
+	{
+		UE_DATAFLOW_IRENDERABLE_PRIMARY_TYPE(FDataflowVectorSampler, Sampler);
+		UE_DATAFLOW_IRENDERABLE_RENDER_GROUP(Sampler);
+		UE_DATAFLOW_IRENDERABLE_VIEW_MODE(FDataflowConstruction3DViewMode);
+		UE_DATAFLOW_IRENDERABLE_SETTINGS(UDataflowVectorSamplerRenderSettings);
+
+	public:
+		FVectorSamplerSurfaceRenderableType()
+		{
+			constexpr bool bDepthTested = false;
+			Material = (bDepthTested)
+				? LoadObject<UMaterial>(nullptr, TEXT("/MeshModelingToolsetExp/Materials/PointSetComponentMaterial"))
+				: LoadObject<UMaterial>(nullptr, TEXT("/MeshModelingToolsetExp/Materials/PointSetOverlaidComponentMaterial"));
+		}
+
+		virtual void GetPrimitiveComponents(const FRenderableTypeInstance& Instance, FRenderableComponents& OutComponents) const override
+		{
+			const UDataflowVectorSamplerRenderSettings* Settings = Instance.GetTypedRenderSettings<UDataflowVectorSamplerRenderSettings>();
+			if (Settings)
+			{
+				const FDataflowVectorSampler& InSampler = GetSampler(Instance);
+
+				FBox RenderBounds = InSampler.GetRenderBounds();
+				if (Settings->bOverrideBounds)
+				{
+					RenderBounds.Min = Settings->Center - Settings->Extent;
+					RenderBounds.Max = Settings->Center + Settings->Extent;
+				}
+
+				TArray<FVector> Points = GetPointsOnSlice(RenderBounds.GetCenter(), 2.f * RenderBounds.GetExtent(), Settings->Plane, Settings->Delta, Settings->Offset);
+
+				const int32 NumPoints = Points.Num();
+
+				TArray<FVector3f> Values; Values.SetNumUninitialized(NumPoints);
+				InSampler.Sample(TArray<FVector3f>(Points), Values);
+
+				const FName PointComponentName = Instance.GetComponentName(TEXT("VectorSamplerPoints"));
+				const FName LineComponentName = Instance.GetComponentName(TEXT("VectorSampler"));
+
+				TArray<FRenderablePoint> RenderablePoints;
+				RenderablePoints.SetNumUninitialized(NumPoints);
+
+				TArray<FRenderableLine> RenderableLines;
+				RenderableLines.SetNumUninitialized(NumPoints);
+
+				for (int32 Idx = 0; Idx < NumPoints; ++Idx)
+				{
+					const float AttrValue = FMath::Clamp(Values[Idx].Length(), Settings->Min, Settings->Max);
+
+					FLinearColor LinearColor = FLinearColor::Green;
+					FColor Color = LinearColor.ToFColor(true).WithAlpha(255);
+
+					if (Settings->Max - Settings->Min > UE_SMALL_NUMBER)
+					{
+						float NormAttrValue = (AttrValue - Settings->Min) / (Settings->Max - Settings->Min);
+
+						LinearColor = Settings->ColorRamp.GetLinearColorValue(NormAttrValue);
+						Color = LinearColor.ToFColor(true).WithAlpha(255);
+					}
+
+					RenderablePoints[Idx] = FRenderablePoint(Points[Idx], Color, Settings->Size);
+					RenderableLines[Idx] = FRenderableLine(Points[Idx], Points[Idx] + FVector(Values[Idx]) * Settings->LengthScalar, Color, Settings->LineThickness);
+				}
+				
+				if (ULineSetComponent* LineComponent = OutComponents.AddNewComponent<ULineSetComponent>(LineComponentName))				
+				{
+					LineComponent->AddLines(RenderableLines);
+					LineComponent->SetLineMaterial(Material);
+				}
+
+				// Render box
+				static const FName BoxComponentName = TEXT("RenderBounds");
+
+				if (UBoxComponent* BoxComponent = OutComponents.AddNewComponent<UBoxComponent>(BoxComponentName))
+				{
+					BoxComponent->SetBoxExtent(RenderBounds.GetExtent());
+
+					FTransform Transform = FTransform::Identity;
+					Transform.SetTranslation(RenderBounds.GetCenter());
+
+					BoxComponent->SetWorldTransform(Transform);
+				}
+
+				// Render points
+				if (UPointSetComponent* PointComponent = OutComponents.AddNewComponent<UPointSetComponent>(PointComponentName))
+				{
+					PointComponent->ReservePoints(NumPoints);
+					PointComponent->AddPoints(RenderablePoints);
+
+					PointComponent->SetPointMaterial(Material);
+				}
+			}
+		}
+
+		UMaterialInterface* Material = nullptr;
+	};
+
+	///////////////////////////////////////////////////////////////////////////////////////////////////////////////
+
+	void RegisterSamplerRenderableTypes()
+	{
+		UE_DATAFLOW_REGISTER_RENDERABLE_TYPE(FFloatSamplerSurfaceRenderableType);
+		UE_DATAFLOW_REGISTER_RENDERABLE_TYPE(FVectorSamplerSurfaceRenderableType);
+	}
+
+	///////////////////////////////////////////////////////////////////////////////////////////////////////////////
+
+}
\ No newline at end of file
diff --git a/Engine/Plugins/Experimental/Dataflow/Source/DataflowEditor/Private/DataflowRendering/DataflowSamplerRenderableType.h b/Engine/Plugins/Experimental/Dataflow/Source/DataflowEditor/Private/DataflowRendering/DataflowSamplerRenderableType.h
new file mode 100644
--- /dev/null
+++ b/Engine/Plugins/Experimental/Dataflow/Source/DataflowEditor/Private/DataflowRendering/DataflowSamplerRenderableType.h
@@ -0,0 +1,133 @@
+// Copyright Epic Games, Inc. All Rights Reserved.
+
+#pragma once
+
+#include "DataflowRendering/DataflowRenderableTypeSettings.h"
+#include "DataflowRendering/DataflowRenderableTypeUtils.h"
+#include "Curves/LinearColorRamp.h"
+#include "UObject/ObjectPtr.h"
+
+#include "DataflowSamplerRenderableType.generated.h"
+
+UENUM(BlueprintType)
+enum class EDataflowSlicePlane : uint8
+{
+	XYPlane UMETA(DisplayName = "XY Plane"),
+	YZPlane UMETA(DisplayName = "YZ Plane"),
+	ZXPlane UMETA(DisplayName = "ZX Plane"),
+};
+
+UCLASS(MinimalAPI, AutoExpandCategories = "")
+class UDataflowFloatSamplerRenderSettings : public UDataflowRenderableTypeSettings
+{
+	GENERATED_BODY()
+
+	UDataflowFloatSamplerRenderSettings(const FObjectInitializer& ObjectInitializer = FObjectInitializer::Get());
+
+public:
+	/** Override the sampler's specified render bounds and use the Center end Extent instead */
+	UPROPERTY(EditAnywhere, Category = "Render Bounds")
+	bool bOverrideBounds = false;
+
+	/** Center of the render bounds */
+	UPROPERTY(EditAnywhere, Category = "Render Bounds", meta = (EditCondition = bOverrideBounds))
+	FVector Center = FVector(0.f);
+
+	/** Extent of the render bounds */
+	UPROPERTY(EditAnywhere, Category = "Render Bounds", meta = (UIMin = 0.01f, ClampMin = 0.01f, EditCondition = bOverrideBounds))
+	FVector Extent = FVector(50.0f);
+
+	/** Plane orientation for the visualization */
+	UPROPERTY(EditAnywhere, Category = "Slice Settings")
+	EDataflowSlicePlane Plane = EDataflowSlicePlane::XYPlane;
+
+	/** Distance between the displayed points on the plane */
+	UPROPERTY(EditAnywhere, Category = "Slice Settings", meta = (UIMin = 0.01f, ClampMin = 0.01f))
+	float Delta = 5.f;
+
+	/** Relative offset of the plane inside the render bounds. The offset is beetween -1 and 1 */
+	UPROPERTY(EditAnywhere, Category = "Slice Settings", meta = (ClampMin = "-1", ClampMax = "1"))
+	float Offset = 0.f;
+
+	/** Size of the displayed points */
+	UPROPERTY(EditAnywhere, Category = "Points", meta = (UIMin = 0.01f, ClampMin = 0.01f))
+	float Size = 5.0f;
+
+	/** Attribute minimum value for coloring. This value represents the minimum color on the color ramp */
+	UPROPERTY(EditAnywhere, Category = "Color")
+	float Min = 0.f;
+
+	/** Attribute maximum value for coloring. This value represents the maximum color on the color ramp */
+	UPROPERTY(EditAnywhere, Category = "Color")
+	float Max = 1.f;
+
+	/** Color ramp specifiyng the color for displayed points using a single value attribute */
+	UPROPERTY(EditAnywhere, Category = "Color")
+	FLinearColorRamp ColorRamp;
+};
+
+
+UCLASS(MinimalAPI, AutoExpandCategories = "")
+class UDataflowVectorSamplerRenderSettings : public UDataflowRenderableTypeSettings
+{
+	GENERATED_BODY()
+
+	UDataflowVectorSamplerRenderSettings(const FObjectInitializer& ObjectInitializer = FObjectInitializer::Get());
+
+public:
+	/** Override the sampler's specified render bounds and use the Center end Extent instead */
+	UPROPERTY(EditAnywhere, Category = "Render Bounds")
+	bool bOverrideBounds = false;
+
+	/** Center of the render bounds */
+	UPROPERTY(EditAnywhere, Category = "Render Bounds", meta = (EditCondition = bOverrideBounds))
+	FVector Center = FVector(0.f);
+
+	/** Extent of the render bounds */
+	UPROPERTY(EditAnywhere, Category = "Render Bounds", meta = (UIMin = 0.01f, ClampMin = 0.01f, EditCondition = bOverrideBounds))
+	FVector Extent = FVector(50.0f);
+
+	/** Plane orientation for the visualization */
+	UPROPERTY(EditAnywhere, Category = "Slice Settings")
+	EDataflowSlicePlane Plane = EDataflowSlicePlane::XYPlane;
+
+	/** Distance between the displayed points on the plane */
+	UPROPERTY(EditAnywhere, Category = "Slice Settings", meta = (UIMin = 0.01f, ClampMin = 0.01f))
+	float Delta = 5.f;
+
+	/** Relative offset of the plane inside the render bounds. The offset is beetween -1 and 1 */
+	UPROPERTY(EditAnywhere, Category = "Slice Settings", meta = (ClampMin = "-1", ClampMax = "1"))
+	float Offset = 0.f;
+
+	/** Size of displayed points */
+	UPROPERTY(EditAnywhere, Category = "Points", meta = (UIMin = 0.01f, ClampMin = 0.01f))
+	float Size = 5.0f;
+
+	/** Scalar applied to the vectors */
+	UPROPERTY(EditAnywhere, Category = "Vectors")
+	float LengthScalar = 1.f;
+
+	/** Line thickness of the displayed vectors */
+	UPROPERTY(EditAnywhere, Category = "Vectors")
+	float LineThickness = 1.f;
+
+	/** Attribute minimum value for coloring. This value represents the minimum color on the color ramp */
+	UPROPERTY(EditAnywhere, Category = "Color")
+	float Min = 0.f;
+
+	/** Attribute maximum value for coloring. This value represents the maximum color on the color ramp */
+	UPROPERTY(EditAnywhere, Category = "Color")
+	float Max = 1.f;
+
+	/** Color ramp specifiyng the color for displayed points and vectors using a the length of the vector attribute */
+	UPROPERTY(EditAnywhere, Category = "Color")
+	FLinearColorRamp ColorRamp;
+};
+
+
+
+namespace UE::Dataflow::Private
+{
+	void RegisterSamplerRenderableTypes();
+}
+
