Dataflow : move debug mesh visualization to a component so it works with all dynamic meshes and integrate better with the render settings
- older method remains for now , but will be removed soon

#Synced-CL 51198938

diff --git a/Engine/Plugins/Experimental/Dataflow/Source/DataflowEditor/Private/Dataflow/DataflowConstructionScene.cpp b/Engine/Plugins/Experimental/Dataflow/Source/DataflowEditor/Private/Dataflow/DataflowConstructionScene.cpp
--- a/Engine/Plugins/Experimental/Dataflow/Source/DataflowEditor/Private/Dataflow/DataflowConstructionScene.cpp
+++ b/Engine/Plugins/Experimental/Dataflow/Source/DataflowEditor/Private/Dataflow/DataflowConstructionScene.cpp
@@ -4,6 +4,7 @@
 
 #include "AssetEditorModeManager.h"
 #include "Dataflow/CollectionRenderingPatternUtility.h"
+#include "Dataflow/DataflowDebugDrawComponent.h"
 #include "Dataflow/DataflowEditorCollectionComponent.h"
 #include "Dataflow/DataflowEditor.h"
 #include "Dataflow/DataflowEditorToolkit.h"
@@ -11,8 +12,10 @@
 #include "Dataflow/DataflowEditorStyle.h"
 #include "Dataflow/DataflowRenderingFactory.h"
 #include "Dataflow/DataflowEditorUtil.h"
+#include "DataflowRendering/DataflowDebugMeshRenderableType.h"
 #include "DataflowRendering/DataflowRenderableComponents.h"
 #include "DataflowRendering/DataflowRenderableTypeInstance.h"
+#include "Debug/DebugDrawComponent.h"
 #include "Drawing/MeshElementsVisualizer.h"
 #include "Elements/Framework/EngineElementsLibrary.h"
 #include "Selection.h"
@@ -45,14 +48,39 @@ TArray<TObjectPtr<UDynamicMeshComponent>> FDataflowConstructionScene::GetDynamic
 	return MoveTemp(OutValues);
 }
 
+TArray<TObjectPtr<UDebugDrawComponent>> FDataflowConstructionScene::GetDebugDrawComponents() const
+{
+	TArray<TObjectPtr<UDebugDrawComponent>> DebugDrawComponents;
+	for (TObjectPtr<UPrimitiveComponent> PrimitiveComponent : PrimitiveComponents)
+	{
+		if (UDebugDrawComponent* DebugDrawComponentToAdd = Cast<UDebugDrawComponent>(PrimitiveComponent))
+		{
+			DebugDrawComponents.Add(DebugDrawComponentToAdd);
+		}
+	}
+	if (UDebugDrawComponent* DebugDrawComponentToAdd = Cast<UDebugDrawComponent>(GetDebugDrawComponent()))
+	{
+		DebugDrawComponents.Add(DebugDrawComponentToAdd);
+	}
+	return DebugDrawComponents;
+}
+
 TArray<UDataflowRenderableTypeSettings*> FDataflowConstructionScene::GetRenderableTypeSettings() const
 {
 	TArray<UDataflowRenderableTypeSettings*> RenderSettingsArray;
-	for (const UE::Dataflow::FRenderableTypeInstance& RenderableInstance : RenderableTypeInstances)
+	// make sure we do not add debug mesh render settings if there's no instances
+	if (RenderableTypeInstances.Num() > 0)
 	{
-		if (UDataflowRenderableTypeSettings* RenderSettings = RenderableInstance.GetRenderSettings())
+		for (const UE::Dataflow::FRenderableTypeInstance& RenderableInstance : RenderableTypeInstances)
 		{
-			RenderSettingsArray.AddUnique(RenderSettings);
+			if (UDataflowRenderableTypeSettings* RenderSettings = RenderableInstance.GetRenderSettings())
+			{
+				RenderSettingsArray.AddUnique(RenderSettings);
+			}
+		}
+		if (DebugMeshRenderSettings)
+		{
+			RenderSettingsArray.AddUnique(DebugMeshRenderSettings);
 		}
 	}
 	return RenderSettingsArray;
@@ -95,6 +123,7 @@ void FDataflowConstructionScene::AddReferencedObjects(FReferenceCollector& Colle
 
 	Collector.AddReferencedObjects(DynamicMeshComponents);
 	Collector.AddReferencedObjects(WireframeElements);
+	Collector.AddReferencedObject(DebugMeshRenderSettings);
 }
 
 void FDataflowConstructionScene::TickDataflowScene(const float DeltaSeconds)
@@ -179,6 +208,11 @@ void FDataflowConstructionScene::UpdateDynamicMeshComponents()
 {
 	using namespace UE::Geometry;//FDynamicMesh3
 
+	if (!DebugMeshRenderSettings)
+	{
+		DebugMeshRenderSettings = NewObject<UDataflowDebugMeshRenderSettings>();
+	}
+
 	// The preview scene for the construction view will be
 	// cleared and rebuilt from scratch. This will generate a 
 	// list of UPrimitiveComponents for rendering.
@@ -722,6 +756,12 @@ void FDataflowConstructionScene::UpdateComponentsFromNode(const UDataflowEdNode*
 
 					ProgressNotification.AddProgress(RenderableTypeProgress);
 				}
+
+				if (DebugMeshRenderSettings)
+				{
+					UE::Dataflow::Private::AddDebugComponents(NodeRenderableComponents, *DebugMeshRenderSettings);
+				}
+
 				PrimitiveComponents.Append(NodeRenderableComponents.GetComponents());
 
 				int32 MeshIndex = 0;
diff --git a/Engine/Plugins/Experimental/Dataflow/Source/DataflowEditor/Private/Dataflow/DebugDrawText3dVisualization.cpp b/Engine/Plugins/Experimental/Dataflow/Source/DataflowEditor/Private/Dataflow/DebugDrawText3dVisualization.cpp
--- a/Engine/Plugins/Experimental/Dataflow/Source/DataflowEditor/Private/Dataflow/DebugDrawText3dVisualization.cpp
+++ b/Engine/Plugins/Experimental/Dataflow/Source/DataflowEditor/Private/Dataflow/DebugDrawText3dVisualization.cpp
@@ -36,13 +36,17 @@ namespace UE::Dataflow
 	{
 		if (ConstructionScene)
 		{
-			if (const UDataflowDebugDrawComponent* DebugDrawComponent = ConstructionScene->GetDebugDrawComponent())
+			const TArray<TObjectPtr<UDebugDrawComponent>> DebugDrawComponents = ConstructionScene->GetDebugDrawComponents();
+			for (const TObjectPtr<UDebugDrawComponent>& DebugDrawComponent : DebugDrawComponents)
 			{
-				if (const FDataflowDebugRenderSceneProxy* SceneProxy = static_cast<FDataflowDebugRenderSceneProxy*>(DebugDrawComponent->GetSceneProxy()))
+				if (DebugDrawComponent)
 				{
-					for (const FDebugRenderSceneProxy::FText3d& Text : SceneProxy->Texts)
+					if (const FDebugRenderSceneProxy* SceneProxy = static_cast<FDebugRenderSceneProxy*>(DebugDrawComponent->GetSceneProxy()))
 					{
-						Private::DrawText(Canvas, SceneView, Text.Location, Text.Text, Text.Color);
+						for (const FDebugRenderSceneProxy::FText3d& Text : SceneProxy->Texts)
+						{
+							Private::DrawText(Canvas, SceneView, Text.Location, Text.Text, Text.Color);
+						}
 					}
 				}
 			}
diff --git a/Engine/Plugins/Experimental/Dataflow/Source/DataflowEditor/Private/DataflowRendering/DataflowDebugMeshRenderableType.cpp b/Engine/Plugins/Experimental/Dataflow/Source/DataflowEditor/Private/DataflowRendering/DataflowDebugMeshRenderableType.cpp
new file mode 100644
--- /dev/null
+++ b/Engine/Plugins/Experimental/Dataflow/Source/DataflowEditor/Private/DataflowRendering/DataflowDebugMeshRenderableType.cpp
@@ -0,0 +1,221 @@
+// Copyright Epic Games, Inc. All Rights Reserved.
+
+#include "DataflowRendering/DataflowDebugMeshRenderableType.h"
+
+#include "Components/DynamicMeshComponent.h"
+#include "Dataflow/DataflowDebugDrawComponent.h"
+#include "DataflowRendering/DataflowRenderableComponents.h"
+#include "Debug/DebugDrawComponent.h"
+#include "UObject/ObjectPtr.h"
+
+///////////////////////////////////////////////////////////////////////////////////////////////////////////////////////////////////
+
+UDataflowDebugMeshRenderSettings::UDataflowDebugMeshRenderSettings(const FObjectInitializer& ObjectInitializer)
+{
+	VertexSettings.IDsColor = FColor::Cyan;
+	VertexSettings.NormalsColor = FColor::Blue;
+
+	FaceSettings.IDsColor = FColor::Orange;
+	FaceSettings.NormalsColor = FColor::Red;
+}
+
+///////////////////////////////////////////////////////////////////////////////////////////////////////////////////////////////////
+
+class FDataflowDebugMeshSceneProxy final : public FDebugRenderSceneProxy
+{
+public:
+
+	FDataflowDebugMeshSceneProxy(const UPrimitiveComponent* InComponent)
+		: FDebugRenderSceneProxy(InComponent)
+	{
+	}
+private:
+	virtual FPrimitiveViewRelevance GetViewRelevance(const FSceneView* View) const override
+	{
+		FPrimitiveViewRelevance Result;
+		Result.bDrawRelevance = IsShown(View);
+		Result.bDynamicRelevance = true;
+		// ideally the TranslucencyRelevance should be filled out by the material, here we do it conservative
+		Result.bSeparateTranslucency = Result.bNormalTranslucency = true;
+		return Result;
+	}
+};
+
+///////////////////////////////////////////////////////////////////////////////////////////////////////////////////////////////////
+
+FDebugRenderSceneProxy* UDataflowDebugMeshComponent::CreateDebugSceneProxy()
+{
+	FDataflowDebugMeshSceneProxy* NewProxy = new FDataflowDebugMeshSceneProxy(this);
+	if (NewProxy)
+	{
+		NewProxy->Texts = IDs;
+		NewProxy->Lines = Normals;
+
+		GetDebugDrawDelegateHelper().InitDelegateHelper(NewProxy);
+	}
+	GetDebugDrawDelegateHelper().ProcessDeferredRegister();
+	return NewProxy;
+}
+
+FBoxSphereBounds UDataflowDebugMeshComponent::CalcBounds(const FTransform& LocalToWorld) const
+{
+	FBoxSphereBounds::Builder BoundsBuilder;
+	BoundsBuilder += LocalToWorld.GetLocation();
+
+	for (const UDataflowDebugMeshComponent::FLine& Line : Normals)
+	{
+		BoundsBuilder += LocalToWorld.TransformPosition(Line.Start);
+		BoundsBuilder += LocalToWorld.TransformPosition(Line.End);
+	}
+	for (const UDataflowDebugMeshComponent::FText3d& Text : IDs)
+	{
+		BoundsBuilder += FSphere(LocalToWorld.TransformPosition(Text.Location), 10.0f); // Approximation
+	}
+
+	FBoxSphereBounds ReturnBounds(BoundsBuilder);
+	ReturnBounds = ReturnBounds.ExpandBy(5);
+	return ReturnBounds;
+}
+
+///////////////////////////////////////////////////////////////////////////////////////////////////////////////////////////////////
+
+namespace UE::Dataflow::Private
+{
+	static void UpdateDebugDrawComponentFromDynamicMesh(UDataflowDebugMeshComponent& DebugDrawComponent,  UE::Geometry::FDynamicMesh3& Mesh, const UDataflowDebugMeshRenderSettings& Settings)
+	{
+		if (Settings.VertexSettings.bShowIDs)
+		{
+			DebugDrawComponent.IDs.Reserve(DebugDrawComponent.IDs.Num() + Mesh.VertexCount());
+			for (const int32 VertexId : Mesh.VertexIndicesItr())
+			{
+				if (Mesh.IsVertex(VertexId))
+				{
+					UDataflowDebugMeshComponent::FText3d Text;
+					Text.Color = Settings.VertexSettings.IDsColor;
+					Text.Location = Mesh.GetVertex(VertexId);
+					Text.Text = FString::FromInt(VertexId);
+					DebugDrawComponent.IDs.Add(MoveTemp(Text));
+				}
+			}			
+		}
+		if (Settings.VertexSettings.bShowNormals)
+		{
+			DebugDrawComponent.Normals.Reserve(DebugDrawComponent.Normals.Num() + Mesh.VertexCount());
+			if (Mesh.HasVertexNormals())
+			{
+				for (const int32 VertexId : Mesh.VertexIndicesItr())
+				{
+					if (Mesh.IsVertex(VertexId))
+					{
+						const FVector Position = Mesh.GetVertex(VertexId);
+						UDataflowDebugMeshComponent::FLine Line(
+							/* Start */ Position,
+							/* End   */ Position + FVector(Mesh.GetVertexNormal(VertexId) * Settings.VertexSettings.NormalsLength),
+							/* Color */ Settings.VertexSettings.NormalsColor,
+							/* Thickness */ Settings.VertexSettings.NormalsThickness
+						);
+						DebugDrawComponent.Normals.Add(MoveTemp(Line));
+					}
+				}
+			}
+			else if (Mesh.HasAttributes() && Mesh.Attributes()->PrimaryNormals())
+			{
+				const UE::Geometry::FDynamicMeshNormalOverlay* const NormalOverlay = Mesh.Attributes()->PrimaryNormals();
+
+				TArray<int> OverlayElements;
+				for (const int32 VertexId : Mesh.VertexIndicesItr())
+				{
+					if (Mesh.IsVertex(VertexId))
+					{
+						OverlayElements.Reset();
+						NormalOverlay->GetVertexElements(VertexId, OverlayElements);
+
+						FVector3f AvgNormal(0.0f, 0.0f, 0.0f);
+						if (OverlayElements.Num() > 0)
+						{
+							for (int32 ElementID : OverlayElements)
+							{
+								AvgNormal += NormalOverlay->GetElement(ElementID);
+							}
+							AvgNormal /= (float)OverlayElements.Num();
+						}
+
+						const FVector Position = Mesh.GetVertex(VertexId);
+						UDataflowDebugMeshComponent::FLine Line(
+							/* Start */ Position,
+							/* End   */ Position + FVector(AvgNormal * Settings.VertexSettings.NormalsLength),
+							/* Color */ Settings.VertexSettings.NormalsColor,
+							/* Thickness */ Settings.VertexSettings.NormalsThickness
+						);
+						DebugDrawComponent.Normals.Add(MoveTemp(Line));
+					}
+				}
+			}
+		}
+		if (Settings.FaceSettings.bShowIDs)
+		{
+			DebugDrawComponent.IDs.Reserve(DebugDrawComponent.IDs.Num() + Mesh.VertexCount());
+			for (const int32 TriangleId : Mesh.TriangleIndicesItr())
+			{
+				if (Mesh.IsTriangle(TriangleId))
+				{
+					const UE::Geometry::FIndex3i Triangle = Mesh.GetTriangle(TriangleId);
+					const FVector Vtx0 = Mesh.GetVertex(Triangle[0]);
+					const FVector Vtx1 = Mesh.GetVertex(Triangle[1]);
+					const FVector Vtx2 = Mesh.GetVertex(Triangle[2]);
+					const FVector Center = (Vtx0 + Vtx1 + Vtx2) / 3.0;
+
+					UDataflowDebugMeshComponent::FText3d Text;
+					Text.Color = Settings.FaceSettings.IDsColor;
+					Text.Location = Center;
+					Text.Text = FString::FromInt(TriangleId);
+					DebugDrawComponent.IDs.Add(MoveTemp(Text));
+				}
+			}
+		}
+		if (Settings.FaceSettings.bShowNormals)
+		{
+			DebugDrawComponent.Normals.Reserve(DebugDrawComponent.Normals.Num() + Mesh.VertexCount());
+			for (const int32 TriangleId : Mesh.TriangleIndicesItr())
+			{
+				if (Mesh.IsTriangle(TriangleId))
+				{
+					const UE::Geometry::FIndex3i Triangle = Mesh.GetTriangle(TriangleId);
+					const FVector Vtx0 = Mesh.GetVertex(Triangle[0]);
+					const FVector Vtx1 = Mesh.GetVertex(Triangle[1]);
+					const FVector Vtx2 = Mesh.GetVertex(Triangle[2]);
+					const FVector Center = (Vtx0 + Vtx1 + Vtx2) / 3.0;
+
+					UDataflowDebugMeshComponent::FLine Line(
+						/* Start */ Center,
+						/* End   */ Center + FVector(Mesh.GetTriNormal(TriangleId) * Settings.FaceSettings.NormalsLength),
+						/* Color */ Settings.FaceSettings.NormalsColor,
+						/* Thickness */ Settings.FaceSettings.NormalsThickness
+					);
+					DebugDrawComponent.Normals.Add(MoveTemp(Line));
+				}
+			}
+		}
+	}
+
+	void AddDebugComponents(FRenderableComponents& OutComponents, const UDataflowDebugMeshRenderSettings& Settings)
+	{
+		static FName DebugMeshComponentBaseName = TEXT("DebugInfo");
+
+		const TArray<UPrimitiveComponent*> InputComponents = OutComponents.GetComponents();
+		for (UPrimitiveComponent* InComponent : InputComponents)
+		{
+			if (UDynamicMeshComponent* DynamicMeshComponent = Cast<UDynamicMeshComponent>(InComponent))
+			{
+				if (UE::Geometry::FDynamicMesh3* Mesh = DynamicMeshComponent->GetMesh())
+				{
+					if (UDataflowDebugMeshComponent* DebugDrawComponent = OutComponents.AddNewComponent<UDataflowDebugMeshComponent>(DebugMeshComponentBaseName, DynamicMeshComponent))
+					{
+						UpdateDebugDrawComponentFromDynamicMesh(*DebugDrawComponent, *Mesh, Settings);
+					}
+				}
+			}
+			// TODO : extend this for Static mesh potentially 
+		}
+	}
+}
\ No newline at end of file
diff --git a/Engine/Plugins/Experimental/Dataflow/Source/DataflowEditor/Private/DataflowRendering/DataflowDebugMeshRenderableType.h b/Engine/Plugins/Experimental/Dataflow/Source/DataflowEditor/Private/DataflowRendering/DataflowDebugMeshRenderableType.h
new file mode 100644
--- /dev/null
+++ b/Engine/Plugins/Experimental/Dataflow/Source/DataflowEditor/Private/DataflowRendering/DataflowDebugMeshRenderableType.h
@@ -0,0 +1,82 @@
+// Copyright Epic Games, Inc. All Rights Reserved.
+
+#pragma once
+
+#include "DataflowRendering/DataflowRenderableTypeSettings.h"
+#include "Debug/DebugDrawComponent.h"
+#include "Curves/LinearColorRamp.h"
+#include "UObject/ObjectPtr.h"
+
+#include "DataflowDebugMeshRenderableType.generated.h"
+
+USTRUCT()
+struct FDataflowDebugMeshRenderElementSettings
+{
+	GENERATED_BODY()
+
+	/** Show Ids */
+	UPROPERTY(EditAnywhere, Category = "IDs", meta = (InlineEditConditionToggle))
+	bool bShowIDs = false;
+
+	/** Color of the IDs text */
+	UPROPERTY(EditAnywhere, Category = "IDs", meta = (EditCondition = bShowIDs))
+	FColor IDsColor = FColor::Blue;
+
+	/** Show normals */
+	UPROPERTY(EditAnywhere, Category = "Normals")
+	bool bShowNormals = false;
+
+	/** Length of the normals */
+	UPROPERTY(EditAnywhere, Category = "Normals", meta = (EditCondition = bShowNormals, Units = cm))
+	float NormalsLength = 2.0f;
+
+	/** Thickness of the normals */
+	UPROPERTY(EditAnywhere, Category = "Normals", meta = (EditCondition = bShowNormals))
+	float NormalsThickness = 1.f;
+
+	/** Color of the Normals */
+	UPROPERTY(EditAnywhere, Category = "Normals", meta = (EditCondition = bShowNormals))
+	FColor NormalsColor = FColor::Blue;
+};
+
+UCLASS(MinimalAPI, AutoExpandCategories = "MeshDebug")
+class UDataflowDebugMeshRenderSettings : public UDataflowRenderableTypeSettings
+{
+	GENERATED_BODY()
+
+	UDataflowDebugMeshRenderSettings(const FObjectInitializer& ObjectInitializer = FObjectInitializer::Get());
+
+public:
+	UPROPERTY(EditAnywhere, Category = "MeshDebug")
+	FDataflowDebugMeshRenderElementSettings VertexSettings;
+
+	UPROPERTY(EditAnywhere, Category = "MeshDebug")
+	FDataflowDebugMeshRenderElementSettings FaceSettings;
+};
+
+UCLASS(MinimalAPI)
+class UDataflowDebugMeshComponent : public UDebugDrawComponent
+{
+	GENERATED_BODY()
+
+public:
+	using FText3d = FDebugRenderSceneProxy::FText3d;
+	using FLine = FDebugRenderSceneProxy::FDebugLine;
+
+	TArray<FText3d> IDs;
+	TArray<FLine> Normals;
+
+private:
+	virtual FDebugRenderSceneProxy* CreateDebugSceneProxy() override;
+	virtual FBoxSphereBounds CalcBounds(const FTransform& LocalToWorld) const override;
+};
+
+namespace UE::Dataflow
+{
+	struct FRenderableComponents;
+
+	namespace Private
+	{
+		void AddDebugComponents(FRenderableComponents& OutComponents, const UDataflowDebugMeshRenderSettings& Settings);
+	}
+}
diff --git a/Engine/Plugins/Experimental/Dataflow/Source/DataflowEditor/Public/Dataflow/DataflowConstructionScene.h b/Engine/Plugins/Experimental/Dataflow/Source/DataflowEditor/Public/Dataflow/DataflowConstructionScene.h
--- a/Engine/Plugins/Experimental/Dataflow/Source/DataflowEditor/Public/Dataflow/DataflowConstructionScene.h
+++ b/Engine/Plugins/Experimental/Dataflow/Source/DataflowEditor/Public/Dataflow/DataflowConstructionScene.h
@@ -11,11 +11,13 @@
 
 class UDataflowEditor;
 class USkeletalMeshComponent;
+class UDebugDrawComponent;
 class UDynamicMeshComponent;
 class UMeshElementsVisualizer;
 class UDataflowEdNode;
 class UInteractiveToolPropertySet;
 class UDataflowRenderableTypeSettings;
+class UDataflowDebugMeshRenderSettings;
 
 namespace UE
 {
@@ -51,6 +53,8 @@ class FDataflowConstructionScene : public FDataflowPreviewSceneBase
 
 public:
 
+	inline static FName DebugMeshComponentTag = TEXT("MeshDebugComponent");
+
 	UE_API FDataflowConstructionScene(FPreviewScene::ConstructionValues ConstructionValues, UDataflowEditor* Editor);
 	UE_API virtual ~FDataflowConstructionScene();
 
@@ -78,6 +82,9 @@ public:
 	/** Array of all DynamicMeshComponents in the scene */
 	UE_API TArray<TObjectPtr<UDynamicMeshComponent>> GetDynamicMeshComponents() const;
 
+	/** Get all Debug draw components */
+	UE_API TArray<TObjectPtr<UDebugDrawComponent>> GetDebugDrawComponents() const;
+
 	/** Get the render settings for the selected nodes */
 	UE_API TArray<UObject*> GetRenderSettings() const;
 
@@ -155,6 +162,8 @@ private:
 
 	/** list of renderable type instance, there's one per node/output that can be rendered */
 	TArray<UE::Dataflow::FRenderableTypeInstance> RenderableTypeInstances;
+
+	TObjectPtr<UDataflowDebugMeshRenderSettings> DebugMeshRenderSettings = nullptr;
 };
 
 #undef UE_API
