Dataflow:
- Added more Samplers: Clamp, Negate, Normalize, OneMinus, Remap, ScaleVector

#Synced-CL 51265690

diff --git a/Engine/Plugins/Experimental/Dataflow/Source/DataflowNodes/Private/Dataflow/SamplerNodes/DataflowClampSamplerNode.cpp b/Engine/Plugins/Experimental/Dataflow/Source/DataflowNodes/Private/Dataflow/SamplerNodes/DataflowClampSamplerNode.cpp
new file mode 100644
--- /dev/null
+++ b/Engine/Plugins/Experimental/Dataflow/Source/DataflowNodes/Private/Dataflow/SamplerNodes/DataflowClampSamplerNode.cpp
@@ -0,0 +1,62 @@
+// Copyright Epic Games, Inc. All Rights Reserved.
+
+#include "Dataflow/SamplerNodes/DataflowClampSamplerNode.h"
+#include "Dataflow/SamplerNodes/DataflowSamplerFunctions.h"
+#include "Dataflow/DataflowNodeFactory.h"
+
+FBox FDataflowClampFloatSampler::GetRenderBounds() const
+{
+	return FBox(FVector(-50.0), FVector(50.0));
+}
+
+void FDataflowClampFloatSampler::Sample(TArrayView<const FVector3f> Positions, TArrayView<float> OutValues) const
+{
+	if (Sampler.IsValid() && Positions.Num() == OutValues.Num())
+	{
+		TArray<float> OutValueFromInputSampler;
+		OutValueFromInputSampler.SetNumUninitialized(OutValues.Num());
+	
+		Sampler->Sample(Positions, OutValueFromInputSampler);
+	
+		for (int32 Idx = 0; Idx < OutValues.Num(); ++Idx)
+		{
+			OutValues[Idx] = FMath::Clamp(OutValueFromInputSampler[Idx], Min, Max);
+		}
+	}
+	else
+	{
+		for (int32 Idx = 0; Idx < OutValues.Num(); ++Idx)
+		{
+			OutValues[Idx] = 0.f;
+		}
+	}
+}
+
+FDataflowClampFloatSamplerNode::FDataflowClampFloatSamplerNode(const UE::Dataflow::FNodeParameters& InParam, FGuid InGuid)
+	: FDataflowNode(InParam, InGuid)
+{
+	RegisterInputConnection(&Sampler);
+	RegisterInputConnection(&ClampSampler.Min, GET_MEMBER_NAME_CHECKED(FDataflowClampFloatSampler, Min));
+	RegisterInputConnection(&ClampSampler.Max, GET_MEMBER_NAME_CHECKED(FDataflowClampFloatSampler, Max));
+	RegisterOutputConnection(&Sampler, &Sampler);
+}
+
+void FDataflowClampFloatSamplerNode::Evaluate(UE::Dataflow::FContext& Context, const FDataflowOutput* Out) const
+{
+	if (Out->IsA(&Sampler))
+	{
+		const FDataflowFloatSampler& FloatSampler = GetValue(Context, &Sampler);
+
+		TSharedRef<FDataflowClampFloatSampler> Impl = MakeShared<FDataflowClampFloatSampler>();
+
+		Impl->Sampler = FloatSampler.GetImpl();
+		Impl->Min = GetValue(Context, &ClampSampler.Min);
+		Impl->Max = GetValue(Context, &ClampSampler.Max);
+
+		FDataflowFloatSampler OutSampler(Impl.ToSharedPtr());
+
+		SetValue(Context, OutSampler, &Sampler);
+	}
+}
+
+
diff --git a/Engine/Plugins/Experimental/Dataflow/Source/DataflowNodes/Private/Dataflow/SamplerNodes/DataflowClampSamplerNode.h b/Engine/Plugins/Experimental/Dataflow/Source/DataflowNodes/Private/Dataflow/SamplerNodes/DataflowClampSamplerNode.h
new file mode 100644
--- /dev/null
+++ b/Engine/Plugins/Experimental/Dataflow/Source/DataflowNodes/Private/Dataflow/SamplerNodes/DataflowClampSamplerNode.h
@@ -0,0 +1,52 @@
+// Copyright Epic Games, Inc. All Rights Reserved.
+
+#pragma once 
+
+#include "Dataflow/DataflowNode.h"
+#include "Dataflow/SamplerNodes/DataflowSamplerTypes.h"
+
+#include "DataflowClampSamplerNode.generated.h"
+
+USTRUCT()
+struct FDataflowClampFloatSampler : public FDataflowFloatSamplerBase
+{
+	GENERATED_BODY();
+
+	virtual ~FDataflowClampFloatSampler() override = default;
+	virtual void Sample(TArrayView<const FVector3f> Positions, TArrayView<float> OutValues) const override;
+	virtual FBox GetRenderBounds() const override;
+
+public:
+	/**  */
+	UPROPERTY(EditAnywhere, Category = Clamp, meta = (DataflowInput))
+	float Min = 0.f;
+	
+	/**  */
+	UPROPERTY(EditAnywhere, Category = Clamp, meta = (DataflowInput))
+	float Max = 1.f;
+		
+	/**  */
+	TSharedPtr<const FDataflowFloatSamplerBase> Sampler;
+};
+
+/**
+ * Output an Clamp float sampler
+ */
+USTRUCT()
+struct FDataflowClampFloatSamplerNode : public FDataflowNode
+{
+	GENERATED_BODY()
+	DATAFLOW_NODE_DEFINE_INTERNAL(FDataflowClampFloatSamplerNode, "Clamp Float Sampler", "Samplers", "")
+
+public:
+	FDataflowClampFloatSamplerNode(const UE::Dataflow::FNodeParameters& InParam, FGuid InGuid = FGuid::NewGuid());
+
+private:
+	UPROPERTY(meta = (DataflowInput, DataflowOutput, DataflowPassthrough = Sampler))
+	FDataflowFloatSampler Sampler;
+	
+	UPROPERTY(EditAnywhere, Category = Sampler, meta = (ShowOnlyInnerProperties, SkipInDisplayNameChain))
+	FDataflowClampFloatSampler ClampSampler;
+
+	void Evaluate(UE::Dataflow::FContext& Context, const FDataflowOutput* Out) const override;
+};
diff --git a/Engine/Plugins/Experimental/Dataflow/Source/DataflowNodes/Private/Dataflow/SamplerNodes/DataflowNegateSamplerNode.cpp b/Engine/Plugins/Experimental/Dataflow/Source/DataflowNodes/Private/Dataflow/SamplerNodes/DataflowNegateSamplerNode.cpp
new file mode 100644
--- /dev/null
+++ b/Engine/Plugins/Experimental/Dataflow/Source/DataflowNodes/Private/Dataflow/SamplerNodes/DataflowNegateSamplerNode.cpp
@@ -0,0 +1,109 @@
+// Copyright Epic Games, Inc. All Rights Reserved.
+
+#include "Dataflow/SamplerNodes/DataflowNegateSamplerNode.h"
+#include "Dataflow/SamplerNodes/DataflowSamplerFunctions.h"
+#include "Dataflow/DataflowNodeFactory.h"
+
+FBox FDataflowNegateFloatSampler::GetRenderBounds() const
+{
+	return FBox(FVector(-50.0), FVector(50.0));
+}
+
+void FDataflowNegateFloatSampler::Sample(TArrayView<const FVector3f> Positions, TArrayView<float> OutValues) const
+{
+	if (Sampler.IsValid() && Positions.Num() == OutValues.Num())
+	{
+		TArray<float> OutValueFromInputSampler;
+		OutValueFromInputSampler.SetNumUninitialized(OutValues.Num());
+	
+		Sampler->Sample(Positions, OutValueFromInputSampler);
+	
+		for (int32 Idx = 0; Idx < OutValues.Num(); ++Idx)
+		{
+			OutValues[Idx] = -1.f * OutValueFromInputSampler[Idx];
+		}
+	}
+	else
+	{
+		for (int32 Idx = 0; Idx < OutValues.Num(); ++Idx)
+		{
+			OutValues[Idx] = 0.f;
+		}
+	}
+}
+
+FBox FDataflowNegateVectorSampler::GetRenderBounds() const
+{
+	return FBox(FVector(-50.0), FVector(50.0));
+}
+
+void FDataflowNegateVectorSampler::Sample(TArrayView<const FVector3f> Positions, TArrayView<FVector3f> OutValues) const
+{
+	if (Sampler.IsValid() && Positions.Num() == OutValues.Num())
+	{
+		TArray<FVector3f> OutValueFromInputSampler;
+		OutValueFromInputSampler.SetNumUninitialized(OutValues.Num());
+	
+		Sampler->Sample(Positions, OutValueFromInputSampler);
+	
+		for (int32 Idx = 0; Idx < OutValues.Num(); ++Idx)
+		{
+			OutValues[Idx] = FVector3f(-1.f * OutValueFromInputSampler[Idx]);
+		}
+	}
+	else
+	{
+		for (int32 Idx = 0; Idx < OutValues.Num(); ++Idx)
+		{
+			OutValues[Idx] = FVector3f::ZeroVector;
+		}
+	}
+}
+
+FDataflowNegateSamplerNode::FDataflowNegateSamplerNode(const UE::Dataflow::FNodeParameters& InParam, FGuid InGuid)
+	: FDataflowNode(InParam, InGuid)
+{
+	static const FName MainTypeGroup("Main");
+	
+	RegisterInputConnection(&Sampler)
+		.SetTypeDependencyGroup(MainTypeGroup);
+	RegisterOutputConnection(&Sampler)
+		.SetPassthroughInput(&Sampler)
+		.SetTypeDependencyGroup(MainTypeGroup);
+}
+
+void FDataflowNegateSamplerNode::Evaluate(UE::Dataflow::FContext& Context, const FDataflowOutput* Out) const
+{
+	if (Out->IsA(&Sampler))
+	{
+		const FDataflowFloatSampler* FloatSampler = GetValue(Context, &Sampler).TryGet <FDataflowFloatSampler>();
+		if (FloatSampler)
+		{
+			TSharedRef<FDataflowNegateFloatSampler> Impl = MakeShared<FDataflowNegateFloatSampler>();
+
+			Impl->Sampler = FloatSampler->GetImpl();
+
+			FDataflowFloatSampler OutSampler(Impl.ToSharedPtr());
+
+			SetValue(Context, OutSampler, &Sampler);
+			return;
+		}
+
+		const FDataflowVectorSampler* VectorSampler = GetValue(Context, &Sampler).TryGet<FDataflowVectorSampler>();
+		if (VectorSampler)
+		{
+			TSharedRef<FDataflowNegateVectorSampler> Impl = MakeShared<FDataflowNegateVectorSampler>();
+		
+			Impl->Sampler = VectorSampler->GetImpl();
+		
+			FDataflowVectorSampler OutSampler(Impl.ToSharedPtr());
+
+			SetValue(Context, OutSampler, &Sampler);
+			return;
+		}
+
+		SafeForwardInput(Context, &Sampler, &Sampler);
+	}
+}
+
+
diff --git a/Engine/Plugins/Experimental/Dataflow/Source/DataflowNodes/Private/Dataflow/SamplerNodes/DataflowNegateSamplerNode.h b/Engine/Plugins/Experimental/Dataflow/Source/DataflowNodes/Private/Dataflow/SamplerNodes/DataflowNegateSamplerNode.h
new file mode 100644
--- /dev/null
+++ b/Engine/Plugins/Experimental/Dataflow/Source/DataflowNodes/Private/Dataflow/SamplerNodes/DataflowNegateSamplerNode.h
@@ -0,0 +1,55 @@
+// Copyright Epic Games, Inc. All Rights Reserved.
+
+#pragma once 
+
+#include "Dataflow/DataflowNode.h"
+#include "Dataflow/SamplerNodes/DataflowSamplerTypes.h"
+
+#include "DataflowNegateSamplerNode.generated.h"
+
+USTRUCT()
+struct FDataflowNegateFloatSampler : public FDataflowFloatSamplerBase
+{
+	GENERATED_BODY();
+
+	virtual ~FDataflowNegateFloatSampler() override = default;
+	virtual void Sample(TArrayView<const FVector3f> Positions, TArrayView<float> OutValues) const override;
+	virtual FBox GetRenderBounds() const override;
+
+public:
+	/**  */
+	TSharedPtr<const FDataflowFloatSamplerBase> Sampler;
+};
+
+USTRUCT()
+struct FDataflowNegateVectorSampler : public FDataflowVectorSamplerBase
+{
+	GENERATED_BODY();
+
+	virtual ~FDataflowNegateVectorSampler() override = default;
+	virtual void Sample(TArrayView<const FVector3f> Positions, TArrayView<FVector3f> OutValues) const override;
+	virtual FBox GetRenderBounds() const override;
+
+public:
+	/**  */
+	TSharedPtr<const FDataflowVectorSamplerBase> Sampler;
+};
+
+/**
+ * Output an Negate float sampler
+ */
+USTRUCT()
+struct FDataflowNegateSamplerNode : public FDataflowNode
+{
+	GENERATED_BODY()
+	DATAFLOW_NODE_DEFINE_INTERNAL(FDataflowNegateSamplerNode, "Negate Sampler", "Samplers", "")
+
+public:
+	FDataflowNegateSamplerNode(const UE::Dataflow::FNodeParameters& InParam, FGuid InGuid = FGuid::NewGuid());
+
+private:
+	UPROPERTY(meta = (DataflowInput, DataflowOutput, DataflowPassthrough = Sampler))
+	FDataflowSamplerTypes Sampler;
+
+	void Evaluate(UE::Dataflow::FContext& Context, const FDataflowOutput* Out) const override;
+};
diff --git a/Engine/Plugins/Experimental/Dataflow/Source/DataflowNodes/Private/Dataflow/SamplerNodes/DataflowNormalizeSamplerNode.cpp b/Engine/Plugins/Experimental/Dataflow/Source/DataflowNodes/Private/Dataflow/SamplerNodes/DataflowNormalizeSamplerNode.cpp
new file mode 100644
--- /dev/null
+++ b/Engine/Plugins/Experimental/Dataflow/Source/DataflowNodes/Private/Dataflow/SamplerNodes/DataflowNormalizeSamplerNode.cpp
@@ -0,0 +1,108 @@
+// Copyright Epic Games, Inc. All Rights Reserved.
+
+#include "Dataflow/SamplerNodes/DataflowNormalizeSamplerNode.h"
+#include "Dataflow/SamplerNodes/DataflowSamplerFunctions.h"
+#include "Dataflow/DataflowNodeFactory.h"
+
+FBox FDataflowNormalizeFloatSampler::GetRenderBounds() const
+{
+	return FBox(FVector(-50.0), FVector(50.0));
+}
+
+void FDataflowNormalizeFloatSampler::Sample(TArrayView<const FVector3f> Positions, TArrayView<float> OutValues) const
+{
+	if (Sampler.IsValid() && Positions.Num() == OutValues.Num())
+	{
+		TArray<float> OutValueFromInputSampler;
+		OutValueFromInputSampler.SetNumUninitialized(OutValues.Num());
+	
+		Sampler->Sample(Positions, OutValueFromInputSampler);
+	
+		for (int32 Idx = 0; Idx < OutValues.Num(); ++Idx)
+		{
+			OutValues[Idx] = UE::Dataflow::Sampler::Remap(OutValueFromInputSampler[Idx], Min, Max, 0.f, 1.f, 1.0, false);
+		}
+	}
+	else
+	{
+		for (int32 Idx = 0; Idx < OutValues.Num(); ++Idx)
+		{
+			OutValues[Idx] = 0.f;
+		}
+	}
+}
+
+FBox FDataflowNormalizeVectorSampler::GetRenderBounds() const
+{
+	return FBox(FVector(-50.0), FVector(50.0));
+}
+
+void FDataflowNormalizeVectorSampler::Sample(TArrayView<const FVector3f> Positions, TArrayView<FVector3f> OutValues) const
+{
+	if (Sampler.IsValid() && Positions.Num() == OutValues.Num())
+	{
+		TArray<FVector3f> OutValueFromInputSampler;
+		OutValueFromInputSampler.SetNumUninitialized(OutValues.Num());
+
+		Sampler->Sample(Positions, OutValueFromInputSampler);
+
+		for (int32 Idx = 0; Idx < OutValues.Num(); ++Idx)
+		{
+			OutValueFromInputSampler[Idx].Normalize();
+			OutValues[Idx] = OutValueFromInputSampler[Idx];
+		}
+	}
+	else
+	{
+		for (int32 Idx = 0; Idx < OutValues.Num(); ++Idx)
+		{
+			OutValues[Idx] = FVector3f::ZeroVector;
+		}
+	}
+}
+
+FDataflowNormalizeSamplerNode::FDataflowNormalizeSamplerNode(const UE::Dataflow::FNodeParameters& InParam, FGuid InGuid)
+	: FDataflowNode(InParam, InGuid)
+{
+	static const FName MainTypeGroup("Main");
+
+	RegisterInputConnection(&Sampler)
+		.SetTypeDependencyGroup(MainTypeGroup);
+	RegisterOutputConnection(&Sampler)
+		.SetPassthroughInput(&Sampler)
+		.SetTypeDependencyGroup(MainTypeGroup);
+}
+
+void FDataflowNormalizeSamplerNode::Evaluate(UE::Dataflow::FContext& Context, const FDataflowOutput* Out) const
+{
+	if (Out->IsA(&Sampler))
+	{
+		const FDataflowFloatSampler* FloatSampler = GetValue(Context, &Sampler).TryGet <FDataflowFloatSampler>();
+		if (FloatSampler)
+		{
+			TSharedRef<FDataflowNormalizeFloatSampler> Impl = MakeShared<FDataflowNormalizeFloatSampler>();
+
+			Impl->Sampler = FloatSampler->GetImpl();
+
+			FDataflowFloatSampler OutSampler(Impl.ToSharedPtr());
+
+			SetValue(Context, OutSampler, &Sampler);
+			return;
+		}
+
+		const FDataflowVectorSampler* VectorSampler = GetValue(Context, &Sampler).TryGet<FDataflowVectorSampler>();
+		if (VectorSampler)
+		{
+			TSharedRef<FDataflowNormalizeVectorSampler> Impl = MakeShared<FDataflowNormalizeVectorSampler>();
+
+			Impl->Sampler = VectorSampler->GetImpl();
+
+			FDataflowVectorSampler OutSampler(Impl.ToSharedPtr());
+
+			SetValue(Context, OutSampler, &Sampler);
+			return;
+		}
+
+		SafeForwardInput(Context, &Sampler, &Sampler);
+	}
+}
diff --git a/Engine/Plugins/Experimental/Dataflow/Source/DataflowNodes/Private/Dataflow/SamplerNodes/DataflowNormalizeSamplerNode.h b/Engine/Plugins/Experimental/Dataflow/Source/DataflowNodes/Private/Dataflow/SamplerNodes/DataflowNormalizeSamplerNode.h
new file mode 100644
--- /dev/null
+++ b/Engine/Plugins/Experimental/Dataflow/Source/DataflowNodes/Private/Dataflow/SamplerNodes/DataflowNormalizeSamplerNode.h
@@ -0,0 +1,67 @@
+// Copyright Epic Games, Inc. All Rights Reserved.
+
+#pragma once 
+
+#include "Dataflow/DataflowNode.h"
+#include "Dataflow/SamplerNodes/DataflowSamplerTypes.h"
+
+#include "DataflowNormalizeSamplerNode.generated.h"
+
+/***********************************************************************************************************************
+ * No Input Sampler nodes
+ ***********************************************************************************************************************/
+USTRUCT()
+struct FDataflowNormalizeFloatSampler : public FDataflowFloatSamplerBase
+{
+	GENERATED_BODY();
+
+	virtual ~FDataflowNormalizeFloatSampler() override = default;
+	virtual void Sample(TArrayView<const FVector3f> Positions, TArrayView<float> OutValues) const override;
+	virtual FBox GetRenderBounds() const override;
+
+public:
+	/**  */
+	UPROPERTY(EditAnywhere, Category = Normalize, meta = (DataflowInput))
+	float Min = 0.f;
+
+	/**  */
+	UPROPERTY(EditAnywhere, Category = Normalize, meta = (DataflowInput))
+	float Max = 1.f;
+	
+	/**  */
+	TSharedPtr<const FDataflowFloatSamplerBase> Sampler;
+};
+
+USTRUCT()
+struct FDataflowNormalizeVectorSampler : public FDataflowVectorSamplerBase
+{
+	GENERATED_BODY();
+
+	virtual ~FDataflowNormalizeVectorSampler() override = default;
+	virtual void Sample(TArrayView<const FVector3f> Positions, TArrayView<FVector3f> OutValues) const override;
+	virtual FBox GetRenderBounds() const override;
+
+public:
+	/**  */
+	TSharedPtr<const FDataflowVectorSamplerBase> Sampler;
+};
+
+/**
+ * Output a normalize vector sampler
+ */
+USTRUCT()
+struct FDataflowNormalizeSamplerNode : public FDataflowNode
+{
+	GENERATED_BODY()
+	DATAFLOW_NODE_DEFINE_INTERNAL(FDataflowNormalizeSamplerNode, "Normalize Sampler", "Samplers", "Normal")
+
+public:
+	FDataflowNormalizeSamplerNode(const UE::Dataflow::FNodeParameters& InParam, FGuid InGuid = FGuid::NewGuid());
+
+private:
+	UPROPERTY(meta = (DataflowInput, DataflowOutput, DataflowPassthrough = Sampler))
+	FDataflowSamplerTypes Sampler;
+
+	void Evaluate(UE::Dataflow::FContext& Context, const FDataflowOutput* Out) const override;
+};
+
diff --git a/Engine/Plugins/Experimental/Dataflow/Source/DataflowNodes/Private/Dataflow/SamplerNodes/DataflowOneMinusSamplerNode.cpp b/Engine/Plugins/Experimental/Dataflow/Source/DataflowNodes/Private/Dataflow/SamplerNodes/DataflowOneMinusSamplerNode.cpp
new file mode 100644
--- /dev/null
+++ b/Engine/Plugins/Experimental/Dataflow/Source/DataflowNodes/Private/Dataflow/SamplerNodes/DataflowOneMinusSamplerNode.cpp
@@ -0,0 +1,58 @@
+// Copyright Epic Games, Inc. All Rights Reserved.
+
+#include "Dataflow/SamplerNodes/DataflowOneMinusSamplerNode.h"
+#include "Dataflow/SamplerNodes/DataflowSamplerFunctions.h"
+#include "Dataflow/DataflowNodeFactory.h"
+
+FBox FDataflowOneMinusFloatSampler::GetRenderBounds() const
+{
+	return FBox(FVector(-50.0), FVector(50.0));
+}
+
+void FDataflowOneMinusFloatSampler::Sample(TArrayView<const FVector3f> Positions, TArrayView<float> OutValues) const
+{
+	if (Sampler.IsValid() && Positions.Num() == OutValues.Num())
+	{
+		TArray<float> OutValueFromInputSampler;
+		OutValueFromInputSampler.SetNumUninitialized(OutValues.Num());
+	
+		Sampler->Sample(Positions, OutValueFromInputSampler);
+	
+		for (int32 Idx = 0; Idx < OutValues.Num(); ++Idx)
+		{
+			OutValues[Idx] = 1.f - OutValueFromInputSampler[Idx];
+		}
+	}
+	else
+	{
+		for (int32 Idx = 0; Idx < OutValues.Num(); ++Idx)
+		{
+			OutValues[Idx] = 0.f;
+		}
+	}
+}
+
+FDataflowOneMinusFloatSamplerNode::FDataflowOneMinusFloatSamplerNode(const UE::Dataflow::FNodeParameters& InParam, FGuid InGuid)
+	: FDataflowNode(InParam, InGuid)
+{
+	RegisterInputConnection(&Sampler);
+	RegisterOutputConnection(&Sampler, &Sampler);
+}
+
+void FDataflowOneMinusFloatSamplerNode::Evaluate(UE::Dataflow::FContext& Context, const FDataflowOutput* Out) const
+{
+	if (Out->IsA(&Sampler))
+	{
+		const FDataflowFloatSampler& FloatSampler = GetValue(Context, &Sampler);
+
+		TSharedRef<FDataflowOneMinusFloatSampler> Impl = MakeShared<FDataflowOneMinusFloatSampler>();
+
+		Impl->Sampler = FloatSampler.GetImpl();
+
+		FDataflowFloatSampler OutSampler(Impl.ToSharedPtr());
+
+		SetValue(Context, OutSampler, &Sampler);
+	}
+}
+
+
diff --git a/Engine/Plugins/Experimental/Dataflow/Source/DataflowNodes/Private/Dataflow/SamplerNodes/DataflowOneMinusSamplerNode.h b/Engine/Plugins/Experimental/Dataflow/Source/DataflowNodes/Private/Dataflow/SamplerNodes/DataflowOneMinusSamplerNode.h
new file mode 100644
--- /dev/null
+++ b/Engine/Plugins/Experimental/Dataflow/Source/DataflowNodes/Private/Dataflow/SamplerNodes/DataflowOneMinusSamplerNode.h
@@ -0,0 +1,44 @@
+// Copyright Epic Games, Inc. All Rights Reserved.
+
+#pragma once 
+
+#include "Dataflow/DataflowNode.h"
+#include "Dataflow/SamplerNodes/DataflowSamplerTypes.h"
+
+#include "DataflowOneMinusSamplerNode.generated.h"
+
+USTRUCT()
+struct FDataflowOneMinusFloatSampler : public FDataflowFloatSamplerBase
+{
+	GENERATED_BODY();
+
+	virtual ~FDataflowOneMinusFloatSampler() override = default;
+	virtual void Sample(TArrayView<const FVector3f> Positions, TArrayView<float> OutValues) const override;
+	virtual FBox GetRenderBounds() const override;
+
+public:
+	/**  */
+	TSharedPtr<const FDataflowFloatSamplerBase> Sampler;
+};
+
+/**
+ * Output an OneMinus float sampler
+ */
+USTRUCT()
+struct FDataflowOneMinusFloatSamplerNode : public FDataflowNode
+{
+	GENERATED_BODY()
+	DATAFLOW_NODE_DEFINE_INTERNAL(FDataflowOneMinusFloatSamplerNode, "OneMinus Float Sampler", "Samplers", "-")
+
+public:
+	FDataflowOneMinusFloatSamplerNode(const UE::Dataflow::FNodeParameters& InParam, FGuid InGuid = FGuid::NewGuid());
+
+private:
+	UPROPERTY(meta = (DataflowInput, DataflowOutput, DataflowPassthrough = Sampler))
+	FDataflowFloatSampler Sampler;
+	
+	UPROPERTY(EditAnywhere, Category = Sampler, meta = (ShowOnlyInnerProperties, SkipInDisplayNameChain))
+	FDataflowOneMinusFloatSampler OneMinusSampler;
+
+	void Evaluate(UE::Dataflow::FContext& Context, const FDataflowOutput* Out) const override;
+};
diff --git a/Engine/Plugins/Experimental/Dataflow/Source/DataflowNodes/Private/Dataflow/SamplerNodes/DataflowRemapSamplerNode.cpp b/Engine/Plugins/Experimental/Dataflow/Source/DataflowNodes/Private/Dataflow/SamplerNodes/DataflowRemapSamplerNode.cpp
new file mode 100644
--- /dev/null
+++ b/Engine/Plugins/Experimental/Dataflow/Source/DataflowNodes/Private/Dataflow/SamplerNodes/DataflowRemapSamplerNode.cpp
@@ -0,0 +1,69 @@
+// Copyright Epic Games, Inc. All Rights Reserved.
+
+#include "Dataflow/SamplerNodes/DataflowRemapSamplerNode.h"
+#include "Dataflow/SamplerNodes/DataflowSamplerFunctions.h"
+#include "Dataflow/DataflowNodeFactory.h"
+
+FBox FDataflowRemapFloatSampler::GetRenderBounds() const
+{
+	return FBox(FVector(-50.0), FVector(50.0));
+}
+
+void FDataflowRemapFloatSampler::Sample(TArrayView<const FVector3f> Positions, TArrayView<float> OutValues) const
+{
+	if (Sampler.IsValid() && Positions.Num() == OutValues.Num())
+	{
+		TArray<float> OutValueFromInputSampler;
+		OutValueFromInputSampler.SetNumUninitialized(OutValues.Num());
+	
+		Sampler->Sample(Positions, OutValueFromInputSampler);
+	
+		for (int32 Idx = 0; Idx < OutValues.Num(); ++Idx)
+		{
+			OutValues[Idx] = UE::Dataflow::Sampler::Remap(OutValueFromInputSampler[Idx], OriginalMin, OriginalMax, NewMin, NewMax, Power, bClamp);
+		}
+	}
+	else
+	{
+		for (int32 Idx = 0; Idx < OutValues.Num(); ++Idx)
+		{
+			OutValues[Idx] = 0.f;
+		}
+	}
+}
+
+FDataflowRemapFloatSamplerNode::FDataflowRemapFloatSamplerNode(const UE::Dataflow::FNodeParameters& InParam, FGuid InGuid)
+	: FDataflowNode(InParam, InGuid)
+{
+	RegisterInputConnection(&Sampler);
+	RegisterInputConnection(&RemapSampler.OriginalMin, GET_MEMBER_NAME_CHECKED(FDataflowRemapFloatSampler, OriginalMin));
+	RegisterInputConnection(&RemapSampler.OriginalMax, GET_MEMBER_NAME_CHECKED(FDataflowRemapFloatSampler, OriginalMax));
+	RegisterInputConnection(&RemapSampler.NewMin, GET_MEMBER_NAME_CHECKED(FDataflowRemapFloatSampler, NewMin));
+	RegisterInputConnection(&RemapSampler.NewMax, GET_MEMBER_NAME_CHECKED(FDataflowRemapFloatSampler, NewMax));
+	RegisterInputConnection(&RemapSampler.Power, GET_MEMBER_NAME_CHECKED(FDataflowRemapFloatSampler, Power));
+	RegisterOutputConnection(&Sampler, &Sampler);
+}
+
+void FDataflowRemapFloatSamplerNode::Evaluate(UE::Dataflow::FContext& Context, const FDataflowOutput* Out) const
+{
+	if (Out->IsA(&Sampler))
+	{
+		const FDataflowFloatSampler& FloatSampler = GetValue(Context, &Sampler);
+
+		TSharedRef<FDataflowRemapFloatSampler> Impl = MakeShared<FDataflowRemapFloatSampler>();
+		Impl->OriginalMin = GetValue(Context, &RemapSampler.OriginalMin);
+		Impl->OriginalMax = GetValue(Context, &RemapSampler.OriginalMax);
+		Impl->NewMin = GetValue(Context, &RemapSampler.NewMin);
+		Impl->NewMax = GetValue(Context, &RemapSampler.NewMax);
+		Impl->Power = GetValue(Context, &RemapSampler.Power);
+		Impl->bClamp = GetValue(Context, &RemapSampler.bClamp);
+
+		Impl->Sampler = FloatSampler.GetImpl();
+
+		FDataflowFloatSampler OutSampler(Impl.ToSharedPtr());
+
+		SetValue(Context, OutSampler, &Sampler);
+	}
+}
+
+
diff --git a/Engine/Plugins/Experimental/Dataflow/Source/DataflowNodes/Private/Dataflow/SamplerNodes/DataflowRemapSamplerNode.h b/Engine/Plugins/Experimental/Dataflow/Source/DataflowNodes/Private/Dataflow/SamplerNodes/DataflowRemapSamplerNode.h
new file mode 100644
--- /dev/null
+++ b/Engine/Plugins/Experimental/Dataflow/Source/DataflowNodes/Private/Dataflow/SamplerNodes/DataflowRemapSamplerNode.h
@@ -0,0 +1,69 @@
+// Copyright Epic Games, Inc. All Rights Reserved.
+
+#pragma once 
+
+#include "Dataflow/DataflowNode.h"
+#include "Dataflow/SamplerNodes/DataflowSamplerTypes.h"
+
+#include "DataflowRemapSamplerNode.generated.h"
+
+
+USTRUCT()
+struct FDataflowRemapFloatSampler : public FDataflowFloatSamplerBase
+{
+	GENERATED_BODY();
+
+	virtual ~FDataflowRemapFloatSampler() override = default;
+	virtual void Sample(TArrayView<const FVector3f> Positions, TArrayView<float> OutValues) const override;
+	virtual FBox GetRenderBounds() const override;
+
+public:
+	/**  */
+	UPROPERTY(EditAnywhere, Category = Remap, meta = (DataflowInput))
+	float OriginalMin = 0.f;
+	
+	/**  */
+	UPROPERTY(EditAnywhere, Category = Remap, meta = (DataflowInput))
+	float OriginalMax = 1.f;
+	
+	/**  */
+	UPROPERTY(EditAnywhere, Category = Remap, meta = (DataflowInput))
+	float NewMin = 0.f;
+	
+	/**  */
+	UPROPERTY(EditAnywhere, Category = Remap, meta = (DataflowInput))
+	float NewMax = 2.f;
+	
+	/**  */
+	UPROPERTY(EditAnywhere, Category = Remap, meta = (DataflowInput))
+	float Power = 1.f;
+	
+	/**  */
+	UPROPERTY(EditAnywhere, Category = Remap)
+	bool bClamp = false;
+
+	/**  */
+	TSharedPtr<const FDataflowFloatSamplerBase> Sampler;
+};
+
+/**
+ * Output an Remap float sampler
+ */
+USTRUCT()
+struct FDataflowRemapFloatSamplerNode : public FDataflowNode
+{
+	GENERATED_BODY()
+	DATAFLOW_NODE_DEFINE_INTERNAL(FDataflowRemapFloatSamplerNode, "Remap Float Sampler", "Samplers", "map")
+
+public:
+	FDataflowRemapFloatSamplerNode(const UE::Dataflow::FNodeParameters& InParam, FGuid InGuid = FGuid::NewGuid());
+
+private:
+	UPROPERTY(meta = (DataflowInput, DataflowOutput, DataflowPassthrough = Sampler))
+	FDataflowFloatSampler Sampler;
+
+	UPROPERTY(EditAnywhere, Category = Sampler, meta = (ShowOnlyInnerProperties, SkipInDisplayNameChain))
+	FDataflowRemapFloatSampler RemapSampler;
+
+	void Evaluate(UE::Dataflow::FContext& Context, const FDataflowOutput* Out) const override;
+};
diff --git a/Engine/Plugins/Experimental/Dataflow/Source/DataflowNodes/Private/Dataflow/SamplerNodes/DataflowSamplerTypes.cpp b/Engine/Plugins/Experimental/Dataflow/Source/DataflowNodes/Private/Dataflow/SamplerNodes/DataflowSamplerTypes.cpp
--- a/Engine/Plugins/Experimental/Dataflow/Source/DataflowNodes/Private/Dataflow/SamplerNodes/DataflowSamplerTypes.cpp
+++ b/Engine/Plugins/Experimental/Dataflow/Source/DataflowNodes/Private/Dataflow/SamplerNodes/DataflowSamplerTypes.cpp
@@ -11,11 +11,16 @@
 #include "Dataflow/SamplerNodes/DataflowTurbulenceSamplerNode.h"
 #include "Dataflow/SamplerNodes/DataflowAbsSamplerNode.h"
 #include "Dataflow/SamplerNodes/DataflowSignSamplerNode.h"
-
 #include "Dataflow/SamplerNodes/DataflowAddSamplerNode.h"
 #include "Dataflow/SamplerNodes/DataflowMultiplySamplerNode.h"
 #include "Dataflow/SamplerNodes/DataflowMinSamplerNode.h"
 #include "Dataflow/SamplerNodes/DataflowMaxSamplerNode.h"
+#include "Dataflow/SamplerNodes/DataflowScaleSamplerNode.h"
+#include "Dataflow/SamplerNodes/DataflowClampSamplerNode.h"
+#include "Dataflow/SamplerNodes/DataflowNegateSamplerNode.h"
+#include "Dataflow/SamplerNodes/DataflowNormalizeSamplerNode.h"
+#include "Dataflow/SamplerNodes/DataflowOneMinusSamplerNode.h"
+#include "Dataflow/SamplerNodes/DataflowRemapSamplerNode.h"
 
 #include "Dataflow/DataflowNodeFactory.h"
 #include "Dataflow/DataflowCategoryRegistry.h"
@@ -40,29 +45,28 @@ namespace UE::Dataflow
 		DATAFLOW_NODE_REGISTER_CREATION_FACTORY(FDataflowLinearGradientFromBoxFloatSamplerNode);
 		DATAFLOW_NODE_REGISTER_CREATION_FACTORY(FDataflowRadialGradientFloatSamplerNode);
 
-		DATAFLOW_NODE_REGISTER_CREATION_FACTORY(FDataflowTransformFloatSamplerNode);
-
-		DATAFLOW_NODE_REGISTER_CREATION_FACTORY(FDataflowUniformFloatSamplerNode);
-		DATAFLOW_NODE_REGISTER_CREATION_FACTORY(FDataflowUniformVectorSamplerNode);
-
-		DATAFLOW_NODE_REGISTER_CREATION_FACTORY(FDataflowRandomFloatSamplerNode);
-		DATAFLOW_NODE_REGISTER_CREATION_FACTORY(FDataflowRandomVectorSamplerNode);
-
-		DATAFLOW_NODE_REGISTER_CREATION_FACTORY(FDataflowPerlinNoiseFloatSamplerNode);
-		DATAFLOW_NODE_REGISTER_CREATION_FACTORY(FDataflowPerlinNoiseVectorSamplerNode);
-
-		DATAFLOW_NODE_REGISTER_CREATION_FACTORY(FDataflowfBmFloatSamplerNode);
-		DATAFLOW_NODE_REGISTER_CREATION_FACTORY(FDataflowfBmVectorSamplerNode);
-
-		DATAFLOW_NODE_REGISTER_CREATION_FACTORY(FDataflowTurbulenceFloatSamplerNode);
-
 		DATAFLOW_NODE_REGISTER_CREATION_FACTORY(FDataflowAbsSamplerNode);
-		DATAFLOW_NODE_REGISTER_CREATION_FACTORY(FDataflowSignSamplerNode);
-
 		DATAFLOW_NODE_REGISTER_CREATION_FACTORY(FDataflowAddSamplerNode);
-		DATAFLOW_NODE_REGISTER_CREATION_FACTORY(FDataflowMultiplySamplerNode);
-		DATAFLOW_NODE_REGISTER_CREATION_FACTORY(FDataflowMinSamplerNode);
+		DATAFLOW_NODE_REGISTER_CREATION_FACTORY(FDataflowClampFloatSamplerNode);
+		DATAFLOW_NODE_REGISTER_CREATION_FACTORY(FDataflowfBmFloatSamplerNode);
+		DATAFLOW_NODE_REGISTER_CREATION_FACTORY(FDataflowfBmVectorSamplerNode);
 		DATAFLOW_NODE_REGISTER_CREATION_FACTORY(FDataflowMaxSamplerNode);
+		DATAFLOW_NODE_REGISTER_CREATION_FACTORY(FDataflowMinSamplerNode);
+		DATAFLOW_NODE_REGISTER_CREATION_FACTORY(FDataflowMultiplySamplerNode);
+		DATAFLOW_NODE_REGISTER_CREATION_FACTORY(FDataflowNegateSamplerNode);
+		DATAFLOW_NODE_REGISTER_CREATION_FACTORY(FDataflowNormalizeSamplerNode);
+		DATAFLOW_NODE_REGISTER_CREATION_FACTORY(FDataflowOneMinusFloatSamplerNode);
+		DATAFLOW_NODE_REGISTER_CREATION_FACTORY(FDataflowPerlinNoiseFloatSamplerNode);
+		DATAFLOW_NODE_REGISTER_CREATION_FACTORY(FDataflowPerlinNoiseVectorSamplerNode);
+		DATAFLOW_NODE_REGISTER_CREATION_FACTORY(FDataflowRandomFloatSamplerNode);
+		DATAFLOW_NODE_REGISTER_CREATION_FACTORY(FDataflowRandomVectorSamplerNode);
+		DATAFLOW_NODE_REGISTER_CREATION_FACTORY(FDataflowRemapFloatSamplerNode);
+		DATAFLOW_NODE_REGISTER_CREATION_FACTORY(FDataflowScaleVectorSamplerNode);
+		DATAFLOW_NODE_REGISTER_CREATION_FACTORY(FDataflowSignSamplerNode);
+		DATAFLOW_NODE_REGISTER_CREATION_FACTORY(FDataflowTransformFloatSamplerNode);
+		DATAFLOW_NODE_REGISTER_CREATION_FACTORY(FDataflowTurbulenceFloatSamplerNode);
+		DATAFLOW_NODE_REGISTER_CREATION_FACTORY(FDataflowUniformFloatSamplerNode);
+		DATAFLOW_NODE_REGISTER_CREATION_FACTORY(FDataflowUniformVectorSamplerNode);
 
 		DATAFLOW_NODE_REGISTER_CREATION_FACTORY_NODE_COLORS_BY_CATEGORY("Samplers", FLinearColor(1.f, 0.f, 0.93f), CDefaultNodeBodyTintColor);
 		DATAFLOW_NODE_REGISTER_CREATION_FACTORY_PIN_SETTINGS_BY_TYPE("FDataflowFloatSampler", FLinearColor(1.f, 0.3f, 0.95f), CDefaultWireThickness);
diff --git a/Engine/Plugins/Experimental/Dataflow/Source/DataflowNodes/Private/Dataflow/SamplerNodes/DataflowScaleSamplerNode.cpp b/Engine/Plugins/Experimental/Dataflow/Source/DataflowNodes/Private/Dataflow/SamplerNodes/DataflowScaleSamplerNode.cpp
new file mode 100644
--- /dev/null
+++ b/Engine/Plugins/Experimental/Dataflow/Source/DataflowNodes/Private/Dataflow/SamplerNodes/DataflowScaleSamplerNode.cpp
@@ -0,0 +1,64 @@
+// Copyright Epic Games, Inc. All Rights Reserved.
+
+#include "Dataflow/SamplerNodes/DataflowScaleSamplerNode.h"
+#include "Dataflow/SamplerNodes/DataflowSamplerFunctions.h"
+#include "Dataflow/DataflowNodeFactory.h"
+
+FBox FDataflowScaleVectorSampler::GetRenderBounds() const
+{
+	return FBox(FVector(-50.0), FVector(50.0));
+}
+
+void FDataflowScaleVectorSampler::Sample(TArrayView<const FVector3f> Positions, TArrayView<FVector3f> OutValues) const
+{
+	if (VectorSampler.IsValid() && Positions.Num() == OutValues.Num())
+	{
+		VectorSampler->Sample(Positions, OutValues);
+
+		if (FloatSampler.IsValid())
+		{
+			TArray<float> OutValueFromFloatSampler;
+			OutValueFromFloatSampler.SetNumUninitialized(OutValues.Num());
+
+			FloatSampler->Sample(Positions, OutValueFromFloatSampler);
+
+			for (int32 Idx = 0; Idx < OutValues.Num(); ++Idx)
+			{
+				OutValues[Idx] *= OutValueFromFloatSampler[Idx];
+			}
+		}
+	}
+	else
+	{
+		for (int32 Idx = 0; Idx < OutValues.Num(); ++Idx)
+		{
+			OutValues[Idx] = FVector3f::ZeroVector;
+		}
+	}
+}
+
+FDataflowScaleVectorSamplerNode::FDataflowScaleVectorSamplerNode(const UE::Dataflow::FNodeParameters& InParam, FGuid InGuid)
+	: FDataflowNode(InParam, InGuid)
+{
+	RegisterInputConnection(&VectorSampler);
+	RegisterInputConnection(&FloatSampler);
+	RegisterOutputConnection(&VectorSampler, &VectorSampler);
+}
+
+void FDataflowScaleVectorSamplerNode::Evaluate(UE::Dataflow::FContext& Context, const FDataflowOutput* Out) const
+{
+	if (Out->IsA(&VectorSampler))
+	{
+		const FDataflowFloatSampler& InFloatSampler = GetValue(Context, &FloatSampler);
+		const FDataflowVectorSampler& InVectorSampler = GetValue(Context, &VectorSampler);
+
+		TSharedRef<FDataflowScaleVectorSampler> Impl = MakeShared<FDataflowScaleVectorSampler>();
+
+		Impl->FloatSampler = InFloatSampler.GetImpl();
+		Impl->VectorSampler = InVectorSampler.GetImpl();
+
+		FDataflowVectorSampler OutSampler(Impl.ToSharedPtr());
+
+		SetValue(Context, OutSampler, &VectorSampler);
+	}
+}
diff --git a/Engine/Plugins/Experimental/Dataflow/Source/DataflowNodes/Private/Dataflow/SamplerNodes/DataflowScaleSamplerNode.h b/Engine/Plugins/Experimental/Dataflow/Source/DataflowNodes/Private/Dataflow/SamplerNodes/DataflowScaleSamplerNode.h
new file mode 100644
--- /dev/null
+++ b/Engine/Plugins/Experimental/Dataflow/Source/DataflowNodes/Private/Dataflow/SamplerNodes/DataflowScaleSamplerNode.h
@@ -0,0 +1,47 @@
+// Copyright Epic Games, Inc. All Rights Reserved.
+
+#pragma once 
+
+#include "Dataflow/DataflowNode.h"
+#include "Dataflow/SamplerNodes/DataflowSamplerTypes.h"
+
+#include "DataflowScaleSamplerNode.generated.h"
+
+USTRUCT()
+struct FDataflowScaleVectorSampler : public FDataflowVectorSamplerBase
+{
+	GENERATED_BODY();
+
+	virtual ~FDataflowScaleVectorSampler() override = default;
+	virtual void Sample(TArrayView<const FVector3f> Positions, TArrayView<FVector3f> OutValues) const override;
+	virtual FBox GetRenderBounds() const override;
+
+public:
+	/**  */
+	TSharedPtr<const FDataflowFloatSamplerBase> FloatSampler;
+
+	/**  */
+	TSharedPtr<const FDataflowVectorSamplerBase> VectorSampler;
+};
+
+/**
+ * Output a scale vector sampler
+ */
+USTRUCT()
+struct FDataflowScaleVectorSamplerNode : public FDataflowNode
+{
+	GENERATED_BODY()
+	DATAFLOW_NODE_DEFINE_INTERNAL(FDataflowScaleVectorSamplerNode, "Scale Vector Sampler", "Samplers", "* Multiply Magnitude")
+
+public:
+	FDataflowScaleVectorSamplerNode(const UE::Dataflow::FNodeParameters& InParam, FGuid InGuid = FGuid::NewGuid());
+
+private:
+	UPROPERTY(meta = (DataflowInput, DataflowOutput, DataflowPassthrough = VectorSampler))
+	FDataflowVectorSampler VectorSampler;
+
+	UPROPERTY(meta = (DataflowInput))
+	FDataflowFloatSampler FloatSampler;
+
+	void Evaluate(UE::Dataflow::FContext& Context, const FDataflowOutput* Out) const override;
+};
