Dataflow : Add static mesh terminal node
- added utility methods in DataflowTerminalNode to ease development of future

#Synced-CL 51254598

diff --git a/Engine/Plugins/Experimental/Dataflow/Dataflow.uplugin b/Engine/Plugins/Experimental/Dataflow/Dataflow.uplugin
--- a/Engine/Plugins/Experimental/Dataflow/Dataflow.uplugin
+++ b/Engine/Plugins/Experimental/Dataflow/Dataflow.uplugin
@@ -76,6 +76,10 @@
 			"Name": "EditorDataStorageFeatures",
 			"Enabled": true
 		},
+		{
+			"Name": "MeshModelingToolset",
+			"Enabled": true
+		},
 		{
 			"Name": "SkeletalMeshModelingTools",
 			"Enabled": true,
diff --git a/Engine/Plugins/Experimental/Dataflow/Source/DataflowNodes/DataflowNodes.Build.cs b/Engine/Plugins/Experimental/Dataflow/Source/DataflowNodes/DataflowNodes.Build.cs
--- a/Engine/Plugins/Experimental/Dataflow/Source/DataflowNodes/DataflowNodes.Build.cs
+++ b/Engine/Plugins/Experimental/Dataflow/Source/DataflowNodes/DataflowNodes.Build.cs
@@ -29,7 +29,9 @@ namespace UnrealBuildTool.Rules
 				new string[]
 				{
 					"GeometryFramework",
-					"ImageCore"
+					"ImageCore",
+					"ModelingComponents",
+					"StaticMeshDescription",
 				}
 			);
 
diff --git a/Engine/Plugins/Experimental/Dataflow/Source/DataflowNodes/Private/Dataflow/DataflowStaticMeshAssetNodes.cpp b/Engine/Plugins/Experimental/Dataflow/Source/DataflowNodes/Private/Dataflow/DataflowStaticMeshAssetNodes.cpp
new file mode 100644
--- /dev/null
+++ b/Engine/Plugins/Experimental/Dataflow/Source/DataflowNodes/Private/Dataflow/DataflowStaticMeshAssetNodes.cpp
@@ -0,0 +1,175 @@
+// Copyright Epic Games, Inc. All Rights Reserved.
+
+#include "Dataflow/DataflowStaticMeshAssetNodes.h"
+
+#include "AssetUtils/StaticMeshMaterialUtil.h"
+#include "Dataflow/DataflowNodeFactory.h"
+#include "Dataflow/DataflowObjectInterface.h"
+#include "DynamicMeshToMeshDescription.h"
+#include "Engine/StaticMesh.h"
+#include "Materials/MaterialInterface.h"
+#include "MeshConversionOptions.h"
+#include "StaticMeshDescription.h"
+#include "StaticMeshOperations.h"
+#include "UDynamicMesh.h"
+
+#define LOCTEXT_NAMESPACE "DataflowStaticMeshAssetNodes"
+
+namespace UE::Dataflow
+{
+	void RegisterStaticMeshAssetNodes()
+	{
+		DATAFLOW_NODE_REGISTER_CREATION_FACTORY(FDataflowStaticMeshTerminalNode);
+	}
+}
+
+////////////////////////////////////////////////////////////////////////////////////////////
+
+FDataflowStaticMeshTerminalNode::FDataflowStaticMeshTerminalNode(const UE::Dataflow::FNodeParameters& InParam, FGuid InGuid)
+	: FDataflowTerminalNode(InParam, InGuid)
+{
+	RegisterInputConnection(&Mesh);
+	RegisterInputConnection(&Materials);
+	RegisterInputConnection(&AssetPath);
+	RegisterOutputConnection(&StaticMeshAsset);
+}
+
+void FDataflowStaticMeshTerminalNode::Evaluate(UE::Dataflow::FContext& Context) const
+{
+	// Asset is created and stored by the SetAssetValue
+	// Evaluate returns this created asset
+	const UE::Dataflow::FEngineContext* EngineContext = Context.AsType<UE::Dataflow::FEngineContext>();
+	UObject* BoundAsset = EngineContext ? EngineContext->Owner : nullptr;
+
+	const FString InAssetPath = GetAssetPath(GetValue(Context, &AssetPath), BoundAsset);
+	TObjectPtr<UStaticMesh> OutStaticMesh = ::LoadObject<UStaticMesh>(nullptr, InAssetPath);
+	SetValue(Context, OutStaticMesh, &StaticMeshAsset);
+}
+
+void FDataflowStaticMeshTerminalNode::SetAssetValue(TObjectPtr<UObject> Asset, UE::Dataflow::FContext& Context) const
+{
+	const TObjectPtr<UStaticMesh> NullStaticMeshPtr = nullptr;
+
+	const TObjectPtr<UDynamicMesh> InMesh = GetValue(Context, &Mesh);
+	if (!IsValid(InMesh))
+	{
+		Context.Error(LOCTEXT("InvalidInputMesh", "Input dynamic mesh mesh is empty"), this);
+		SetValue(Context, NullStaticMeshPtr, &StaticMeshAsset);
+		return;
+	}
+
+	// Materials are optional, no need to check for validity
+	const TArray<TObjectPtr<UMaterialInterface>>& InMaterials = GetValue(Context, &Materials);
+
+	const FString InAssetPath = GetAssetPath(GetValue(Context, &AssetPath), Asset);
+	if (InAssetPath.IsEmpty())
+	{
+		Context.Error(FText::Format(
+			LOCTEXT("InvalidPathAndBoundAsset", "Asset Path input is not a valid path and no static mesh asset is bound to this dataflow: path = {0}"), 
+			FText::FromString(InAssetPath)), 
+			this);
+		SetValue(Context, NullStaticMeshPtr, &StaticMeshAsset);
+		return;
+	}
+
+	UObject* OutAsset = GetOrCreateAsset(Context, InAssetPath, UStaticMesh::StaticClass());
+	TObjectPtr<UStaticMesh> OutStaticMesh = Cast<UStaticMesh>(OutAsset);
+	if (!OutStaticMesh)
+	{
+		Context.Error(FText::Format(
+			LOCTEXT("InvalidOutputAsset", "Failed to convert the new asset to a Static mesh: path = {0}"),
+			FText::FromString(InAssetPath)),
+			this);
+		SetValue(Context, NullStaticMeshPtr, &StaticMeshAsset);
+		return;
+	}
+
+	if (!ConvertDynamicMeshToStaticMesh(Context, *InMesh, InMaterials, *OutStaticMesh))
+	{
+		SetValue(Context, NullStaticMeshPtr, &StaticMeshAsset);
+		return;
+	}
+	SetValue(Context, OutStaticMesh, &StaticMeshAsset);
+
+}
+
+bool FDataflowStaticMeshTerminalNode::ConvertDynamicMeshToStaticMesh(
+	UE::Dataflow::FContext& Context, 
+	const UDynamicMesh& InMesh, 
+	const TArray<TObjectPtr<UMaterialInterface>>& InMaterials, 
+	UStaticMesh& OutMesh) const
+{
+	using namespace UE::Geometry;
+
+	UStaticMeshDescription* StaticMeshDescription = OutMesh.GetStaticMeshDescription(0);
+	if (!StaticMeshDescription)
+	{
+		StaticMeshDescription = UStaticMesh::CreateStaticMeshDescription(&OutMesh);
+	}
+	if (!StaticMeshDescription)
+	{
+		Context.Error(LOCTEXT("FailToCreateStaticMeshDesc", "Failed to create static mesh description"), this);
+		return false;
+	}
+
+	// start fresh, empty the description 
+	FMeshDescription& MeshDescription = StaticMeshDescription->GetMeshDescription();
+	MeshDescription.Empty();
+
+	InMesh.ProcessMesh([&MeshDescription](const FDynamicMesh3& InSourceMesh)
+		{
+			FConversionToMeshDescriptionOptions ConversionOptions;
+			ConversionOptions.bUpdateTangents = true;
+			ConversionOptions.bUpdateUVs = true;
+
+			FDynamicMeshToMeshDescription Converter(ConversionOptions);
+			Converter.Convert(&InSourceMesh, MeshDescription, /** CopyTangents */true);
+		});
+
+	// assign the materials
+	TArray<FStaticMaterial> NewStaticMaterials;
+	NewStaticMaterials.Reserve(InMaterials.Num());
+	for (int32 MaterialIndex = 0; MaterialIndex < InMaterials.Num(); ++MaterialIndex)
+	{
+		FStaticMaterial& NewMaterial = NewStaticMaterials.AddDefaulted_GetRef();
+		NewMaterial.MaterialInterface = InMaterials[MaterialIndex].Get();
+		NewMaterial.MaterialSlotName = UE::AssetUtils::GenerateNewMaterialSlotName(NewStaticMaterials, NewMaterial.MaterialInterface, MaterialIndex);
+		NewMaterial.ImportedMaterialSlotName = NewMaterial.MaterialSlotName;
+		NewMaterial.UVChannelData = FMeshUVChannelInfo(1.f);
+	}
+	OutMesh.SetStaticMaterials(NewStaticMaterials);
+
+	TPolygonGroupAttributesRef<FName> PolygonGroupImportedMaterialSlotNames = StaticMeshDescription->GetRequiredAttributes().GetPolygonGroupMaterialSlotNames();
+	for (int32 SlotIdx = 0; SlotIdx < NewStaticMaterials.Num(); ++SlotIdx)
+	{
+		if (SlotIdx < PolygonGroupImportedMaterialSlotNames.GetNumElements())
+		{
+			PolygonGroupImportedMaterialSlotNames.Set(SlotIdx, NewStaticMaterials[SlotIdx].MaterialSlotName);
+		}
+	}
+
+	FStaticMeshOperations::ComputeTriangleTangentsAndNormals(MeshDescription);
+	FStaticMeshOperations::ComputeTangentsAndNormals(MeshDescription, EComputeNTBsFlags::Normals);
+
+	// Build the static mesh
+	UStaticMesh::FBuildMeshDescriptionsParams Params;
+	Params.bFastBuild = true;
+	Params.bAllowCpuAccess = true;
+	Params.bCommitMeshDescription = true;
+	const TArray<const FMeshDescription*> MeshDescriptions{ &MeshDescription };
+	const bool bResult = OutMesh.BuildFromMeshDescriptions(MeshDescriptions, Params);
+
+#if WITH_EDITOR
+	UStaticMesh::FBuildParameters BuildParameters;
+	BuildParameters.bInSilent = true;
+	BuildParameters.bInRebuildUVChannelData = true;
+	BuildParameters.bInEnforceLightmapRestrictions = true;
+	OutMesh.Build(BuildParameters);
+#endif
+
+	return bResult;
+}
+
+
+
+#undef LOCTEXT_NAMESPACE
diff --git a/Engine/Plugins/Experimental/Dataflow/Source/DataflowNodes/Private/Dataflow/DataflowStaticMeshAssetNodes.h b/Engine/Plugins/Experimental/Dataflow/Source/DataflowNodes/Private/Dataflow/DataflowStaticMeshAssetNodes.h
new file mode 100644
--- /dev/null
+++ b/Engine/Plugins/Experimental/Dataflow/Source/DataflowNodes/Private/Dataflow/DataflowStaticMeshAssetNodes.h
@@ -0,0 +1,55 @@
+// Copyright Epic Games, Inc. All Rights Reserved.
+
+#pragma once 
+
+#include "Dataflow/DataflowTerminalNode.h"
+
+#include "DataflowStaticMeshAssetNodes.generated.h"
+
+class UStaticMesh;
+class UDynamicMesh;
+class UMaterialInterface;
+
+/*
+* terminal node to a save a static mesh asset
+*/
+USTRUCT(meta = (DataflowTerminal))
+struct FDataflowStaticMeshTerminalNode : public FDataflowTerminalNode
+{
+	GENERATED_USTRUCT_BODY()
+	DATAFLOW_NODE_DEFINE_INTERNAL(FDataflowStaticMeshTerminalNode, "StaticMeshTerminal", "Terminal", "")
+
+private:
+	/** Input dynamic mesh to convert to a static mesh */
+	UPROPERTY(meta = (DataflowInput))
+	TObjectPtr<UDynamicMesh> Mesh = nullptr;
+
+	/** Materials to set on the static mesh */
+	UPROPERTY(meta = (DataflowInput))
+	TArray<TObjectPtr<UMaterialInterface>> Materials;
+
+	/** 
+	* Path of the asset to save to 
+	* if not specified this will use the asset bound the dataflow graph this node belong to
+	*/
+	UPROPERTY(EditAnywhere, Category = Asset, meta = (DataflowInput))
+	FString AssetPath;
+
+	/** Output static mesh asset created by this node */
+	UPROPERTY(meta = (DataflowOutput))
+	TObjectPtr<UStaticMesh> StaticMeshAsset = nullptr;
+
+	virtual void Evaluate(UE::Dataflow::FContext& Context) const override;
+	virtual void SetAssetValue(TObjectPtr<UObject> Asset, UE::Dataflow::FContext& Context) const override;
+
+	bool ConvertDynamicMeshToStaticMesh(UE::Dataflow::FContext& Context, const UDynamicMesh& InMesh, const TArray<TObjectPtr<UMaterialInterface>>& InMaterials, UStaticMesh& OutMesh) const;
+
+public:
+	FDataflowStaticMeshTerminalNode(const UE::Dataflow::FNodeParameters& InParam, FGuid InGuid = FGuid::NewGuid());
+};
+
+namespace UE::Dataflow
+{
+	void RegisterStaticMeshAssetNodes();
+}
+
diff --git a/Engine/Source/Runtime/Experimental/Dataflow/Core/Private/Dataflow/DataflowTerminalNode.cpp b/Engine/Source/Runtime/Experimental/Dataflow/Core/Private/Dataflow/DataflowTerminalNode.cpp
--- a/Engine/Source/Runtime/Experimental/Dataflow/Core/Private/Dataflow/DataflowTerminalNode.cpp
+++ b/Engine/Source/Runtime/Experimental/Dataflow/Core/Private/Dataflow/DataflowTerminalNode.cpp
@@ -6,4 +6,92 @@
 #include "Dataflow/DataflowInputOutput.h"
 #include "Dataflow/DataflowArchive.h"
 
+#include "AssetRegistry/AssetRegistryModule.h"
+#include "Misc/PackageName.h"
+
 #include UE_INLINE_GENERATED_CPP_BY_NAME(DataflowTerminalNode)
+
+#define LOCTEXT_NAMESPACE "DataflowTerminalNode"
+
+FString FDataflowTerminalNode::GetAssetPath(const FString& InputPath, UObject* BoundAsset)
+{
+	static FString EmptyString;
+
+	FString OutPath = InputPath;
+	if (!FPackageName::IsValidObjectPath(OutPath))
+	{
+		// no valid path , is there a bound asset ?
+		if (BoundAsset)
+		{
+			OutPath = BoundAsset->GetPathName();
+		}
+		if (!FPackageName::IsValidObjectPath(OutPath))
+		{
+			OutPath = EmptyString;
+		}
+	}
+	return OutPath;
+}
+
+UObject* FDataflowTerminalNode::GetOrCreateAsset(UE::Dataflow::FContext& Context, const FString& InAssetPath, const UClass* AssetClass) const
+{
+#if WITH_EDITOR
+	if (!FPackageName::IsValidObjectPath(InAssetPath))
+	{
+		Context.Error(FText::Format(
+			LOCTEXT("InvalidAssetPath", "invalid asset path: Path = {0}"), FText::FromString(InAssetPath)),
+			this);
+		return nullptr;
+	}
+
+	// Try to create the package
+	const FString PackageName = FPackageName::ObjectPathToPackageName(InAssetPath);
+	UPackage* Package = Cast<UPackage>(FindPackage(nullptr, *PackageName));
+	if (Package == nullptr)
+	{
+		Package = CreatePackage(*PackageName);
+	}
+	if (Package == nullptr)
+	{
+		Context.Error(FText::Format(
+			LOCTEXT("InvalidPackage", "Failed to find or create package {0}"), FText::FromString(PackageName)),
+			this);
+		return nullptr;
+	}
+
+	// Check if there's an existing object matching this path and compatible with this type
+	const FName AssetName = FName(FPackageName::GetLongPackageAssetName(PackageName));
+	UObject* ExistingObject = StaticFindObjectFastInternal( /*Class=*/ NULL, Package, AssetName, EFindObjectFlags::ExactClass);
+	if (ExistingObject && !ExistingObject->GetClass()->IsChildOf(AssetClass))
+	{
+		Context.Error(FText::Format(
+			LOCTEXT("IncompatibleType", "Asset {0} already exists but is not a {1} compatible type"),
+			FText::FromString(InAssetPath), FText::FromName(AssetClass->GetFName())),
+			this);
+		return nullptr;
+	}
+
+	UObject* OutObject = ExistingObject;
+	if (!OutObject)
+	{
+		OutObject = NewObject<UObject>(Package, AssetClass, AssetName, RF_Public | RF_Standalone | RF_Transactional);
+		if (!OutObject)
+		{
+			Context.Error(FText::Format(
+				LOCTEXT("AssetCreationFailed", "Failed to create asset {0}"), FText::FromString(InAssetPath)),
+				this);
+			return nullptr;
+		}
+	}
+	// make sure the asset is set properly
+	OutObject->MarkPackageDirty();
+	FAssetRegistryModule::AssetCreated(OutObject);
+
+	return OutObject;
+#else
+	Context.Error(LOCTEXT("AssetCreationOnlyEditor", "Creation of assets only supported in Editor"), this);
+	return nullptr;
+#endif
+}
+
+#undef LOCTEXT_NAMESPACE
\ No newline at end of file
diff --git a/Engine/Source/Runtime/Experimental/Dataflow/Core/Public/Dataflow/DataflowTerminalNode.h b/Engine/Source/Runtime/Experimental/Dataflow/Core/Public/Dataflow/DataflowTerminalNode.h
--- a/Engine/Source/Runtime/Experimental/Dataflow/Core/Public/Dataflow/DataflowTerminalNode.h
+++ b/Engine/Source/Runtime/Experimental/Dataflow/Core/Public/Dataflow/DataflowTerminalNode.h
@@ -10,9 +10,14 @@
 
 #include "DataflowTerminalNode.generated.h"
 
+#define UE_API DATAFLOWCORE_API
+
 struct FDataflowInput;
 struct FDataflowOutput;
 
+class UObject;
+class UClass;
+
 /**
 * FDataflowTerminalNode
 *		Base class for terminal nodes within the Dataflow graph. 
@@ -64,6 +69,11 @@ struct FDataflowTerminalNode : public FDataflowNode
 
 	virtual void Evaluate(UE::Dataflow::FContext& Context) const { ensure(false); }
 
+protected:
+	DATAFLOWCORE_API static FString GetAssetPath(const FString& InputPath, UObject* BoundAsset);
+
+	DATAFLOWCORE_API UObject* GetOrCreateAsset(UE::Dataflow::FContext& Context, const FString& InAssetPath, const UClass* AssetClass) const;
+
 private:
 	virtual void Evaluate(UE::Dataflow::FContext& Context, const FDataflowOutput* Out) const override
 	{
@@ -71,4 +81,4 @@ private:
 	};
 };
 
-
+#undef UE_API 
