Dataflow : fix build warning

#Synced-CL 51254598

diff --git a/Engine/Plugins/Experimental/Dataflow/Source/DataflowNodes/Private/Dataflow/DataflowStaticMeshAssetNodes.cpp b/Engine/Plugins/Experimental/Dataflow/Source/DataflowNodes/Private/Dataflow/DataflowStaticMeshAssetNodes.cpp
--- a/Engine/Plugins/Experimental/Dataflow/Source/DataflowNodes/Private/Dataflow/DataflowStaticMeshAssetNodes.cpp
+++ b/Engine/Plugins/Experimental/Dataflow/Source/DataflowNodes/Private/Dataflow/DataflowStaticMeshAssetNodes.cpp
@@ -39,7 +39,7 @@ void FDataflowStaticMeshTerminalNode::Evaluate(UE::Dataflow::FContext& Context)
 	// Asset is created and stored by the SetAssetValue
 	// Evaluate returns this created asset
 	const UE::Dataflow::FEngineContext* EngineContext = Context.AsType<UE::Dataflow::FEngineContext>();
-	UObject* BoundAsset = EngineContext ? EngineContext->Owner : nullptr;
+	TObjectPtr<UObject> BoundAsset = EngineContext ? EngineContext->Owner : nullptr;
 
 	const FString InAssetPath = GetAssetPath(GetValue(Context, &AssetPath), BoundAsset);
 	TObjectPtr<UStaticMesh> OutStaticMesh = ::LoadObject<UStaticMesh>(nullptr, InAssetPath);
