Dataflow:
- Fixed possible overflow

#Synced-CL 51288899

diff --git a/Engine/Plugins/Experimental/Dataflow/Source/DataflowNodes/Private/Dataflow/SamplerNodes/DataflowSamplerFunctions.cpp b/Engine/Plugins/Experimental/Dataflow/Source/DataflowNodes/Private/Dataflow/SamplerNodes/DataflowSamplerFunctions.cpp
--- a/Engine/Plugins/Experimental/Dataflow/Source/DataflowNodes/Private/Dataflow/SamplerNodes/DataflowSamplerFunctions.cpp
+++ b/Engine/Plugins/Experimental/Dataflow/Source/DataflowNodes/Private/Dataflow/SamplerNodes/DataflowSamplerFunctions.cpp
@@ -7,11 +7,11 @@ namespace UE::Dataflow::Sampler
 {
 	float FRandom(const FVector InPosition, float InRandomMin, float InRandomMax, const int32 InSeed)
 	{
-		const int32 SeedX = int32(InPosition.X);
-		const int32 SeedY = int32(InPosition.Y);
-		const int32 SeedZ = int32(InPosition.Z);
+		const uint32 SeedX = uint32(FMath::Abs(InPosition.X));
+		const uint32 SeedY = uint32(FMath::Abs(InPosition.Y));
+		const uint32 SeedZ = uint32(FMath::Abs(InPosition.Z));
 	
-		const int64 SeedAll = int64(SeedX + SeedY + SeedZ + InSeed);
+		const uint32 SeedAll = ::HashCombineFast(SeedX, SeedY, SeedZ, InSeed);
 	
 		FRandomStream RandStream((int32)SeedAll);
 
@@ -25,11 +25,11 @@ namespace UE::Dataflow::Sampler
 
 	FVector VRandom(const FVector InPosition, float InRandomMin, float InRandomMax, const int32 InSeed)
 	{
-		const int32 SeedX = int32(InPosition.X);
-		const int32 SeedY = int32(InPosition.Y);
-		const int32 SeedZ = int32(InPosition.Z);
+		const uint32 SeedX = uint32(FMath::Abs(InPosition.X));
+		const uint32 SeedY = uint32(FMath::Abs(InPosition.Y));
+		const uint32 SeedZ = uint32(FMath::Abs(InPosition.Z));
 
-		const int64 SeedAll = int64(SeedX + SeedY + SeedZ + InSeed);
+		const uint32 SeedAll = ::HashCombineFast(SeedX, SeedY, SeedZ, InSeed);
 
 		FRandomStream RandStream((int32)SeedAll);
 
