Dataflow Editor : Remove mesh construction visualization as this has been replaced by the render settings

#Synced-CL 51288786
#Synced-CL 51288786

diff --git a/Engine/Plugins/Experimental/Dataflow/Source/DataflowEditor/Private/Dataflow/DataflowEditorModule.cpp b/Engine/Plugins/Experimental/Dataflow/Source/DataflowEditor/Private/Dataflow/DataflowEditorModule.cpp
--- a/Engine/Plugins/Experimental/Dataflow/Source/DataflowEditor/Private/Dataflow/DataflowEditorModule.cpp
+++ b/Engine/Plugins/Experimental/Dataflow/Source/DataflowEditor/Private/Dataflow/DataflowEditorModule.cpp
@@ -328,7 +328,6 @@ void FDataflowEditorModule::StartupModule()
 	
 	UE::Dataflow::FDataflowConstructionVisualizationRegistry& ConstructionVisualizationRegistry = UE::Dataflow::FDataflowConstructionVisualizationRegistry::GetInstance();
 	ConstructionVisualizationRegistry.RegisterVisualization(MakeUnique<UE::Dataflow::FMeshStatsConstructionVisualization>());
-	ConstructionVisualizationRegistry.RegisterVisualization(MakeUnique<UE::Dataflow::FMeshConstructionVisualization>());
 	ConstructionVisualizationRegistry.RegisterVisualization(MakeUnique<UE::Dataflow::FDebugDrawText3dVisualization>());
 	
 	FDataflowEditorCommands::Register();
@@ -398,7 +397,6 @@ void FDataflowEditorModule::ShutdownModule()
 
 	UE::Dataflow::FDataflowConstructionVisualizationRegistry& ConstructionVisualizationRegistry = UE::Dataflow::FDataflowConstructionVisualizationRegistry::GetInstance();
 	ConstructionVisualizationRegistry.DeregisterVisualization(UE::Dataflow::FMeshStatsConstructionVisualization::Name);
-	ConstructionVisualizationRegistry.DeregisterVisualization(UE::Dataflow::FMeshConstructionVisualization::Name);
 	ConstructionVisualizationRegistry.DeregisterVisualization(UE::Dataflow::FDebugDrawText3dVisualization::Name);
 }
 
