Dataflow : invalidate nodes when chaging the type of their any type output via the pin context menu

#Synced-CL 51362814

diff --git a/Engine/Plugins/Experimental/Dataflow/Source/DataflowEditor/Private/Dataflow/DataflowSchema.cpp b/Engine/Plugins/Experimental/Dataflow/Source/DataflowEditor/Private/Dataflow/DataflowSchema.cpp
--- a/Engine/Plugins/Experimental/Dataflow/Source/DataflowEditor/Private/Dataflow/DataflowSchema.cpp
+++ b/Engine/Plugins/Experimental/Dataflow/Source/DataflowEditor/Private/Dataflow/DataflowSchema.cpp
@@ -138,6 +138,7 @@ void UDataflowSchema::GetContextMenuActions(class UToolMenu* Menu, class UGraphN
 												{
 													MutableEdNode->UpdatePinsFromDataflowNode();
 												}
+												Node->Invalidate();
 											}
 										}),
 									// Fine to capture raw pointers because the lifetime of teh menu is shorter than the connection and node lifetime
