Dataflow : add group name to the selection watch string

#Synced-CL 51362814
#Synced-CL 51362814

diff --git a/Engine/Plugins/Experimental/Dataflow/Source/DataflowEditor/Private/Dataflow/DataflowSNode.cpp b/Engine/Plugins/Experimental/Dataflow/Source/DataflowEditor/Private/Dataflow/DataflowSNode.cpp
--- a/Engine/Plugins/Experimental/Dataflow/Source/DataflowEditor/Private/Dataflow/DataflowSNode.cpp
+++ b/Engine/Plugins/Experimental/Dataflow/Source/DataflowEditor/Private/Dataflow/DataflowSNode.cpp
@@ -71,7 +71,7 @@ namespace UE::Dataflow::Private
 		{
 			const FDataflowSelection DefaultValue;
 			const FDataflowSelection& Value = Output.ReadValue(Context, DefaultValue);
-			OutString += FString::Printf(TEXT("%s : %d/%d\n"), *Output.GetName().ToString(), Value.NumSelected(), Value.Num());
+			OutString += FString::Printf(TEXT("%s : (%s) %d/%d\n"), *Output.GetName().ToString(), *Value.GetGroupName().ToString(), Value.NumSelected(), Value.Num());
 		}
 		else if (FDataflowStringConvertibleTypePolicy::SupportsTypeStatic(Output.GetType()))
 		{
