Dataflow : make the collection spreadsheet refresh asynchornously on the next tick

#Synced-CL 51402414

diff --git a/Engine/Plugins/Experimental/Dataflow/Source/DataflowEditor/Private/Dataflow/DataflowCollectionSpreadSheet.cpp b/Engine/Plugins/Experimental/Dataflow/Source/DataflowEditor/Private/Dataflow/DataflowCollectionSpreadSheet.cpp
--- a/Engine/Plugins/Experimental/Dataflow/Source/DataflowEditor/Private/Dataflow/DataflowCollectionSpreadSheet.cpp
+++ b/Engine/Plugins/Experimental/Dataflow/Source/DataflowEditor/Private/Dataflow/DataflowCollectionSpreadSheet.cpp
@@ -1,6 +1,8 @@
 // Copyright Epic Games, Inc. All Rights Reserved.
 
 #include "Dataflow/DataflowCollectionSpreadSheet.h"
+
+#include "Containers/Ticker.h"
 #include "Dataflow/DataflowObjectInterface.h"
 #include "Dataflow/DataflowCollectionSpreadSheetWidget.h"
 #include "Dataflow/DataflowContent.h"
@@ -25,46 +27,48 @@ void FDataflowCollectionSpreadSheet::SetSupportedOutputTypes()
 
 void FDataflowCollectionSpreadSheet::UpdateViewData()
 {
-	if (CollectionSpreadSheet && CollectionSpreadSheet->GetCollectionTable())
-	{
-		CollectionSpreadSheet->GetCollectionTable()->GetCollectionInfoMap().Empty();
+	TWeakObjectPtr<UDataflowEdNode> WeakSelectedNode = GetSelectedNode();
+	TWeakObjectPtr<UDataflowBaseContent> WeakEditorContent = GetEditorContent();
+
+	ExecuteOnGameThread(TEXT("DataflowCollectionSpreadSheetRefresh"), [CollectionSpreadSheet = CollectionSpreadSheet, WeakSelectedNode, WeakEditorContent]() {
 
-		if (GetSelectedNode())
+		FString NodeName;
+		if (CollectionSpreadSheet && CollectionSpreadSheet->GetCollectionTable())
 		{
-			if (GetSelectedNode()->IsBound())
+			CollectionSpreadSheet->GetCollectionTable()->GetCollectionInfoMap().Empty();
+
+			if (TStrongObjectPtr<UDataflowEdNode> SelectedEdNode = WeakSelectedNode.Pin())
 			{
-				if (TSharedPtr<FDataflowNode> DataflowNode = GetSelectedNode()->DataflowGraph->FindBaseNode(GetSelectedNode()->DataflowNodeGuid))
+				if (SelectedEdNode->IsBound())
 				{
-					TArray<FDataflowOutput*> Outputs = DataflowNode->GetOutputs();
-					if (const TObjectPtr<UDataflowBaseContent> Content = GetEditorContent())
+					if (TSharedPtr<FDataflowNode> DataflowNode = SelectedEdNode->GetDataflowNode())
 					{
-						if (TSharedPtr<UE::Dataflow::FEngineContext> Context = Content->GetDataflowContext())
+						TArray<FDataflowOutput*> Outputs = DataflowNode->GetOutputs();
+						if (TStrongObjectPtr<UDataflowBaseContent> Content = WeakEditorContent.Pin())
 						{
-							for (FDataflowOutput* Output : Outputs)
+							if (TSharedPtr<UE::Dataflow::FEngineContext> Context = Content->GetDataflowContext())
 							{
-								if (Output->GetType() == "FManagedArrayCollection")
+								for (FDataflowOutput* Output : Outputs)
 								{
-									const FName Name = Output->GetName();
-									const FManagedArrayCollection DefaultCollection;
-									const FManagedArrayCollection& Value = Output->ReadValue(*Context, DefaultCollection);
-									TSharedPtr<const FManagedArrayCollection> CopiedCollectionPtr = MakeShared<const FManagedArrayCollection>(Value);
-									CollectionSpreadSheet->GetCollectionTable()->GetCollectionInfoMap().Add(Name.ToString(), { CopiedCollectionPtr });
+									if (Output->GetType() == TEXT("FManagedArrayCollection"))
+									{
+										const FName Name = Output->GetName();
+										const FManagedArrayCollection DefaultCollection;
+										const FManagedArrayCollection& Value = Output->ReadValue(*Context, DefaultCollection);
+										TSharedPtr<const FManagedArrayCollection> CopiedCollectionPtr = MakeShared<const FManagedArrayCollection>(Value);
+										CollectionSpreadSheet->GetCollectionTable()->GetCollectionInfoMap().Add(Name.ToString(), { CopiedCollectionPtr });
+									}
 								}
 							}
 						}
+						NodeName = SelectedEdNode->GetName();
 					}
 				}
 			}
-
-			CollectionSpreadSheet->SetData(GetSelectedNode()->GetName());
+			CollectionSpreadSheet->SetData(NodeName);
+			CollectionSpreadSheet->RefreshWidget();
 		}
-		else
-		{
-			CollectionSpreadSheet->SetData(FString());
-		}
-
-		CollectionSpreadSheet->RefreshWidget();
-	}
+	});
 }
 
 
