Dataflow editor : remove unused DebugMesh
- this was still computed while not being used

#Synced-CL 51402414

diff --git a/Engine/Plugins/Experimental/Dataflow/Source/DataflowEditor/Private/Dataflow/DataflowConstructionScene.cpp b/Engine/Plugins/Experimental/Dataflow/Source/DataflowEditor/Private/Dataflow/DataflowConstructionScene.cpp
--- a/Engine/Plugins/Experimental/Dataflow/Source/DataflowEditor/Private/Dataflow/DataflowConstructionScene.cpp
+++ b/Engine/Plugins/Experimental/Dataflow/Source/DataflowEditor/Private/Dataflow/DataflowConstructionScene.cpp
@@ -178,32 +178,6 @@ void FDataflowConstructionScene::TickDataflowScene(const float DeltaSeconds)
 	}
 }
 
-void FDataflowConstructionScene::FDebugMesh::Reset()
-{
-	VertexMap.Reset();
-	FaceMap.Reset();
-}
-
-void FDataflowConstructionScene::FDebugMesh::Build(const TArray<TObjectPtr<UDynamicMeshComponent>>& InDynamicMeshComponents)
-{
-	ResultMesh.Clear();
-	ResultMesh.EnableAttributes();
-
-	for (const UDynamicMeshComponent* const DynamicMeshComponent : InDynamicMeshComponents)
-	{
-		if (const UE::Geometry::FDynamicMesh3* const Mesh = DynamicMeshComponent->GetMesh())
-		{
-			ResultMesh.AppendWithOffsets(*Mesh);
-		}
-	}
-
-	// Disable it for now
-	// TODO: Optimize this
-#if 0
-	Spatial = UE::Geometry::FDynamicMeshAABBTree3(&ResultMesh, true);
-#endif
-}
-
 void FDataflowConstructionScene::UpdateDynamicMeshComponents()
 {
 	using namespace UE::Geometry;//FDynamicMesh3
@@ -248,9 +222,6 @@ void FDataflowConstructionScene::UpdateDynamicMeshComponents()
 				}
 			}
 		}
-	
-		// Build a single mesh out of all the components
-		DebugMesh.Build(GetDynamicMeshComponents());
 	}
 
 	UpdateFloorVisibility();
@@ -566,37 +537,6 @@ void FDataflowConstructionScene::SelectNodeComponents(const UDataflowEdNode* EdN
 	}
 }
 
-void FDataflowConstructionScene::UpdateDebugMesh(const GeometryCollection::Facades::FRenderingFacade& RenderingFacade)
-{
-	DebugMesh.Reset();
-
-	int32 VertexIndex = 0;
-	int32 FaceIndex = 0;
-
-	const int32 NumGeometry = RenderingFacade.NumGeometry();
-	for (int32 MeshIndex = 0; MeshIndex < NumGeometry; ++MeshIndex)
-	{
-		// Accumulate in a unified debug mesh for visualization
-		// todo(dataflow) : we should get rid of this eventually as this is costly to build and maintain
-		const TManagedArray<int32>& VertexStart = RenderingFacade.GetVertexStart();
-		const TManagedArray<int32>& VertexCount = RenderingFacade.GetVertexCount();
-		for (int32 Idx = 0; Idx < VertexCount[MeshIndex]; ++Idx)
-		{
-			DebugMesh.VertexMap.Add(VertexIndex, VertexStart[MeshIndex] + Idx);
-			VertexIndex++;
-		}
-
-		const TManagedArray<int32>& FaceStart = RenderingFacade.GetIndicesStart();
-		const TManagedArray<int32>& FaceCount = RenderingFacade.GetIndicesCount();
-		for (int32 Idx = 0; Idx < FaceCount[MeshIndex]; ++Idx)
-		{
-			DebugMesh.FaceMap.Add(FaceIndex, FaceStart[MeshIndex] + Idx);
-			FaceIndex++;
-		}
-	}
-}
-
-
 TArray<UMaterialInterface*> FDataflowConstructionScene::GetMaterialsFromRenderingFacade(const GeometryCollection::Facades::FRenderingFacade& RenderingFacade, int32 MeshIndex)
 {
 	TArray<UMaterialInterface*> Materials;
@@ -730,7 +670,6 @@ void FDataflowConstructionScene::UpdateComponentsFromNode(const UDataflowEdNode*
 
 					ProgressNotification.AddProgress(GeometryProgress);
 				}
-				UpdateDebugMesh(RenderingFacade);
 			}
 			else // New method where rendering callbacks generate primitive components directly by output types
 			{
diff --git a/Engine/Plugins/Experimental/Dataflow/Source/DataflowEditor/Private/Dataflow/DataflowEditorModule.cpp b/Engine/Plugins/Experimental/Dataflow/Source/DataflowEditor/Private/Dataflow/DataflowEditorModule.cpp
--- a/Engine/Plugins/Experimental/Dataflow/Source/DataflowEditor/Private/Dataflow/DataflowEditorModule.cpp
+++ b/Engine/Plugins/Experimental/Dataflow/Source/DataflowEditor/Private/Dataflow/DataflowEditorModule.cpp
@@ -15,7 +15,6 @@
 #include "Dataflow/DataflowSNodeFactories.h"
 #include "Dataflow/DebugDrawText3dVisualization.h"
 #include "Dataflow/MeshStatsConstructionVisualization.h"
-#include "Dataflow/MeshConstructionVisualization.h"
 #include "Dataflow/ScalarVertexPropertyGroupCustomization.h"
 #include "Dataflow/DataflowAnyTypeCustomization.h"
 #include "Dataflow/DataflowToolRegistry.h"
diff --git a/Engine/Plugins/Experimental/Dataflow/Source/DataflowEditor/Private/Dataflow/MeshConstructionVisualization.cpp b/Engine/Plugins/Experimental/Dataflow/Source/DataflowEditor/Private/Dataflow/MeshConstructionVisualization.cpp
deleted file mode 100644
diff --git a/Engine/Plugins/Experimental/Dataflow/Source/DataflowEditor/Private/Dataflow/MeshConstructionVisualization.h b/Engine/Plugins/Experimental/Dataflow/Source/DataflowEditor/Private/Dataflow/MeshConstructionVisualization.h
deleted file mode 100644
diff --git a/Engine/Plugins/Experimental/Dataflow/Source/DataflowEditor/Public/Dataflow/DataflowConstructionScene.h b/Engine/Plugins/Experimental/Dataflow/Source/DataflowEditor/Public/Dataflow/DataflowConstructionScene.h
--- a/Engine/Plugins/Experimental/Dataflow/Source/DataflowEditor/Public/Dataflow/DataflowConstructionScene.h
+++ b/Engine/Plugins/Experimental/Dataflow/Source/DataflowEditor/Public/Dataflow/DataflowConstructionScene.h
@@ -40,17 +40,6 @@ class FDataflowConstructionScene : public FDataflowPreviewSceneBase
 	/*(const void* ObjectPointer, const int32 ObjectIndex)*/
 	typedef TPair<TObjectPtr<const UDataflowEdNode>, int32> FDataflowRenderKey;
 
-	struct FDebugMesh
-	{
-		UE::Geometry::FDynamicMesh3 ResultMesh;
-		UE::Geometry::FDynamicMeshAABBTree3 Spatial;
-		TMap<int32, int32> VertexMap;
-		TMap<int32, int32> FaceMap;
-
-		void Reset();
-		void Build(const TArray<TObjectPtr<UDynamicMeshComponent>>& InDynamicMeshComponents);
-	};
-
 public:
 
 	inline static FName DebugMeshComponentTag = TEXT("MeshDebugComponent");
@@ -91,8 +80,6 @@ public:
 	/** select only the components that are relates to a specific node */
 	UE_API void SelectNodeComponents(const UDataflowEdNode* EdNode);
 
-	FDebugMesh DebugMesh;
-
 	/** Get the render settings for matching the selected nodes */
 	UE_API TArray<UDataflowRenderableTypeSettings*> GetRenderableTypeSettings() const;
 
@@ -132,8 +119,6 @@ private:
 	/** Remove the scene component from selection and destroy it */
 	UE_API void RemoveSceneComponent(USelection* SelectedComponents, UPrimitiveComponent* PrimitiveComponent);
 
-	void UpdateDebugMesh(const GeometryCollection::Facades::FRenderingFacade& RenderingFacade);
-
 	/** update floor visibility based on the view type */
 	void UpdateFloorVisibility();
 
