Dataflow:
- Added MakeSelectionSet and GetSelectionSet nodes

#Synced-CL 51475101

diff --git a/Engine/Plugins/Experimental/Dataflow/Source/DataflowEditor/Private/Dataflow/DataflowGraphEditor.cpp b/Engine/Plugins/Experimental/Dataflow/Source/DataflowEditor/Private/Dataflow/DataflowGraphEditor.cpp
--- a/Engine/Plugins/Experimental/Dataflow/Source/DataflowEditor/Private/Dataflow/DataflowGraphEditor.cpp
+++ b/Engine/Plugins/Experimental/Dataflow/Source/DataflowEditor/Private/Dataflow/DataflowGraphEditor.cpp
@@ -13,7 +13,6 @@
 #include "Dataflow/DataflowSCommentNode.h"
 #include "Dataflow/DataflowNodeParameters.h"
 #include "Dataflow/DataflowGraphSchemaAction.h"
-#include "Dataflow/DataflowSelectionNodes.h"
 #include "Dataflow/DataflowSubGraphNodes.h"
 #include "Dataflow/DataflowSchema.h"
 #include "Dataflow/DataflowSubGraph.h"
diff --git a/Engine/Plugins/Experimental/Dataflow/Source/DataflowNodes/Private/Dataflow/DataflowSelectionNodes.cpp b/Engine/Plugins/Experimental/Dataflow/Source/DataflowNodes/Private/Dataflow/DataflowSelectionNodes.cpp
--- a/Engine/Plugins/Experimental/Dataflow/Source/DataflowNodes/Private/Dataflow/DataflowSelectionNodes.cpp
+++ b/Engine/Plugins/Experimental/Dataflow/Source/DataflowNodes/Private/Dataflow/DataflowSelectionNodes.cpp
@@ -2,6 +2,8 @@
 
 #include "Dataflow/DataflowSelectionNodes.h"
 #include "Engine/StaticMesh.h"
+#include "GeometryCollection/Facades/CollectionSelectionFacade.h"
+#include "Dataflow/DataflowSelection.h"
 
 #include UE_INLINE_GENERATED_CPP_BY_NAME(DataflowSelectionNodes)
 
@@ -10,6 +12,8 @@ namespace UE::Dataflow
 	void RegisterSelectionNodes()
 	{
 		DATAFLOW_NODE_REGISTER_CREATION_FACTORY(FSelectionSetDataflowNode);
+		DATAFLOW_NODE_REGISTER_CREATION_FACTORY(FMakeSelectionSetDataflowNode);
+		DATAFLOW_NODE_REGISTER_CREATION_FACTORY(FGetSelectionSetDataflowNode);
 	}
 };
 
@@ -32,5 +36,222 @@ void FSelectionSetDataflowNode::Evaluate(UE::Dataflow::FContext& Context, const
 	}
 }
 
+// --------------------------------------------------------------------------------------------------------------------------------------
 
+namespace UE::Dataflow::SelectionSet::Private
+{
+	FName GetSelectionTypeFromCollection(const FDataflowConnection& InConnection)
+	{
+		if (InConnection.IsType<FDataflowTransformSelection>())
+		{
+			return 	FDataflowTransformSelection::TransformGroupName;
+		}
+		else if (InConnection.IsType<FDataflowVertexSelection>())
+		{
+			return 	FDataflowVertexSelection::VerticesGroupName;
+		}
+		else if (InConnection.IsType<FDataflowFaceSelection>())
+		{
+			return 	FDataflowFaceSelection::FacesGroupName;
+		}
+		else if (InConnection.IsType<FDataflowGeometrySelection>())
+		{
+			return 	FDataflowGeometrySelection::GeometryGroupName;
+		}
+		else if (InConnection.IsType<FDataflowCurveSelection>())
+		{
+			return 	FDataflowCurveSelection::CurveGroupName;
+		}
+		else if (InConnection.IsType<FDataflowPointsSelection>())
+		{
+			return 	FDataflowPointsSelection::PointsGroupName;
+		}
+		else if (InConnection.IsType<FDataflowMaterialSelection>())
+		{
+			return 	FDataflowMaterialSelection::MaterialGroupName;
+		}
+
+		return NAME_None;
+	}
+}
+
+FMakeSelectionSetDataflowNode::FMakeSelectionSetDataflowNode(const UE::Dataflow::FNodeParameters& InParam, FGuid InGuid)
+	: FDataflowNode(InParam, InGuid)
+{
+	RegisterInputConnection(&Collection);
+	RegisterInputConnection(&Selection);
+	RegisterInputConnection(&Attribute);
+	RegisterOutputConnection(&Collection, &Collection);
+}
+
+FDataflowNode::FAttributeKey FMakeSelectionSetDataflowNode::GetVertexAttributeToVisualize(UE::Dataflow::FContext& Context, const FName OutputName, const FName RenderGroup) const
+{
+	FName GroupName;
+	if (const FDataflowInput* SelectionInput = FindInput(&Selection))
+	{
+		GroupName = UE::Dataflow::SelectionSet::Private::GetSelectionTypeFromCollection(*SelectionInput);
+	}
+
+	return FAttributeKey
+	{
+		.AttributeName = FName(GetValue(Context, &Attribute)),
+		.GroupName = GroupName
+	};
+}
+
+void FMakeSelectionSetDataflowNode::Evaluate(UE::Dataflow::FContext& Context, const FDataflowOutput* Out) const
+{
+	if (Out->IsA(&Collection))
+	{
+		if (IsConnected(&Collection))
+		{
+			if (IsConnected(&Selection))
+			{
+				FManagedArrayCollection InCollection = GetValue(Context, &Collection);
+				GeometryCollection::Facades::FSelectionFacade SelectionFacade(InCollection);
+
+				FString InAttribute = GetValue(Context, &Attribute);
+
+				const FDataflowInput* SelectionInput = FindInput(&Selection);
+				if (SelectionInput)
+				{
+					const FDataflowSelection& InSelection = GetValue(Context, &Selection);
+
+					if (SelectionInput->IsType<FDataflowTransformSelection>())
+					{
+						FDataflowTransformSelection NewSelection;
+						NewSelection.Initialize(InSelection);
+
+						SelectionFacade.AddSelectionToCollection(FDataflowTransformSelection::TransformGroupName, NewSelection.GetBitArray(), InAttribute, bStoreAsFloat);
+					}
+					else if (SelectionInput->IsType<FDataflowVertexSelection>())
+					{
+						FDataflowVertexSelection NewSelection;
+						NewSelection.Initialize(InSelection);
+
+						SelectionFacade.AddSelectionToCollection(FDataflowVertexSelection::VerticesGroupName, NewSelection.GetBitArray(), InAttribute, bStoreAsFloat);
+					}
+					else if (SelectionInput->IsType<FDataflowFaceSelection>())
+					{
+						FDataflowFaceSelection NewSelection;
+						NewSelection.Initialize(InSelection);
+
+						SelectionFacade.AddSelectionToCollection(FDataflowFaceSelection::FacesGroupName, NewSelection.GetBitArray(), InAttribute, bStoreAsFloat);
+					}
+					else if (SelectionInput->IsType<FDataflowGeometrySelection>())
+					{
+						FDataflowGeometrySelection NewSelection;
+						NewSelection.Initialize(InSelection);
+
+						SelectionFacade.AddSelectionToCollection(FDataflowGeometrySelection::GeometryGroupName, NewSelection.GetBitArray(), InAttribute, bStoreAsFloat);
+					}
+					else if (SelectionInput->IsType<FDataflowCurveSelection>())
+					{
+						FDataflowCurveSelection NewSelection;
+						NewSelection.Initialize(InSelection);
+
+						SelectionFacade.AddSelectionToCollection(FDataflowCurveSelection::CurveGroupName, NewSelection.GetBitArray(), InAttribute, bStoreAsFloat);
+					}
+					else if (SelectionInput->IsType<FDataflowPointsSelection>())
+					{
+						FDataflowPointsSelection NewSelection;
+						NewSelection.Initialize(InSelection);
+
+						SelectionFacade.AddSelectionToCollection(FDataflowPointsSelection::PointsGroupName, NewSelection.GetBitArray(), InAttribute, bStoreAsFloat);
+					}
+					else if (SelectionInput->IsType<FDataflowMaterialSelection>())
+					{
+						FDataflowMaterialSelection NewSelection;
+						NewSelection.Initialize(InSelection);
+
+						SelectionFacade.AddSelectionToCollection(FDataflowMaterialSelection::MaterialGroupName, NewSelection.GetBitArray(), InAttribute, bStoreAsFloat);
+					}
+
+					SetValue(Context, MoveTemp(InCollection), &Collection);
+					return;
+				}
+			}
+
+			SafeForwardInput(Context, &Collection, &Collection);
+			return;
+		}
+		
+		SetValue(Context, FManagedArrayCollection(), &Collection);
+	}
+}
+
+// --------------------------------------------------------------------------------------------------------------------------------------
+
+FGetSelectionSetDataflowNode::FGetSelectionSetDataflowNode(const UE::Dataflow::FNodeParameters& InParam, FGuid InGuid)
+	: FDataflowNode(InParam, InGuid)
+{
+	RegisterInputConnection(&Collection);
+	RegisterInputConnection(&Attribute);
+	RegisterOutputConnection(&Collection, &Collection);
+	RegisterOutputConnection(&Selection);
+}
+
+FDataflowNode::FAttributeKey FGetSelectionSetDataflowNode::GetVertexAttributeToVisualize(UE::Dataflow::FContext& Context, const FName OutputName, const FName RenderGroup) const
+{
+	FName GroupName;
+	if (const FDataflowOutput* SelectionOutput = FindOutput(&Selection))
+	{
+		GroupName = UE::Dataflow::SelectionSet::Private::GetSelectionTypeFromCollection(*SelectionOutput);
+	}
+
+	return FAttributeKey
+	{
+		.AttributeName = FName(GetValue(Context, &Attribute)),
+		.GroupName = GroupName
+	};
+}
+
+void FGetSelectionSetDataflowNode::Evaluate(UE::Dataflow::FContext& Context, const FDataflowOutput* Out) const
+{
+	if (Out->IsA(&Selection))
+	{
+		if (IsConnected(&Collection))
+		{
+			FManagedArrayCollection InCollection = GetValue(Context, &Collection);
+			GeometryCollection::Facades::FSelectionFacade SelectionFacade(InCollection);
+
+			const FDataflowOutput* SelectionOutput = FindOutput(&Selection);
+
+			FString InAttribute = GetValue(Context, &Attribute);
+
+			FName GroupName;
+			if (SelectionOutput)
+			{
+				using namespace GeometryCollection::Facades;
+
+				GroupName = UE::Dataflow::SelectionSet::Private::GetSelectionTypeFromCollection(*SelectionOutput);
+
+				TBitArray<> SelectionArray;
+				FSelectionFacade::EErrorCode ErrorCode = SelectionFacade.GetSelectionFromCollection(GroupName, InAttribute, SelectionArray);
+
+				if (ErrorCode == FSelectionFacade::EErrorCode::SelectionGroupCantBeFound)
+				{
+					SetError(Context, &Collection, TEXT("Collection doesn't contain selection(s)"));
+				}
+				else if (ErrorCode == FSelectionFacade::EErrorCode::AttributeCantBeFound)
+				{
+					SetError(Context, &Collection, TEXT("Collection doesn't have a selection with the specified attribute"));
+				}
+
+				FDataflowSelection OutSelection(GroupName);
+				OutSelection.SetFromBitArray(MoveTemp(SelectionArray));
+
+				SetValue(Context, MoveTemp(OutSelection), &Selection);
+
+				return;
+			}
+		}
+
+		SetValue(Context, FDataflowSelection(), &Selection);
+	}
+	else if (Out->IsA(&Collection))
+	{
+		SafeForwardInput(Context, &Collection, &Collection);
+	}
+}
 
diff --git a/Engine/Plugins/Experimental/Dataflow/Source/DataflowNodes/Private/Dataflow/DataflowSelectionNodes.h b/Engine/Plugins/Experimental/Dataflow/Source/DataflowNodes/Private/Dataflow/DataflowSelectionNodes.h
new file mode 100644
--- /dev/null
+++ b/Engine/Plugins/Experimental/Dataflow/Source/DataflowNodes/Private/Dataflow/DataflowSelectionNodes.h
@@ -0,0 +1,102 @@
+// Copyright Epic Games, Inc. All Rights Reserved.
+
+#pragma once 
+
+#include "CoreMinimal.h"
+#include "Dataflow/DataflowCore.h"
+#include "Dataflow/DataflowEngine.h"
+#include "GeometryCollection/ManagedArrayCollection.h"
+#include "Dataflow/DataflowSelection.h"
+#include "Dataflow/DataflowAnyType.h"
+
+#include "DataflowSelectionNodes.generated.h"
+
+#define UE_API DATAFLOWNODES_API
+
+USTRUCT(meta = (DataflowFlesh))
+struct FSelectionSetDataflowNode : public FDataflowNode
+{
+	GENERATED_USTRUCT_BODY()
+	DATAFLOW_NODE_DEFINE_INTERNAL(FSelectionSetDataflowNode, "SelectionSet", "Dataflow", "")
+
+public:
+	typedef TArray<int32> DataType;
+
+	UPROPERTY(EditAnywhere, Category = "Dataflow")
+		FString Indices = FString("1 2 3");
+
+	UPROPERTY(meta = (DataflowOutput, DisplayName = "SelectionSet"))
+		TArray<int32> IndicesOut;
+
+	FSelectionSetDataflowNode(const UE::Dataflow::FNodeParameters& InParam, FGuid InGuid = FGuid::NewGuid())
+		: FDataflowNode(InParam, InGuid)
+	{
+		RegisterOutputConnection(&IndicesOut);
+	}
+
+	UE_API virtual void Evaluate(UE::Dataflow::FContext& Context, const FDataflowOutput* Out) const override;
+};
+
+USTRUCT(meta = (GeometryCollection))
+struct FMakeSelectionSetDataflowNode : public FDataflowNode
+{
+	GENERATED_USTRUCT_BODY()
+	DATAFLOW_NODE_DEFINE_INTERNAL(FMakeSelectionSetDataflowNode, "MakeSelectionSet", "Dataflow", "")
+
+private:
+	/** Collection to use */
+	UPROPERTY(meta = (DataflowInput, DataflowOutput, DataflowPassthrough = "Collection", DataflowIntrinsic))
+	FManagedArrayCollection Collection;
+
+	/** Selection to store in collection */
+	UPROPERTY(meta = (DataflowInput, DisplayName = "Selection", DataflowIntrinsic))
+	FDataflowSelectionTypes Selection;
+
+	/** Attribute used to store the selection in the collection. The group decided by the incoming selection's type */
+	UPROPERTY(EditAnywhere, Category = "Attribute", meta = (DataflowInput))
+	FString Attribute = FString("Dummy");
+
+	/** Store the selection as float instead of bool */
+	UPROPERTY(EditAnywhere, Category = "Attribute")
+	bool bStoreAsFloat = false;
+
+	UE_API virtual void Evaluate(UE::Dataflow::FContext& Context, const FDataflowOutput* Out) const override;
+	virtual FAttributeKey GetVertexAttributeToVisualize(UE::Dataflow::FContext& Context, const FName OutputName, const FName RenderGroup) const override;
+
+public:
+	FMakeSelectionSetDataflowNode(const UE::Dataflow::FNodeParameters& InParam, FGuid InGuid = FGuid::NewGuid());
+};
+
+USTRUCT(meta = (GeometryCollection))
+struct FGetSelectionSetDataflowNode : public FDataflowNode
+{
+	GENERATED_USTRUCT_BODY()
+	DATAFLOW_NODE_DEFINE_INTERNAL(FGetSelectionSetDataflowNode, "GetSelectionSet", "Dataflow", "")
+
+private:
+	/** Collection to use */
+	UPROPERTY(meta = (DataflowInput, DataflowOutput, DataflowPassthrough = "Collection", DataflowIntrinsic))
+	FManagedArrayCollection Collection;
+
+	/** Attribute used to store the selection in the collection. 
+	To set the selection type rt mouse click on the Selection output */
+	UPROPERTY(EditAnywhere, Category = "Attribute", meta = (DataflowInput))
+	FString Attribute = FString("Dummy");
+
+	/** Output selection */
+	UPROPERTY(meta = (DataflowOutput))
+	FDataflowSelectionTypes Selection;
+
+	UE_API virtual void Evaluate(UE::Dataflow::FContext& Context, const FDataflowOutput* Out) const override;
+	virtual FAttributeKey GetVertexAttributeToVisualize(UE::Dataflow::FContext& Context, const FName OutputName, const FName RenderGroup) const override;
+
+public:
+	FGetSelectionSetDataflowNode(const UE::Dataflow::FNodeParameters& InParam, FGuid InGuid = FGuid::NewGuid());
+};
+
+namespace UE::Dataflow
+{
+	void RegisterSelectionNodes();
+}
+
+#undef UE_API
diff --git a/Engine/Plugins/Experimental/Dataflow/Source/DataflowNodes/Public/Dataflow/DataflowSelectionNodes.h b/Engine/Plugins/Experimental/Dataflow/Source/DataflowNodes/Public/Dataflow/DataflowSelectionNodes.h
deleted file mode 100644
diff --git a/Engine/Source/Runtime/Experimental/Chaos/Private/GeometryCollection/Facades/CollectionSelectionFacade.cpp b/Engine/Source/Runtime/Experimental/Chaos/Private/GeometryCollection/Facades/CollectionSelectionFacade.cpp
--- a/Engine/Source/Runtime/Experimental/Chaos/Private/GeometryCollection/Facades/CollectionSelectionFacade.cpp
+++ b/Engine/Source/Runtime/Experimental/Chaos/Private/GeometryCollection/Facades/CollectionSelectionFacade.cpp
@@ -5,6 +5,8 @@
 =============================================================================*/
 
 #include "GeometryCollection/Facades/CollectionSelectionFacade.h"
+#include "GeometryCollection/TransformCollection.h"
+#include "GeometryCollection/GeometryCollection.h"
 
 namespace GeometryCollection::Facades
 {
@@ -190,6 +192,165 @@ namespace GeometryCollection::Facades
 		}
 	}
 
-};
+	// ------------------------------------------------------------------------------------------------------------------------------------
+
+	void FSelectionFacade::AddSelectionToCollection(const FName InGroupName,
+		const TBitArray<>& InSelectionArray,
+		const FString InAttribute,
+		const bool bStoreAsFloat)
+	{
+		if (Collection)
+		{
+			const FName AttrName = FName(*InAttribute);
+
+			if (Collection->NumElements(InGroupName) == InSelectionArray.Num())
+			{
+				if (!Collection->HasAttribute(AttrName, InGroupName))
+				{
+					if (bStoreAsFloat)
+					{
+						TManagedArray<float>& Attr = Collection->AddAttribute<float>(AttrName, InGroupName);
+
+						for (int32 Idx = 0; Idx < InSelectionArray.Num(); ++Idx)
+						{
+							Attr[Idx] = InSelectionArray[Idx] ? 1.f : 0.f;
+						}
+					}
+					else
+					{
+						TManagedArray<bool>& Attr = Collection->AddAttribute<bool>(AttrName, InGroupName);
+
+						for (int32 Idx = 0; Idx < InSelectionArray.Num(); ++Idx)
+						{
+							Attr[Idx] = InSelectionArray[Idx];
+						}
+					}
+
+					RegisterSelectionInCollection(InGroupName, InAttribute);
+				}
+			}
+		}
+	}
+
+	// ------------------------------------------------------------------------------------------------------------------------------------
+
+	void FSelectionFacade::RegisterSelectionInCollection(const FName InGroupName, FString InAttribute)
+	{
+		static const FName SelectionsGroup = "Selections";
+		static const FName SelectionName = "Name";
+		static const FName TypeName = "Type";
+		
+		if (Collection)
+		{
+			Collection->AddElements(1, SelectionsGroup);
+
+			const FName AttrName = FName(*InAttribute);
+		
+			if (!Collection->HasAttribute(SelectionName, SelectionsGroup))
+			{
+				TManagedArray<FString>& NameAttr = Collection->AddAttribute<FString>(SelectionName, SelectionsGroup);
+			}
+			if (!Collection->HasAttribute(TypeName, SelectionsGroup))
+			{
+				TManagedArray<FString>& TypeAttr = Collection->AddAttribute<FString>(TypeName, SelectionsGroup);
+			}
+
+			TManagedArray<FString>* NameAttr = Collection->FindAttributeTyped<FString>(SelectionName, SelectionsGroup);
+			TManagedArray<FString>* TypeAttr = Collection->FindAttributeTyped<FString>(TypeName, SelectionsGroup);
+
+			if (NameAttr && TypeAttr)
+			{
+				const int32 NumElems = NameAttr->Num();
+
+				(*NameAttr)[NumElems - 1] = InAttribute;
+				(*TypeAttr)[NumElems - 1] = InGroupName.ToString();
+			}
+		}
+	}
+
+	// ------------------------------------------------------------------------------------------------------------------------------------
 
+	FSelectionFacade::EErrorCode FSelectionFacade::GetSelectionFromCollection(const FName InGroupName, const FString& InAttribute, TBitArray<>& OutSelectionArray) const
+	{
+		static const FName SelectionsGroup = "Selections";
+		static const FName SelectionName = "Name";
+		static const FName TypeName = "Type";
 
+		if (ConstCollection.HasGroup(SelectionsGroup))
+		{
+			const FName AttrName = FName(*InAttribute);
+
+			const TManagedArray<FString>* NameAttr = ConstCollection.FindAttributeTyped<FString>(SelectionName, SelectionsGroup);
+			const TManagedArray<FString>* TypeAttr = ConstCollection.FindAttributeTyped<FString>(TypeName, SelectionsGroup);
+
+			if (NameAttr && TypeAttr && ConstCollection.HasAttribute(AttrName, InGroupName))
+			{
+				FManagedArrayCollection::EArrayType AttrType = ConstCollection.GetAttributeType(AttrName, InGroupName);
+
+				const TManagedArray<bool>* BoolAttr = nullptr;
+				const TManagedArray<float>* FloatAttr = nullptr;
+
+				if (AttrType == FManagedArrayCollection::EArrayType::FBoolType)
+				{
+					BoolAttr = ConstCollection.FindAttributeTyped<bool>(AttrName, InGroupName);
+				}
+				else if (AttrType == FManagedArrayCollection::EArrayType::FFloatType)
+				{
+					FloatAttr = ConstCollection.FindAttributeTyped<float>(AttrName, InGroupName);
+				}
+
+				bool bAttrFound = false;
+				for (int32 Idx = 0; Idx < NameAttr->Num(); ++Idx)
+				{
+					if ((*NameAttr)[Idx] == InAttribute && (*TypeAttr)[Idx] == InGroupName.ToString())
+					{
+						bAttrFound = true;
+
+						if (BoolAttr)
+						{
+							OutSelectionArray = BoolAttr->GetConstArray();
+
+							return EErrorCode::None;
+						}
+						else if (FloatAttr)
+						{
+							const int32 NumElems = FloatAttr->Num();
+
+							OutSelectionArray.Init(false, NumElems);
+
+							for (int32 Idx2 = 0; Idx2 < NumElems; ++Idx2)
+							{
+								if ((*FloatAttr)[Idx2] > 0.f)
+								{
+									OutSelectionArray[Idx2] = true;
+								}
+							}
+
+							return EErrorCode::None;
+						}
+						else
+						{
+							return EErrorCode::AttributeCantBeFound;
+						}
+					}
+				}
+
+				if (!bAttrFound)
+				{
+					return EErrorCode::AttributeCantBeFound;
+				}
+			}
+			else
+			{
+				return EErrorCode::AttributeCantBeFound;
+			}
+		}
+		else
+		{
+			return EErrorCode::SelectionGroupCantBeFound;
+		}
+
+		return EErrorCode::None;
+	}
+
+};
diff --git a/Engine/Source/Runtime/Experimental/Chaos/Public/GeometryCollection/Facades/CollectionSelectionFacade.h b/Engine/Source/Runtime/Experimental/Chaos/Public/GeometryCollection/Facades/CollectionSelectionFacade.h
--- a/Engine/Source/Runtime/Experimental/Chaos/Public/GeometryCollection/Facades/CollectionSelectionFacade.h
+++ b/Engine/Source/Runtime/Experimental/Chaos/Public/GeometryCollection/Facades/CollectionSelectionFacade.h
@@ -45,6 +45,12 @@ namespace GeometryCollection::Facades
 		CHAOS_API void InitWeightedBoundedGroup(FName GroupName, FName DependencyGroup, FName BoneDependencyGroup);
 
 	public:
+		enum class EErrorCode
+		{
+			None = 0,
+			SelectionGroupCantBeFound = 1,
+			AttributeCantBeFound = 2
+		};
 
 		// groups
 		static CHAOS_API const FName UnboundGroup;
@@ -126,12 +132,17 @@ namespace GeometryCollection::Facades
 		*/
 		CHAOS_API void GetSelection(const FSelectionKey& Key, int32& BoneIndex, TArray<int32>& Indices, TArray<float>& Weights) const;
 
-	private:
+		CHAOS_API void AddSelectionToCollection(const FName InGroupName, const TBitArray<>& InSelectionArray, const FString InAttribute, const bool bStoreAsFloat);
+
+		CHAOS_API EErrorCode GetSelectionFromCollection(const FName InGroupName, const FString& InAttribute, TBitArray<>& OutSelectionArray) const;
 
+	private:
 		// const collection will be a null pointer, 
 		// while non-const will be valid.
 		const FManagedArrayCollection& ConstCollection;
 		FManagedArrayCollection* Collection = nullptr;
+
+		void RegisterSelectionInCollection(const FName InGroupName, FString InAttribute);
 	};
 
 }
diff --git a/Engine/Source/Runtime/Experimental/Chaos/Public/GeometryCollection/Facades/PointsCollectionFacade.h b/Engine/Source/Runtime/Experimental/Chaos/Public/GeometryCollection/Facades/PointsCollectionFacade.h
--- a/Engine/Source/Runtime/Experimental/Chaos/Public/GeometryCollection/Facades/PointsCollectionFacade.h
+++ b/Engine/Source/Runtime/Experimental/Chaos/Public/GeometryCollection/Facades/PointsCollectionFacade.h
@@ -18,7 +18,7 @@ namespace GeometryCollection::Facades
 	{
 	public:
 		/**
-		* FSelectionFacade Constuctor
+		* FPointsCollectionFacade Constuctor
 		*/
 		CHAOS_API FPointsCollectionFacade(FManagedArrayCollection& InSelf);
 		CHAOS_API FPointsCollectionFacade(const FManagedArrayCollection& InSelf);
diff --git a/Engine/Source/Runtime/Experimental/Dataflow/Core/Private/Dataflow/DataflowSelection.cpp b/Engine/Source/Runtime/Experimental/Dataflow/Core/Private/Dataflow/DataflowSelection.cpp
--- a/Engine/Source/Runtime/Experimental/Dataflow/Core/Private/Dataflow/DataflowSelection.cpp
+++ b/Engine/Source/Runtime/Experimental/Dataflow/Core/Private/Dataflow/DataflowSelection.cpp
@@ -16,6 +16,11 @@ void FDataflowSelection::Initialize(const FDataflowSelection& Other)
 	SelectionArray = Other.SelectionArray;
 }
 
+void FDataflowSelection::Initialize(const TBitArray<>& Other)
+{
+	SelectionArray = Other;
+}
+
 void FDataflowSelection::InitializeFromCollection(const FManagedArrayCollection& InCollection, bool Value)
 {
 	SelectionArray.Init(Value, InCollection.NumElements(GroupName));
@@ -26,6 +31,11 @@ void FDataflowSelection::Clear()
 	SelectionArray.Init(false, SelectionArray.Num());
 }
 
+int32 FDataflowSelection::Num() const 
+{ 
+	return SelectionArray.Num(); 
+}
+
 void FDataflowSelection::AsArrayValidated(TArray<int32>& OutSelectionArr, const FManagedArrayCollection& InCollection) const
 {
 	if (IsValidForCollection(InCollection))
@@ -54,6 +64,11 @@ TArray<int32> FDataflowSelection::AsArrayValidated(const FManagedArrayCollection
 	return SelectionArr;
 }
 
+void FDataflowSelection::SetGroupName(const FName& InGroupName)
+{
+	GroupName = InGroupName;
+}
+
 void FDataflowSelection::AsArray(TArray<int32>& SelectionArr) const
 {
 	SelectionArr.Reset(SelectionArray.Num());
@@ -194,7 +209,7 @@ void FDataflowSelection::SetWithMask(const bool Value, const FDataflowSelection&
 	}
 }
 
-void FDataflowSelection::SetSelected(TArray<int32> Indices)
+void FDataflowSelection::SetSelected(const TArray<int32>& Indices)
 {
 	for (int32 Idx = 0; Idx < Indices.Num(); ++Idx)
 	{
@@ -205,6 +220,23 @@ void FDataflowSelection::SetSelected(TArray<int32> Indices)
 	}
 }
 
+bool FDataflowSelection::SetSelectedWithCheck(const TArray<int32>& Indices)
+{
+	for (int32 Idx = 0; Idx < Indices.Num(); ++Idx)
+	{
+		if (SelectionArray.IsValidIndex(Indices[Idx]))
+		{
+			SetSelected(Indices[Idx]);
+		}
+		else
+		{
+			return false;
+		}
+	}
+
+	return true;
+}
+
 bool FDataflowSelection::IsValidForCollection(const FManagedArrayCollection& InCollection) const
 {
 	return !GroupName.IsNone() && Num() == InCollection.NumElements(GroupName);
diff --git a/Engine/Source/Runtime/Experimental/Dataflow/Core/Public/Dataflow/DataflowAnyType.h b/Engine/Source/Runtime/Experimental/Dataflow/Core/Public/Dataflow/DataflowAnyType.h
--- a/Engine/Source/Runtime/Experimental/Dataflow/Core/Public/Dataflow/DataflowAnyType.h
+++ b/Engine/Source/Runtime/Experimental/Dataflow/Core/Public/Dataflow/DataflowAnyType.h
@@ -4,6 +4,7 @@
 
 #include "CoreMinimal.h"
 #include "Dataflow/DataflowTypePolicy.h"
+#include "Misc/TVariant.h"
 
 #include "DataflowAnyType.generated.h"
 
@@ -46,6 +47,20 @@ struct FDataflowUStructTypes : public FDataflowAnyType
 	GENERATED_BODY()
 };
 
+UENUM()
+enum class EDataflowNumericType : uint8
+{
+	Double UMETA(DisplayName = "double"),
+	Float UMETA(DisplayName = "float"),
+	Int64 UMETA(DisplayName = "int64"),
+	UInt64 UMETA(DisplayName = "uint64"),
+	Int32 UMETA(DisplayName = "int32"),
+	UInt32 UMETA(DisplayName = "uint32"),
+	Int16 UMETA(DisplayName = "int16"),
+	UInt16 UMETA(DisplayName = "uint16"),
+	Int8 UMETA(DisplayName = "int8"),
+	UInt8 UMETA(DisplayName = "uint8")
+};
 
 /**
 * Numeric types
diff --git a/Engine/Source/Runtime/Experimental/Dataflow/Core/Public/Dataflow/DataflowSelection.h b/Engine/Source/Runtime/Experimental/Dataflow/Core/Public/Dataflow/DataflowSelection.h
--- a/Engine/Source/Runtime/Experimental/Dataflow/Core/Public/Dataflow/DataflowSelection.h
+++ b/Engine/Source/Runtime/Experimental/Dataflow/Core/Public/Dataflow/DataflowSelection.h
@@ -22,13 +22,14 @@ struct FDataflowSelection
 
 	DATAFLOWCORE_API void Initialize(int32 NumBits, bool Value);
 	DATAFLOWCORE_API void Initialize(const FDataflowSelection& Other);
+	DATAFLOWCORE_API void Initialize(const TBitArray<>& Other);
 	DATAFLOWCORE_API void InitializeFromCollection(const FManagedArrayCollection& InCollection, bool Value);
 	DATAFLOWCORE_API void Clear();
-	int32 Num() const { return SelectionArray.Num(); }
+	DATAFLOWCORE_API int32 Num() const;
 	DATAFLOWCORE_API int32 NumSelected() const;
 	DATAFLOWCORE_API bool AnySelected() const;
 
-	bool IsValid() const;
+	DATAFLOWCORE_API bool IsValid() const;
 
 	inline bool IsValidIndex(int32 Idx) const
 	{
@@ -36,13 +37,16 @@ struct FDataflowSelection
 	}
 	bool IsSelected(int32 Idx) const { return SelectionArray[Idx]; }
 	void SetSelected(int32 Idx) { SelectionArray[Idx] = true; }
-	DATAFLOWCORE_API void SetSelected(TArray<int32> Indices);
+	DATAFLOWCORE_API void SetSelected(const TArray<int32>& Indices);
+	DATAFLOWCORE_API bool SetSelectedWithCheck(const TArray<int32>& Indices);
 	void SetNotSelected(int32 Idx) { SelectionArray[Idx] = false; }
 	DATAFLOWCORE_API void AsArray(TArray<int32>& SelectionArr) const;
 	DATAFLOWCORE_API TArray<int32> AsArray() const;
 	DATAFLOWCORE_API void AsArrayValidated(TArray<int32>& SelectionArr, const FManagedArrayCollection& InCollection) const;
 	DATAFLOWCORE_API TArray<int32> AsArrayValidated(const FManagedArrayCollection& InCollection) const;
 	
+	DATAFLOWCORE_API void SetGroupName(const FName& InGroupName);
+
 	/**
 	 * Sets the selection from a sparse array (it only contains the indices of the selected items)
 	 * @return true if all indices were successfully set, false if some indices were invalid and could not be set
@@ -160,6 +164,19 @@ struct FDataflowCurveSelection : public FDataflowSelection
 	{}
 };
 
+USTRUCT()
+struct FDataflowPointsSelection : public FDataflowSelection
+{
+	GENERATED_USTRUCT_BODY()
+
+	inline static const FName PointsGroupName = "Points";
+
+	FDataflowPointsSelection()
+		: FDataflowSelection(PointsGroupName)
+	{
+	}
+};
+
 UENUM(BlueprintType)
 enum class EDataflowSelectionType : uint8
 {
@@ -168,5 +185,7 @@ enum class EDataflowSelectionType : uint8
 	Faces UMETA(DisplayName = "Faces"),
 	Geometry UMETA(DisplayName = "Geometry"),
 	Curves UMETA(DisplayName = "Curves"),
+	Points UMETA(DisplayName = "Points"),
+	Material UMETA(DisplayName = "Material"),
 };
 
diff --git a/Engine/Source/Runtime/Experimental/Dataflow/Core/Public/Dataflow/DataflowTypePolicy.h b/Engine/Source/Runtime/Experimental/Dataflow/Core/Public/Dataflow/DataflowTypePolicy.h
--- a/Engine/Source/Runtime/Experimental/Dataflow/Core/Public/Dataflow/DataflowTypePolicy.h
+++ b/Engine/Source/Runtime/Experimental/Dataflow/Core/Public/Dataflow/DataflowTypePolicy.h
@@ -93,6 +93,8 @@ UE_DATAFLOW_POLICY_DECLARE_TYPENAME(FIntVector3)
 UE_DATAFLOW_POLICY_DECLARE_TYPENAME(FIntVector4)
 UE_DATAFLOW_POLICY_DECLARE_TYPENAME(FRotator)
 UE_DATAFLOW_POLICY_DECLARE_TYPENAME(FTransform)
+UE_DATAFLOW_POLICY_DECLARE_TYPENAME(FBox)
+UE_DATAFLOW_POLICY_DECLARE_TYPENAME(FSphere)
 
 UE_DATAFLOW_POLICY_DECLARE_TYPENAME(FDataflowSelection)
 UE_DATAFLOW_POLICY_DECLARE_TYPENAME(FDataflowTransformSelection)
@@ -101,6 +103,7 @@ UE_DATAFLOW_POLICY_DECLARE_TYPENAME(FDataflowFaceSelection)
 UE_DATAFLOW_POLICY_DECLARE_TYPENAME(FDataflowGeometrySelection)
 UE_DATAFLOW_POLICY_DECLARE_TYPENAME(FDataflowMaterialSelection)
 UE_DATAFLOW_POLICY_DECLARE_TYPENAME(FDataflowCurveSelection)
+UE_DATAFLOW_POLICY_DECLARE_TYPENAME(FDataflowPointsSelection)
 
 struct IDataflowTypePolicy
 {
@@ -387,7 +390,7 @@ struct FDataflowStringTypePolicy :
 };
 
 struct FDataflowSelectionTypePolicy :
-	public TDataflowMultiTypePolicy<FDataflowTransformSelection, FDataflowVertexSelection, FDataflowFaceSelection, FDataflowGeometrySelection, FDataflowMaterialSelection, FDataflowCurveSelection>
+	public TDataflowMultiTypePolicy<FDataflowTransformSelection, FDataflowVertexSelection, FDataflowFaceSelection, FDataflowGeometrySelection, FDataflowMaterialSelection, FDataflowCurveSelection, FDataflowPointsSelection>
 {
 };
 
@@ -771,3 +774,4 @@ struct FDataflowConverter<FRotator>
 		}
 	}
 };
+
