Dataflow editor : make async eval and show perf data  option project user settings

#Synced-CL 51402414

diff --git a/Engine/Plugins/Experimental/Dataflow/Source/DataflowEditor/Private/Dataflow/DataflowEditorModule.cpp b/Engine/Plugins/Experimental/Dataflow/Source/DataflowEditor/Private/Dataflow/DataflowEditorModule.cpp
--- a/Engine/Plugins/Experimental/Dataflow/Source/DataflowEditor/Private/Dataflow/DataflowEditorModule.cpp
+++ b/Engine/Plugins/Experimental/Dataflow/Source/DataflowEditor/Private/Dataflow/DataflowEditorModule.cpp
@@ -53,6 +53,7 @@
 #include "Styling/SlateIconFinder.h"
 #include "Widgets/Images/SImage.h"
 #include "AssetViewWidgets.h"
+#include "ISettingsModule.h"
 
 #define LOCTEXT_NAMESPACE "DataflowEditor"
 
@@ -287,6 +288,19 @@ void FDataflowEditorModule::StartupModule()
 {
 	using namespace UE::Dataflow::Private;
 
+#if WITH_EDITOR
+	// register settings
+	if(ISettingsModule* SettingsModule = FModuleManager::GetModulePtr<ISettingsModule>("Settings"))
+	{
+		SettingsModule->RegisterSettings("Editor", "ContentEditors", "DataflowEditor",
+			LOCTEXT("DataflowEditorLabel", "Dataflow Editor"),
+			LOCTEXT("DataflowEditorDesc", "Dataflow Editor Settings"),
+			GetMutableDefault<UDataflowEditorUserSettings>()
+		);
+	}
+#endif //WITH_EDITOR
+
+
 	FDataflowEditorStyle::Get();
 	
 	// Register type customizations
@@ -397,6 +411,11 @@ void FDataflowEditorModule::ShutdownModule()
 	UE::Dataflow::FDataflowConstructionVisualizationRegistry& ConstructionVisualizationRegistry = UE::Dataflow::FDataflowConstructionVisualizationRegistry::GetInstance();
 	ConstructionVisualizationRegistry.DeregisterVisualization(UE::Dataflow::FMeshStatsConstructionVisualization::Name);
 	ConstructionVisualizationRegistry.DeregisterVisualization(UE::Dataflow::FDebugDrawText3dVisualization::Name);
+
+	if (ISettingsModule* SettingsModule = FModuleManager::GetModulePtr<ISettingsModule>("Settings"))
+	{
+		SettingsModule->UnregisterSettings("Editor", "ContentEditors", "DataflowEditor");
+	}
 }
 
 IMPLEMENT_MODULE(FDataflowEditorModule, DataflowEditor)
diff --git a/Engine/Plugins/Experimental/Dataflow/Source/DataflowEditor/Private/Dataflow/DataflowEditorToolkit.cpp b/Engine/Plugins/Experimental/Dataflow/Source/DataflowEditor/Private/Dataflow/DataflowEditorToolkit.cpp
--- a/Engine/Plugins/Experimental/Dataflow/Source/DataflowEditor/Private/Dataflow/DataflowEditorToolkit.cpp
+++ b/Engine/Plugins/Experimental/Dataflow/Source/DataflowEditor/Private/Dataflow/DataflowEditorToolkit.cpp
@@ -91,10 +91,6 @@ bool bDataflowForceGraphEditable = false;
 FAutoConsoleVariableRef CVARbDataflowForceGraphEditable(TEXT("p.Dataflow.Editor.ForceGraphEditable"), bDataflowForceGraphEditable,
 	TEXT("when enabled all graph will be editable and the Unlock/lock editing buttons will not be visible.[def:false]"));
 
-bool bDataflowEnableNewEvaluation = true;
-FAutoConsoleVariableRef CVARDataflowEnableNewEvaluation(TEXT("p.Dataflow.Editor.EnableNewEvaluation"), bDataflowEnableNewEvaluation,
-	TEXT("When enable, use the new evaluation system using evaluators directly - it is bound to the async evaluation option in the dataflow editor [def:true]"));
-
 const FName FDataflowEditorToolkit::GraphCanvasTabId(TEXT("DataflowEditor_GraphCanvas"));
 const FName FDataflowEditorToolkit::SubGraphCanvasTabId(TEXT("DataflowEditor_SubGraphTab"));
 const FName FDataflowEditorToolkit::NodeDetailsTabId(TEXT("DataflowEditor_NodeDetails"));
@@ -821,54 +817,40 @@ void FDataflowEditorToolkit::SetEvaluateGraphMode(EDataflowEditorEvaluationMode
 
 void FDataflowEditorToolkit::TogglePerfData()
 {
-	if (TObjectPtr<UDataflowBaseContent>& EditorContent = GetEditorContent())
+	if (UDataflowEditorUserSettings* Settings = GetMutableDefault<UDataflowEditorUserSettings>())
 	{
-		if (TSharedPtr<UE::Dataflow::FEngineContext> DataflowContext = EditorContent->GetDataflowContext())
-		{
-			DataflowContext->EnablePerfData(!DataflowContext->IsPerfDataEnabled());
-		}
+		Settings->bShowPerformanceData = !Settings->bShowPerformanceData;
+		Settings->PostEditChange();
+		Settings->SaveConfig();
 	}
 }
 
 bool FDataflowEditorToolkit::IsPerfDataEnabled() const
 {
-	if (const TObjectPtr<UDataflowBaseContent>& EditorContent = GetEditorContent())
+	if (UDataflowEditorUserSettings* Settings = GetMutableDefault<UDataflowEditorUserSettings>())
 	{
-		if (TSharedPtr<UE::Dataflow::FEngineContext> DataflowContext = EditorContent->GetDataflowContext())
-		{
-			return DataflowContext->IsPerfDataEnabled();
-		}
+		return Settings->bShowPerformanceData;
 	}
 	return false;
 }
 
 void FDataflowEditorToolkit::ToggleAsyncEvaluation()
 {
-	bDataflowEnableNewEvaluation = !bDataflowEnableNewEvaluation;
-	//if (TObjectPtr<UDataflowBaseContent>& EditorContent = GetEditorContent())
-	//{
-	//	if (TSharedPtr<UE::Dataflow::FEngineContext> DataflowContext = EditorContent->GetDataflowContext())
-	//	{
-	//		if (DataflowContext->IsThreaded())
-	//		{
-	//			DataflowContext->CancelAsyncEvaluation();
-	//		}
-	//		DataflowContext->SetThreaded(!DataflowContext->IsThreaded());
-	//	}
-	//}
+	if (UDataflowEditorUserSettings* Settings = GetMutableDefault<UDataflowEditorUserSettings>())
+	{
+		Settings->bUseAsyncEvaluation = !Settings->bUseAsyncEvaluation;
+		Settings->PostEditChange();
+		Settings->SaveConfig();
+	}
 }
 
 bool FDataflowEditorToolkit::IsAsyncEvaluationEnabled() const
 {
-	//if (const TObjectPtr<UDataflowBaseContent>& EditorContent = GetEditorContent())
-	//{
-	//	if (TSharedPtr<UE::Dataflow::FEngineContext> DataflowContext = EditorContent->GetDataflowContext())
-	//	{
-	//		return DataflowContext->IsThreaded();
-	//	}
-	//}
-	//return false;
-	return bDataflowEnableNewEvaluation;
+	if (UDataflowEditorUserSettings* Settings = GetMutableDefault<UDataflowEditorUserSettings>())
+	{
+		return Settings->bUseAsyncEvaluation;
+	}
+	return false;
 }
 
 void FDataflowEditorToolkit::ClearGraphCache()
@@ -2253,6 +2235,14 @@ void FDataflowEditorToolkit::Tick(float DeltaTime)
 				EditorContent->SetDataflowContext(MakeShared<UE::Dataflow::FEngineContext>(EditorContent->GetDataflowOwner()));
 				InitTimeStamp = UE::Dataflow::FTimestamp::Invalid;
 			}
+			// Sync the global settings
+			if (UDataflowEditorUserSettings* Settings = GetMutableDefault<UDataflowEditorUserSettings>())
+			{
+				if (TSharedPtr<UE::Dataflow::FEngineContext> Context = EditorContent->GetDataflowContext())
+				{
+					Context->EnablePerfData(Settings->bShowPerformanceData);
+				}
+			}
 
 			// Update the list of dataflow terminal contents 
 			DataflowEditor->UpdateTerminalContents(InitTimeStamp);
@@ -2337,8 +2327,12 @@ void FDataflowEditorToolkit::EvaluateNode(const FDataflowNode* Node, const FData
 								ToolKit->bViewsNeedRefresh = true;
 							}
 						};
-
-					if (bDataflowEnableNewEvaluation)
+					bool bUseAsyncEvaluation = false;
+					if (UDataflowEditorUserSettings* Settings = GetMutableDefault<UDataflowEditorUserSettings>())
+					{
+						bUseAsyncEvaluation = Settings->bUseAsyncEvaluation;
+					}
+					if (bUseAsyncEvaluation)
 					{
 						// TODO : maybe this should be in the commands EvaluateNode as there's logic with the timestamp
 						EditorContent->EvaluateGraph(OnEvaluationCompleted);
diff --git a/Engine/Plugins/Experimental/Dataflow/Source/DataflowEditor/Public/Dataflow/DataflowEditor.h b/Engine/Plugins/Experimental/Dataflow/Source/DataflowEditor/Public/Dataflow/DataflowEditor.h
--- a/Engine/Plugins/Experimental/Dataflow/Source/DataflowEditor/Public/Dataflow/DataflowEditor.h
+++ b/Engine/Plugins/Experimental/Dataflow/Source/DataflowEditor/Public/Dataflow/DataflowEditor.h
@@ -33,6 +33,21 @@ class UDataflowEditorSettings : public UObject
 	GENERATED_BODY()
 };
 
+UCLASS(MinimalAPI, config = EditorPerProjectUserSettings)
+class UDataflowEditorUserSettings : public UObject
+{
+	GENERATED_BODY()
+
+public:
+	/** When true use the task graph based evaluation system */
+	UPROPERTY(config, EditAnywhere, Category = "Dataflow")
+	bool bUseAsyncEvaluation = false;
+
+	/** Whether to show the performance data on top of the nodes */
+	UPROPERTY(config, EditAnywhere, Category = "Dataflow")
+	bool bShowPerformanceData = false;
+};
+
 /** 
  * The actual asset editor class doesn't have that much in it, intentionally. 
  * 
