[ModelingTools Exp] Extract most of data flow vert paint tool into a sharable base class

diff --git a/Engine/Plugins/Experimental/Dataflow/Source/DataflowEditor/Private/DataflowEditorTools/DataflowEditorVertexAttributePaintTool.cpp b/Engine/Plugins/Experimental/Dataflow/Source/DataflowEditor/Private/DataflowEditorTools/DataflowEditorVertexAttributePaintTool.cpp
--- a/Engine/Plugins/Experimental/Dataflow/Source/DataflowEditor/Private/DataflowEditorTools/DataflowEditorVertexAttributePaintTool.cpp
+++ b/Engine/Plugins/Experimental/Dataflow/Source/DataflowEditor/Private/DataflowEditorTools/DataflowEditorVertexAttributePaintTool.cpp
@@ -259,118 +259,17 @@ UMeshSurfacePointTool* UDataflowEditorVertexAttributePaintToolBuilder::CreateNew
 
 ///////////////////////////////////////////////////////////////////////////////////////////////////////////////////////////////////////////////////
 //
-// Properties
+// Data Adapter
 //
 ///////////////////////////////////////////////////////////////////////////////////////////////////////////////////////////////////////////////////
-
-UDataflowEditorVertexAttributePaintToolProperties::UDataflowEditorVertexAttributePaintToolProperties()
-	: UInteractiveToolPropertySet()
-{
-	LoadConfig();
-
-	if (DisplayProperties.ColorRamp.IsEmpty())
-	{
-		CreateDefaultColorRamp();
-	}
-
-	constexpr bool bOnlyRGB = true;
-	DisplayProperties.GreyScaleColorRamp.SetColorAtTime(0.0f, FLinearColor::Black, bOnlyRGB);
-	DisplayProperties.GreyScaleColorRamp.SetColorAtTime(1.0f, FLinearColor::White, bOnlyRGB);
-
-	DisplayProperties.WhiteColorRamp.SetColorAtTime(0.0f, FLinearColor::White, bOnlyRGB);
-	DisplayProperties.WhiteColorRamp.SetColorAtTime(1.0f, FLinearColor::White, bOnlyRGB);
-}
-
-void UDataflowEditorVertexAttributePaintToolProperties::CreateDefaultColorRamp()
-{
-	ResetColorRamp();
-}
-
-void UDataflowEditorVertexAttributePaintToolProperties::ResetColorRamp()
-{
-	constexpr FLinearColor HeatMapColors[] =
-	{
-		FLinearColor(0.8f, 0.4f, 0.8f), // Purple
-		FLinearColor(0.0f, 0.0f, 0.5f), // Dark Blue
-		FLinearColor(0.2f, 0.2f, 1.0f), // Light Blue
-		FLinearColor(0.0f, 1.0f, 0.0f), // Green
-		FLinearColor(1.0f, 1.0f, 0.0f), // Yellow
-		FLinearColor(1.0f, 0.65f, 0.0f), // Orange
-		FLinearColor(1.0f, 0.0f, 0.0f, 0.0f), // Red
-	};
-	constexpr int32 NumColors = sizeof(HeatMapColors) / sizeof(FLinearColor);
-	ensure(NumColors > 1);
-
-	// Set ColorRamp
-	DisplayProperties.ColorRamp.Reset();
-	for (int32 Index = 0; Index < NumColors; ++Index)
-	{
-		const float Time = (float)Index / ((float)NumColors - 1.f);
-		DisplayProperties.ColorRamp.SetColorAtTime(Time, HeatMapColors[Index], /*bOnlyRGB*/true);
-	}
-	DisplayProperties.ColorRamp.GetCurves()[3].CurveToEdit->AddKey(0.f, 1.f); // single alpha  value
-}
-
-const FDataflowEditorVertexAttributePaintToolBrushConfig& UDataflowEditorVertexAttributePaintToolProperties::GetBrushConfig(EDataflowEditorToolEditOperation BrushMode) const
-{
-	static FDataflowEditorVertexAttributePaintToolBrushConfig DefaultConfig;
-
-	switch (BrushMode)
-	{
-	case EDataflowEditorToolEditOperation::Add:
-		return BrushConfigAdd;
-
-	case EDataflowEditorToolEditOperation::Replace:
-		return BrushConfigReplace;
-
-	case EDataflowEditorToolEditOperation::Multiply:
-		return BrushConfigMultiply;
-
-	case EDataflowEditorToolEditOperation::Relax:
-		return BrushConfigRelax;
-
-	case EDataflowEditorToolEditOperation::Invert:
-	default:
-		return DefaultConfig;
-	};
-}
-
-void UDataflowEditorVertexAttributePaintToolProperties::SetBrushConfig(EDataflowEditorToolEditOperation BrushMode, const FDataflowEditorVertexAttributePaintToolBrushConfig& BrushConfig)
-{
-	switch (BrushMode)
-	{
-	case EDataflowEditorToolEditOperation::Add:
-		BrushConfigAdd = BrushConfig;
-		break;
-
-	case EDataflowEditorToolEditOperation::Replace:
-		BrushConfigReplace = BrushConfig;
-		break;
-
-	case EDataflowEditorToolEditOperation::Multiply:
-		BrushConfigMultiply = BrushConfig;
-		break;
-
-	case EDataflowEditorToolEditOperation::Relax:
-		BrushConfigRelax = BrushConfig;
-		break;
-
-	case EDataflowEditorToolEditOperation::Invert:
-	default:
-		// nothing to update
-		break;
-	}
-}
-
-///////////////////////////////////////////////////////////////////////////////////////////////////////////////////////////////////////////////////
-//
-// Data
-//
-///////////////////////////////////////////////////////////////////////////////////////////////////////////////////////////////////////////////////
-void FDataflowEditorVertexAttributePaintToolData::Setup(FDynamicMesh3& InMesh, FDataflowVertexAttributeEditableNode* InNodeToUpdate, TObjectPtr<UDataflowContextObject> InDataflowEditorContextObject)
-{
-	// this should only be called once or have shutdown called in between 
-	ensure(ActiveWeightMap == nullptr);
+void FDataflowEditorVertexAttributePaintToolDataAdapter::Setup(
+	const FDynamicMesh3& InMesh,
+	UE::Geometry::FDynamicMeshWeightAttribute* InActiveWeightMap,
+	FDataflowVertexAttributeEditableNode* InNodeToUpdate,
+	TObjectPtr<UDataflowContextObject> InDataflowEditorContextObject
+	)
+{
+	ActiveWeightMap = InActiveWeightMap;
 	NodeToUpdate = InNodeToUpdate;
 
 	if (InDataflowEditorContextObject)
@@ -430,12 +329,6 @@ void FDataflowEditorVertexAttributePaintToolData::Setup(FDynamicMesh3& InMesh, F
 			}
 		}
 
-		// Create an attribute layer to temporarily paint into
-		const int NumAttributeLayers = InMesh.Attributes()->NumWeightLayers();
-		InMesh.Attributes()->SetNumWeightLayers(NumAttributeLayers + 1);
-		ActiveWeightMap = InMesh.Attributes()->GetWeightLayer(NumAttributeLayers);
-		ActiveWeightMap->SetName(FName("PaintLayer"));
-
 		if (NodeToUpdate)
 		{
 			// Setup DynamicMeshToWeight conversion and get Input weight map (if it exists)
@@ -476,79 +369,7 @@ void FDataflowEditorVertexAttributePaintToolData::Setup(FDynamicMesh3& InMesh, F
 	}
 }
 
-void FDataflowEditorVertexAttributePaintToolData::Shutdown()
-{
-	// nothing to do - invalidation is now done in the commit results
-}
-
-void FDataflowEditorVertexAttributePaintToolData::BeginChange(UDynamicMeshComponent* Component)
-{
-	ensure(ActiveWeightEditChangeTracker == nullptr);
-	if (Component)
-	{
-		ActiveWeightEditChangeTracker = MakeUnique<UE::Geometry::FDynamicMeshChangeTracker>(Component->GetMesh());
-		ActiveWeightEditChangeTracker->BeginChange();
-	}
-}
-
-TUniquePtr<UE::DataflowEditorVertexAttributePaintTool::Private::FMeshChange> FDataflowEditorVertexAttributePaintToolData::EndChange(UDynamicMeshComponent* Component)
-{
-	if (ensure(ActiveWeightEditChangeTracker) && Component)
-	{
-		using namespace UE::DataflowEditorVertexAttributePaintTool;
-		TUniquePtr<UE::Geometry::FDynamicMeshChange> EditResult = ActiveWeightEditChangeTracker->EndChange();
-		TUniquePtr<Private::FMeshChange> MeshChange = MakeUnique<Private::FMeshChange>(Component, MoveTemp(EditResult));
-		ActiveWeightEditChangeTracker = nullptr;
-
-		return MeshChange;
-	}
-	return {};
-}
-
-void FDataflowEditorVertexAttributePaintToolData::CancelChange()
-{
-	if (ensure(ActiveWeightEditChangeTracker))
-	{
-		ActiveWeightEditChangeTracker->EndChange();
-		ActiveWeightEditChangeTracker = nullptr;
-	}
-}
-
-float FDataflowEditorVertexAttributePaintToolData::GetValue(int32 VertexIdx) const
-{
-	float Value = 0.f;
-	if (ActiveWeightMap && VertexIdx != IndexConstants::InvalidID)
-	{
-		ActiveWeightMap->GetValue(VertexIdx, &Value);
-	}
-	return Value;
-}
-
-float FDataflowEditorVertexAttributePaintToolData::GetAverageValue(const TArray<int32>& Vertices) const
-{
-	double AverageValue = 0.f;
-	if (ActiveWeightMap)
-	{
-		int32 NumValidValues = 0;
-		for (int32 VertexIdx : Vertices)
-		{
-			if (VertexIdx != IndexConstants::InvalidID)
-			{
-				float Value = 0.0f;
-				ActiveWeightMap->GetValue(VertexIdx, &Value);
-				AverageValue += (double)Value;
-				++NumValidValues;
-			}
-		}
-		if (NumValidValues > 0)
-		{
-			AverageValue = AverageValue / (double)NumValidValues;
-		}
-	}
-	return (float)AverageValue;
-}
-
-void FDataflowEditorVertexAttributePaintToolData::UpdateNode(TObjectPtr<UDataflowContextObject> InDataflowEditorContextObject, IToolsContextTransactionsAPI* TransactionAPI)
+void FDataflowEditorVertexAttributePaintToolDataAdapter::CommitToNode(TObjectPtr<UDataflowContextObject> InDataflowEditorContextObject, IToolsContextTransactionsAPI* TransactionAPI)
 {
 	if (ActiveWeightMap && NodeToUpdate)
 	{
@@ -592,58 +413,41 @@ void FDataflowEditorVertexAttributePaintToolData::UpdateNode(TObjectPtr<UDataflo
 //
 ///////////////////////////////////////////////////////////////////////////////////////////////////////////////////////////////////////////////////
 
-UBaseDynamicMeshComponent* UDataflowEditorVertexAttributePaintTool::GetSculptMeshComponent()
+bool UDataflowEditorVertexAttributePaintTool::SetupToolMesh(FDynamicMesh3& InOutToolMesh, int32& OutInitialAttributeIndex)
 {
-	return PreviewMesh ? Cast<UBaseDynamicMeshComponent>(PreviewMesh->GetRootComponent()) : nullptr;
-}
+	// Create a temp paint layer and allow the selected node to read from/write to that layer
+	const int32 NumAttributeLayers = InOutToolMesh.Attributes()->NumWeightLayers();
 
-void UDataflowEditorVertexAttributePaintTool::SetupBrushEditBehaviorSetup(ULocalTwoAxisPropertyEditInputBehavior& OutBehavior)
-{
-	// Map the horizontal behavior to change the brush size 
-	// (we do not want to use the default as we are internal set to use the adapative brush size but external display world radius)
-	//UMeshSculptToolBase::MapHorizontalBrushEditBehaviorToBrushSize(OutBehavior);
+	InOutToolMesh.Attributes()->SetNumWeightLayers(NumAttributeLayers + 1);
+	UE::Geometry::FDynamicMeshWeightAttribute* ActiveWeightMap = InOutToolMesh.Attributes()->GetWeightLayer(NumAttributeLayers);
+	ActiveWeightMap->SetName(FName("PaintLayer"));
+	
+	FDataflowVertexAttributeEditableNode* NodeToUpdate = DataflowEditorContextObject->GetSelectedNodeOfType<FDataflowVertexAttributeEditableNode>();
+	checkf(NodeToUpdate, TEXT("No Node is currently selected, or more than one node is selected"));
+	
+	// Initialize the active weight map with weights stored on the node
+	DataAdapter.Setup(InOutToolMesh, ActiveWeightMap, NodeToUpdate, DataflowEditorContextObject);
+	
+	OutInitialAttributeIndex = NumAttributeLayers;
 
-	OutBehavior.HorizontalProperty.GetValueFunc = [this]() { return GetCurrentBrushRadius(); };
-	OutBehavior.HorizontalProperty.SetValueFunc = [this](float NewValue) { SetBrushRadius(NewValue); };
-	OutBehavior.HorizontalProperty.MutateDeltaFunc = [this](float Delta)
-		{
-			// Scale delta if brush size is in world units.
-			//return Delta * (BrushProperties->BrushSize.SizeType == EBrushToolSizeType::World ? (CameraState.Position - LastBrushFrameWorld.Origin).Length() : 1.f);
-			return Delta * (CameraState.Position - LastBrushFrameWorld.Origin).Length();
-		};
-	OutBehavior.HorizontalProperty.Name = LOCTEXT("BrushRadius", "Radius");
-	OutBehavior.HorizontalProperty.bEnabled = true;
+	return true;
+}
 
+void UDataflowEditorVertexAttributePaintTool::CommitToolMesh(FDynamicMesh3& InToolMesh)
+{
+	if (UInteractiveToolManager* ToolManager = GetToolManager())
+	{
+		ToolManager->BeginUndoTransaction(LOCTEXT("DataflowEditorVertexAttributePaintTool_TransactionName", "Paint Weights"));
 
-	// map the vertical behavior to change the attribiute value
-	OutBehavior.VerticalProperty.GetValueFunc = [this]()
-		{
-			return ToolProperties->BrushProperties.AttributeValue;
-		};
+		DataAdapter.CommitToNode(DataflowEditorContextObject, ToolManager->GetContextTransactionsAPI());
 
-	OutBehavior.VerticalProperty.SetValueFunc = [this](float NewValue)
-		{
-			ToolProperties->BrushProperties.AttributeValue = FMath::Min(1.0f, FMath::Max(NewValue, 0.f));
-#if WITH_EDITOR
-			FPropertyChangedEvent PropertyChangedEvent(FDataflowEditorVertexAttributePaintToolBrushProperties::StaticStruct()->FindPropertyByName(GET_MEMBER_NAME_CHECKED(FDataflowEditorVertexAttributePaintToolBrushProperties, AttributeValue)));
-			ToolProperties->PostEditChangeProperty(PropertyChangedEvent);
-#endif
-		};
-	OutBehavior.VerticalProperty.Name = LOCTEXT("AttributeValue", "Value");
-	OutBehavior.VerticalProperty.EditRate = 0.005f;
-	OutBehavior.VerticalProperty.bEnabled = true;
+		ToolManager->EndUndoTransaction();
+	}
 }
 
 void UDataflowEditorVertexAttributePaintTool::Setup()
 {
-	MinimumBrushAdaptiveSizeRatio = 0.001;
-
-	UMeshSculptToolBase::Setup();
-
-	// todo(ccaillaud) : this should be a customized by the inherited class 
-	// Get the selected node to read and write the vertex attribute
-	FDataflowVertexAttributeEditableNode* NodeToUpdate = DataflowEditorContextObject->GetSelectedNodeOfType<FDataflowVertexAttributeEditableNode>();
-	checkf(NodeToUpdate, TEXT("No Node is currently selected, or more than one node is selected"));
+	Super::Setup();
 
 	// todo(ccaillaud) : this should be a parameter of the tool ? 
 	SetToolDisplayName(LOCTEXT("ToolName", "Paint Weight Maps"));
@@ -656,214 +460,6 @@ void UDataflowEditorVertexAttributePaintTool::Setup()
 			Scene->SetVisibility(false);
 		}
 	}
-
-	//OnDynamicMeshComponentChangedHandle = DynamicMeshComponent->OnMeshVerticesChanged.AddUObject(this, &UDataflowEditorVertexAttributePaintTool::OnDynamicMeshComponentChanged);
-
-	// Create the preview mesh 
-	PreviewMesh = NewObject<UPreviewMesh>(this);
-	PreviewMesh->CreateInWorld(TargetWorld, FTransform::Identity);
-	PreviewMesh->SetTangentsMode(EDynamicMeshComponentTangentsMode::AutoCalculated);
-	ToolSetupUtil::ApplyRenderingConfigurationToPreview(PreviewMesh, Target);
-
-	// make sure we get a copy of the original dynamic mesh 
-	FGetMeshParameters GetMeshParams;
-	GetMeshParams.bWantMeshTangents = true;
-	FDynamicMesh3 DynamicMeshCopy = UE::ToolTarget::GetDynamicMeshCopy(Target, GetMeshParams);
-	PreviewMesh->UpdatePreview(MoveTemp(DynamicMeshCopy));
-	
-	// Assign materials
-	FComponentMaterialSet MaterialSet = UE::ToolTarget::GetMaterialSet(Target);
-	PreviewMesh->SetMaterials(MaterialSet.Materials);
-
-	// make sure the dynamic mesh has all the necessary attributes
-	FDynamicMesh3* Mesh = GetSculptMesh();
-	check(Mesh);
-	Mesh->EnableVertexColors(FVector3f::One());
-	Mesh->Attributes()->EnablePrimaryColors();
-	Mesh->Attributes()->PrimaryColors()->CreatePerVertex(0.f);
-	MeshLocalBounds = Mesh->GetBounds(true);
-	
-	TFuture<void> PrecomputeFuture = Async(UE::DataflowEditorVertexAttributePaintTool::Private::WeightPaintToolAsyncExecTarget, [this]()
-		{
-			PrecomputeFilterData();
-		});
-
-	TFuture<void> OctreeFuture = Async(UE::DataflowEditorVertexAttributePaintTool::Private::WeightPaintToolAsyncExecTarget, [this, &Mesh]()
-		{
-			// initialize dynamic octree
-			if (Mesh->TriangleCount() > 100000)
-			{
-				Octree.RootDimension = MeshLocalBounds.MaxDim() / 10.0;
-				Octree.SetMaxTreeDepth(4);
-			}
-			else
-			{
-				Octree.RootDimension = MeshLocalBounds.MaxDim();
-				Octree.SetMaxTreeDepth(8);
-			}
-			Octree.Initialize(Mesh);
-		});
-
-	// Setup mesh selector
-	MeshSelector = NewObject<UToolMeshSelector>(this);
-	auto OnSelectionChangedLambda = [this]() { /*OnSelectionModified();*/ };
-	MeshSelector->InitialSetup(TargetWorld, this, OnSelectionChangedLambda);
-
-	// compute the mesh description for the mesh selector to use
-	{
-
-		MeshSelector->SetMesh(PreviewMesh, FTransform::Identity);
-	}
-
-	// initialize brush radius range interval, brush properties
-	UMeshSculptToolBase::InitializeBrushSizeRange(MeshLocalBounds);
-
-	// initialize properties
-	ToolProperties = NewObject<UDataflowEditorVertexAttributePaintToolProperties>(this);
-	{
-		ToolProperties->RestoreProperties(this);
-		ToolProperties->DisplayProperties.ColorRamp.OnColorCurveChangedDelegate.AddUObject(this, &UDataflowEditorVertexAttributePaintTool::OnColorRampChanged);
-	}
-	AddToolPropertySource(ToolProperties);
-
-	// make sure the selector is in sync with the settings
-	MeshSelector->SetComponentSelectionMode(ToolProperties->SelectionProperties.ComponentSelectionMode);
-	MeshSelector->SetSelectionTool(EMeshSelectorTool::Marquee);
-	MeshSelector->GetSelectionMechanic()->SetSelectionDragToolUpdateType(ESelectionDragToolUpdateType::OnTickAndRelease);
-	
-	InitializeIndicator();
-	InitializeMirrorIndicator();
-
-	// initialize our properties
-	AddToolPropertySource(UMeshSculptToolBase::BrushProperties);
-	SetToolPropertySourceEnabled(UMeshSculptToolBase::BrushProperties, false);
-	UMeshSculptToolBase::BrushProperties->bShowPerBrushProps = false;
-	UMeshSculptToolBase::BrushProperties->bShowFalloff = true;
-	UMeshSculptToolBase::BrushProperties->bShowLazyness = false;
-	UMeshSculptToolBase::BrushProperties->FlowRate = 0.0f;
-	CalculateBrushRadius();
-
-	PaintBrushOpProperties = NewObject<UDataflowVertexAttributePaintBrushOpProps>(this);
-	RegisterBrushType(PaintBrushId, LOCTEXT("Paint", "Paint"),
-		MakeUnique<FLambdaMeshSculptBrushOpFactory>([this]() { return MakeUnique<FDataflowVertexAttributePaintBrushOp>(); }),
-		PaintBrushOpProperties);
-
-	// secondary brushes
-	SmoothBrushOpProperties = NewObject<UDataflowWeightMapSmoothBrushOpProps>(this);
-	RegisterSecondaryBrushType(SmoothBrushId, LOCTEXT("SmoothBrushType", "Smooth"),
-		MakeUnique<FLambdaMeshSculptBrushOpFactory>([this]() { return MakeUnique<FDataflowWeightMapSmoothBrushOp>(); }),
-		SmoothBrushOpProperties);
-
-	AddToolPropertySource(UMeshSculptToolBase::ViewProperties);
-	SetToolPropertySourceEnabled(UMeshSculptToolBase::ViewProperties, false);
-
-	AddToolPropertySource(UMeshSculptToolBase::GizmoProperties);
-	SetToolPropertySourceEnabled(UMeshSculptToolBase::GizmoProperties, false);
-
-	// must call before updating brush type so that we register all brush properties?
-	UMeshSculptToolBase::OnCompleteSetup();
-
-
-	UpdateBrushType(ToolProperties->BrushProperties.BrushMode);
-
-	// mesh element display is used to show seams and other features ofthe mesh
-	MeshElementsDisplay = NewObject<UMeshElementsVisualizer>(this);
-
-	MeshElementsDisplay->CreateInWorld(TargetWorld, FTransform::Identity);
-	if (ensure(MeshElementsDisplay->Settings))
-	{
-		MeshElementsDisplay->Settings->bShowBorders = ToolProperties->DisplayProperties.bShowBorders;
-		MeshElementsDisplay->Settings->bShowUVSeams = ToolProperties->DisplayProperties.bShowUVSeams;
-		MeshElementsDisplay->Settings->bShowNormalSeams = ToolProperties->DisplayProperties.bShowNormalSeams;
-		MeshElementsDisplay->Settings->bShowTangentSeams = ToolProperties->DisplayProperties.bShowTangentSeams;
-		MeshElementsDisplay->Settings->bShowColorSeams = ToolProperties->DisplayProperties.bShowColorSeams;
-		MeshElementsDisplay->Settings->ThicknessScale = ToolProperties->DisplayProperties.ThicknessScale;
-		MeshElementsDisplay->Settings->WireframeColor = ToolProperties->DisplayProperties.WireframeColor;
-		MeshElementsDisplay->Settings->BoundaryEdgeColor = ToolProperties->DisplayProperties.BoundaryEdgeColor;
-		MeshElementsDisplay->Settings->UVSeamColor = ToolProperties->DisplayProperties.UVSeamColor;
-		MeshElementsDisplay->Settings->NormalSeamColor = ToolProperties->DisplayProperties.NormalSeamColor;
-		MeshElementsDisplay->Settings->TangentSeamColor = ToolProperties->DisplayProperties.TangentSeamColor;
-		MeshElementsDisplay->Settings->ColorSeamColor = ToolProperties->DisplayProperties.ColorSeamColor;
-		MeshElementsDisplay->Settings->DepthBias = ToolProperties->DisplayProperties.DepthBias;
-		AddToolPropertySource(MeshElementsDisplay->Settings);
-	}
-	MeshElementsDisplay->SetMeshAccessFunction([this](UMeshElementsVisualizer::ProcessDynamicMeshFunc ProcessFunc) {
-		ProcessFunc(*GetSculptMesh());
-		});
-
-	// disable view properties
-	SetViewPropertiesEnabled(false);
-	UpdateMaterialMode(EMeshEditingMaterialModes::VertexColor);
-	UpdateWireframeVisibility(false);
-	UpdateFlatShadingSetting(false);
-
-	VertexData.Setup(*Mesh, NodeToUpdate, DataflowEditorContextObject);
-	
-	PrecomputeFuture.Wait();
-	OctreeFuture.Wait();
-
-	// update colors
-	UpdateVertexColorOverlay();
-	//DynamicMeshComponent->FastNotifyVertexAttributesUpdated(EMeshRenderAttributeFlags::VertexColors);
-	GetToolManager()->PostInvalidation();
-
-	LoadBrushConfig(ToolProperties->BrushProperties.BrushMode);
-	SetBrushMode(ToolProperties->BrushProperties.BrushMode);
-
-	SetPrimaryFalloffType(EMeshSculptFalloffType::Smooth);
-
-}
-
-void UDataflowEditorVertexAttributePaintTool::InitializeIndicator()
-{
-	UMeshSculptToolBase::InitializeIndicator();
-	if (BrushIndicatorMesh)
-	{
-		BrushIndicatorMesh->SetOverrideWireframeRenderMaterial(BrushIndicatorMaterial);
-	}
-}
-
-void UDataflowEditorVertexAttributePaintTool::InitializeMirrorIndicator()
-{
-	// register and spawn brush indicator gizmo
-	GetToolManager()->GetPairedGizmoManager()->RegisterGizmoType(UE::DataflowEditorVertexAttributePaintTool::Private::MirrorBrushGizmoType, NewObject<UBrushStampIndicatorBuilder>());
-	MirrorBrushIndicator = GetToolManager()->GetPairedGizmoManager()->CreateGizmo<UBrushStampIndicator>(UE::DataflowEditorVertexAttributePaintTool::Private::MirrorBrushGizmoType, FString(), this);
-	MirrorBrushIndicatorMesh = MakeBrushIndicatorMesh(this, TargetWorld);
-	MirrorBrushIndicatorMesh->SetOverrideWireframeRenderMaterial(BrushIndicatorMaterial);
-	MirrorBrushIndicator->AttachedComponent = MirrorBrushIndicatorMesh->GetRootComponent();
-	MirrorBrushIndicator->LineThickness = 1.0;
-	MirrorBrushIndicator->bDrawIndicatorLines = true;
-	MirrorBrushIndicator->bDrawRadiusCircle = false;
-	MirrorBrushIndicator->LineColor = FLinearColor(0.9f, 0.4f, 0.4f);
-
-	MirrorPlaneMesh = NewObject<UPreviewMesh>(this);
-	MirrorPlaneMesh->CreateInWorld(TargetWorld, FTransform::Identity);
-	UE::Geometry::FRoundedRectangleMeshGenerator Gen;
-	Gen.Origin = { 0,0,0 };
-	Gen.Normal = { 0,0,1 };
-	Gen.Width = 100.0;
-	Gen.Height = 100.0;
-	Gen.WidthVertexCount = 10;
-	Gen.HeightVertexCount = 10;
-	Gen.bScaleUVByAspectRatio = false;
-	Gen.Radius = 10.0;
-	Gen.AngleSamples = 8;
-	Gen.SharpCorners = UE::Geometry::ERoundedRectangleCorner::None;
-	Gen.Generate();
-	FDynamicMesh3 PlaneMesh(&Gen);
-	MirrorPlaneMesh->UpdatePreview(&PlaneMesh);
-	//BrushIndicatorMaterial = ToolSetupUtil::GetDefaultBrushVolumeMaterial(GetToolManager());
-	if (BrushIndicatorMaterial)
-	{
-		MirrorPlaneMesh->SetMaterial(BrushIndicatorMaterial);
-	}
-
-	// make sure raytracing is disabled on the brush indicator
-	Cast<UDynamicMeshComponent>(MirrorPlaneMesh->GetRootComponent())->SetEnableRaytracing(false);
-	MirrorPlaneMesh->SetShadowsEnabled(false);
-	MirrorPlaneMesh->SetOverrideWireframeRenderMaterial(BrushIndicatorMaterial);
-
-	UpdateMirrorPlaneMesh();
 }
 
 void UDataflowEditorVertexAttributePaintTool::SetFocusInViewport() const
@@ -879,2503 +475,18 @@ void UDataflowEditorVertexAttributePaintTool::SetFocusInViewport() const
 	//}
 }
 
-void UDataflowEditorVertexAttributePaintTool::UpdateMirrorPlaneMesh()
-{
-	if (MirrorPlaneMesh)
-	{
-		FTransform PlaneMeshTransform = FTransform::Identity;
-
-		const bool bIsBrushMirrorVisible = IsInBrushMode() && IsBrushMirroringEnabled() && IsMirrorPlaneWidgetVisible() && !bIsMirrorHiddenBecauseOfBrushStroke;
-		const bool bIsMeshMirrorVisible = !IsInBrushMode() && IsMirrorPlaneWidgetVisible();
-		MirrorPlaneMesh->SetVisible(bIsBrushMirrorVisible || bIsMeshMirrorVisible);
-
-		if (IsBrushMirroringEnabled())
-		{
-			switch(GetMirrorAxis())
-			{
-			case EAxis::X:
-				PlaneMeshTransform = FTransform(FRotator(90, 0, 0));
-				break;
-			case EAxis::Y:
-				PlaneMeshTransform = FTransform(FRotator(0, 0, 90));
-				break;
-			case EAxis::Z:
-				PlaneMeshTransform = FTransform::Identity;
-			}
-		}
-
-		if (const FDynamicMesh3* Mesh = GetSculptMesh())
-		{
-			const FVector MeshExtents = MeshLocalBounds.Extents();
-			
-			// scale of plane
-			FVector PlaneScale(1.0);
-			switch (GetMirrorAxis())
-			{
-			case EAxis::X:
-				PlaneScale.X = MeshExtents.Z;
-				PlaneScale.Y = MeshExtents.Y;
-				break;
-			case EAxis::Y:
-				PlaneScale.X = MeshExtents.X;
-				PlaneScale.Y = MeshExtents.Z;
-				break;
-			case EAxis::Z:
-				PlaneScale.X = MeshExtents.X;
-				PlaneScale.Y = MeshExtents.Y;
-			}
-
-			// the plane is originally 100 units and we want 10% over the size of the box
-			constexpr double ScaleAdjustment = 0.02 * 1.10;
-			PlaneMeshTransform.SetScale3D(PlaneScale * ScaleAdjustment);
-
-			// Position of the plane
-			FVector PlanePosition = MeshLocalBounds.Center();
-			if (!IsBrushMirroringInObjectSpace())
-			{
-				switch (GetMirrorAxis())
-				{
-				case EAxis::X:
-					PlanePosition.X = 0;
-					break;
-				case EAxis::Y:
-					PlanePosition.Y = 0;
-					break;
-				case EAxis::Z:
-					PlanePosition.Z = 0;
-				}
-			}
-			PlaneMeshTransform.SetTranslation(PlanePosition);
-		}
-		MirrorPlaneMesh->SetTransform(PlaneMeshTransform);
-	}
-}
-
-UE::Geometry::FFrame3d UDataflowEditorVertexAttributePaintTool::ComputeMirroredFrame(const UE::Geometry::FFrame3d& InLocalSpaceFrame) const
-{
-	UE::Geometry::FFrame3d OutFrame{ InLocalSpaceFrame };
-
-	const bool bIsObjectSpaceMirror = IsBrushMirroringInObjectSpace();
-	const EAxis::Type MirrorAxis = GetMirrorAxis();
-	FVector MirrorCenter = MeshLocalBounds.Center();
-
-	if (!bIsObjectSpaceMirror)
-	{
-		// Transform the frame to world space before mirroring
-		OutFrame.Transform(CurTargetTransform);
-		MirrorCenter = FVector::ZeroVector;
-	}
-
-	// Mirroring is in bounds space 
-	FTransform LocalTransform = OutFrame.ToFTransform();
-	LocalTransform.SetTranslation(LocalTransform.GetTranslation() - MirrorCenter);
-	LocalTransform.Mirror(MirrorAxis, EAxis::None);
-	LocalTransform.SetTranslation(LocalTransform.GetTranslation() + MirrorCenter);
-	OutFrame = UE::Geometry::FFrame3d(LocalTransform);
-	
-	if (!bIsObjectSpaceMirror)
-	{
-		// Transform back the frame to local space after mirroring
-		OutFrame.Transform(CurTargetTransform.InverseUnsafe());
-	}
-
-	return OutFrame;
-}
-
-void UDataflowEditorVertexAttributePaintTool::UpdateMirrorIndicator()
-{
-	if (BrushIndicatorMesh && BrushIndicator)
-	{
-		MirrorBrushIndicator->bVisible = BrushIndicator->bVisible && IsBrushMirroringEnabled();
-		MirrorBrushIndicatorMesh->SetVisible(BrushIndicatorMesh->IsVisible() && IsBrushMirroringEnabled());
-		MirrorBrushIndicator->bDrawRadiusCircle = BrushIndicator->bDrawRadiusCircle;
-
-		FTransform MirrorBrushTransform = ComputeMirroredFrame(HoverStamp.LocalFrame).ToFTransform();
-
-		// Project onto the surface
-		MirrorBrushTriangleID = IndexConstants::InvalidID;
-		if (const FDynamicMesh3* Mesh = GetSculptMesh())
-		{
-			const double RayMaxDistance = 10.0;
-			const FVector MirrorBrushNormal = MirrorBrushTransform.GetUnitAxis(EAxis::Z);
-			const FVector RayDirection = -MirrorBrushNormal;
-			const FVector RayOrigin = MirrorBrushTransform.GetTranslation() + MirrorBrushNormal * (RayMaxDistance * 0.5);
-
-			double OutRayParameter = TNumericLimits<double>::Max();
-			const int32 HitTID = Octree.FSparseDynamicOctree3::FindNearestHitObject(
-				FRay3d(RayOrigin, RayDirection),
-				[&](int tid) { return Mesh->GetTriBounds(tid); },
-				[&](int tid, const FRay3d& Ray) 
-				{
-					UE::Geometry::FIntrRay3Triangle3d Intr = 
-						UE::Geometry::TMeshQueries<FDynamicMesh3>::TriangleIntersection(*Mesh, tid, Ray);
-					if (Intr.IntersectionType == EIntersectionType::Point)
-					{
-						OutRayParameter = FMath::Min(OutRayParameter, Intr.RayParameter);
-						return Intr.RayParameter;
-					}
-					return TNumericLimits<double>::Max();
-
-				}, RayMaxDistance);
-			if (HitTID != IndexConstants::InvalidID && Mesh->IsTriangle(HitTID))
-			{
-				MirrorBrushTriangleID = HitTID;
-				MirrorBrushTransform.SetTranslation(RayOrigin + RayDirection * OutRayParameter);
-			}
-		}
-		MirrorBrushIndicator->Update(
-			(float)GetCurrentBrushRadius(),
-			MirrorBrushTransform,
-			1.0f - (float)GetCurrentBrushFalloff());
-
-		UpdateMirrorPlaneMesh();
-	}
-}
-
 void UDataflowEditorVertexAttributePaintTool::SetDataflowEditorContextObject(TObjectPtr<UDataflowContextObject> InDataflowEditorContextObject)
 {
 	DataflowEditorContextObject = InDataflowEditorContextObject;
 }
 
-void UDataflowEditorVertexAttributePaintTool::Shutdown(EToolShutdownType ShutdownType)
-{
-	//if (DynamicMeshComponent != nullptr)
-	//{
-	//	DynamicMeshComponent->OnMeshChanged.Remove(OnDynamicMeshComponentChangedHandle);
-	//}
-
-	if (MeshElementsDisplay)
-	{
-		// make sure we move back all settings to the display properties so they can be saved in the config
-		if (ToolProperties && MeshElementsDisplay->Settings)
-		{
-			ToolProperties->DisplayProperties.bShowBorders = MeshElementsDisplay->Settings->bShowBorders;
-			ToolProperties->DisplayProperties.bShowUVSeams = MeshElementsDisplay->Settings->bShowUVSeams;
-			ToolProperties->DisplayProperties.bShowNormalSeams = MeshElementsDisplay->Settings->bShowNormalSeams;
-			ToolProperties->DisplayProperties.bShowTangentSeams = MeshElementsDisplay->Settings->bShowTangentSeams;
-			ToolProperties->DisplayProperties.bShowColorSeams = MeshElementsDisplay->Settings->bShowColorSeams;
-			ToolProperties->DisplayProperties.ThicknessScale = MeshElementsDisplay->Settings->ThicknessScale;
-			ToolProperties->DisplayProperties.WireframeColor = MeshElementsDisplay->Settings->WireframeColor;
-			ToolProperties->DisplayProperties.BoundaryEdgeColor = MeshElementsDisplay->Settings->BoundaryEdgeColor;
-			ToolProperties->DisplayProperties.UVSeamColor = MeshElementsDisplay->Settings->UVSeamColor;
-			ToolProperties->DisplayProperties.NormalSeamColor = MeshElementsDisplay->Settings->NormalSeamColor;
-			ToolProperties->DisplayProperties.TangentSeamColor = MeshElementsDisplay->Settings->TangentSeamColor;
-			ToolProperties->DisplayProperties.ColorSeamColor = MeshElementsDisplay->Settings->ColorSeamColor;
-			ToolProperties->DisplayProperties.DepthBias = MeshElementsDisplay->Settings->DepthBias;
-		}
-		MeshElementsDisplay->Disconnect();
-		MeshElementsDisplay = nullptr;
-	}
-
-	// make sure the tool properties are saved for the next session
-	if (ToolProperties)
-	{
-		SaveBrushConfig(ToolProperties->BrushProperties.BrushMode);
-
-		ToolProperties->DisplayProperties.ColorRamp.OnColorCurveChangedDelegate.RemoveAll(this);
-		ToolProperties->SaveProperties(this);
-		ToolProperties->SaveConfig();
-	}
-
-	if (MeshSelector)
-	{
-		MeshSelector->Shutdown();
-		MeshSelector = nullptr;
-	}
-
-	// we need to call this directly now because we use the preview mesh instead of the built-in component
-	if (ShutdownType == EToolShutdownType::Accept)
-	{
-		CommitResult(GetSculptMeshComponent(), false);
-	}
-
-	VertexData.Shutdown();
-
-	if (PreviewMesh)
-	{
-		PreviewMesh->SetVisible(false);
-		PreviewMesh->Disconnect();
-		PreviewMesh = nullptr;
-	}
-
-	if (MirrorBrushIndicatorMesh)
-	{
-		MirrorBrushIndicatorMesh->Disconnect();
-		MirrorBrushIndicatorMesh = nullptr;
-	}
-
-	if (MirrorPlaneMesh)
-	{
-		MirrorPlaneMesh->Disconnect();
-		MirrorPlaneMesh = nullptr;
-	}
-
-	UMeshSculptToolBase::Shutdown(ShutdownType);
-
-	BrushIndicatorMesh = nullptr;
-	GetToolManager()->GetPairedGizmoManager()->DeregisterGizmoType(UE::DataflowEditorVertexAttributePaintTool::Private::MirrorBrushGizmoType);
-}
-
-
-void UDataflowEditorVertexAttributePaintTool::CommitResult(UBaseDynamicMeshComponent* Component, bool bModifiedTopology)
-{
-	if (UInteractiveToolManager* ToolManager = GetToolManager())
-	{
-		ToolManager->BeginUndoTransaction(LOCTEXT("DataflowEditorVertexAttributePaintTool_TransactionName", "Paint Weights"));
-
-		VertexData.UpdateNode(DataflowEditorContextObject, ToolManager->GetContextTransactionsAPI());
-
-		ToolManager->EndUndoTransaction();
-	}
-}
-
-void UDataflowEditorVertexAttributePaintTool::EnableValuePicker(bool bEnable)
-{
-	bPendingPickWeight = bEnable;
-	SetFocusInViewport();
-}
-
-bool UDataflowEditorVertexAttributePaintTool::IsValuePickerEnabled() const
-{
-	return bPendingPickWeight;
-}
-
-void UDataflowEditorVertexAttributePaintTool::EnableBrushMirroring(bool bEnable)
-{
-	bBrushMirroringEnabled = bEnable;
-	SetFocusInViewport();
-}
-
-bool UDataflowEditorVertexAttributePaintTool::IsBrushMirroringEnabled() const
-{
-	return bBrushMirroringEnabled;
-}
-
-void UDataflowEditorVertexAttributePaintTool::SetMirrorPlaneWidgetVisible(bool bVisible)
-{
-	bMirrorPlaneWidgetVisible = bVisible;
-	UpdateMirrorPlaneMesh();
-
-	SetFocusInViewport();
-}
-
-bool UDataflowEditorVertexAttributePaintTool::IsMirrorPlaneWidgetVisible()
-{
-	return bMirrorPlaneWidgetVisible;
-}
-
-bool UDataflowEditorVertexAttributePaintTool::IsBrushMirroringInObjectSpace() const
-{
-	return ToolProperties ? ToolProperties->MirrorProperties.bObjectSpace : false;
-}
-
-void UDataflowEditorVertexAttributePaintTool::SetMirrorAxis(EAxis::Type InAxis)
-{
-	if (ToolProperties)
-	{
-		ToolProperties->MirrorProperties.MirrorAxis = InAxis;
-		UpdateMirrorPlaneMesh();
-
-		SetFocusInViewport();
-	}
-}
-
-EAxis::Type UDataflowEditorVertexAttributePaintTool::GetMirrorAxis() const
-{
-	return ToolProperties ? ToolProperties->MirrorProperties.MirrorAxis.GetValue() : EAxis::X;
-}
-
-void UDataflowEditorVertexAttributePaintTool::SetMeshSelectionTool(EMeshSelectorTool SelectorTool)
-{
-	if (MeshSelector)
-	{
-		MeshSelector->SetSelectionTool(SelectorTool);
-		SetFocusInViewport();
-	}
-}
-
-EMeshSelectorTool UDataflowEditorVertexAttributePaintTool::GetMeshSelectionTool() const
-{
-	return MeshSelector
-		? MeshSelector->GetSelectionTool()
-		: EMeshSelectorTool::Ray;
-}
-
-
-FInputChord UDataflowEditorVertexAttributePaintTool::GreyscaleDisplayInputChord = FInputChord(EModifierKey::Control, EKeys::NumPadSeven);
-FInputChord UDataflowEditorVertexAttributePaintTool::ColorMapDisplayInputChord = FInputChord(EModifierKey::Control, EKeys::NumPadEight);
-FInputChord UDataflowEditorVertexAttributePaintTool::MaterialColorDisplayInputChord = FInputChord(EModifierKey::Control, EKeys::NumPadSix);
-
-FInputChord UDataflowEditorVertexAttributePaintTool::VerticesSelectionModeInputChord = FInputChord(EModifierKey::None, EKeys::One);
-FInputChord UDataflowEditorVertexAttributePaintTool::EdgesSelectionModeInputChord = FInputChord(EModifierKey::None, EKeys::Two);
-FInputChord UDataflowEditorVertexAttributePaintTool::FacesSelectionModeInputChord = FInputChord(EModifierKey::None, EKeys::Three);
-
-FInputChord UDataflowEditorVertexAttributePaintTool::GrowSelectionInputChord = FInputChord(EModifierKey::None, EKeys::Period);
-FInputChord UDataflowEditorVertexAttributePaintTool::ShrinkSelectionInputChord = FInputChord(EModifierKey::None, EKeys::Comma);
-FInputChord UDataflowEditorVertexAttributePaintTool::FloodSelectionInputChord = FInputChord(EModifierKey::Control, EKeys::A);
-
-FInputChord UDataflowEditorVertexAttributePaintTool::InvertMirrorDirectionInputChord = FInputChord(EModifierKey::Control, EKeys::M);
-FInputChord UDataflowEditorVertexAttributePaintTool::MirrorValuesInputChord = FInputChord(EModifierKey::None, EKeys::M);
-
-FInputChord UDataflowEditorVertexAttributePaintTool::CopyAverageValueInputChord = FInputChord(EModifierKey::Control, EKeys::C);
-FInputChord UDataflowEditorVertexAttributePaintTool::PasteAverageValueInputChord = FInputChord(EModifierKey::Control, EKeys::V);
-
-FInputChord UDataflowEditorVertexAttributePaintTool::ToggleValuePickerInputChord = FInputChord(EModifierKey::None, EKeys::T);
-
-void UDataflowEditorVertexAttributePaintTool::RegisterActions(FInteractiveToolActionSet& ActionSet)
-{
-	// IMPORTANT : we do not register the base class actions as they do not map enough to the wanted actions for this tool
-
-	int32 ActionID = (int32)EStandardToolActions::BaseClientDefinedActionID + 500;
-
-	//-------------------------------------------------------------------------------------------------------
-	// Color modes
-
-	ActionSet.RegisterAction(this, ActionID,
-		TEXT("MeshDisplayGreyscale"),
-		LOCTEXT("MeshDisplayGreyscale", "Greyscale Display"),
-		LOCTEXT("MeshDisplayGreyscaleTooltip", "Use greyscale color ramp to display values"),
-		GreyscaleDisplayInputChord.GetModifierKey(), GreyscaleDisplayInputChord.Key,
-		[this]() { SetColorMode(EDataflowEditorToolColorMode::Greyscale); });
-	++ActionID;
-
-	ActionSet.RegisterAction(this, ActionID,
-		TEXT("MeshDisplayColorRamp"),
-		LOCTEXT("MeshDisplayColorRamp", "Color Ramp Display"),
-		LOCTEXT("MeshDisplayColorRampTooltip", "Use custom color ramp to display values"),
-		ColorMapDisplayInputChord.GetModifierKey(), ColorMapDisplayInputChord.Key,
-		[this]() { SetColorMode(EDataflowEditorToolColorMode::Ramp); });
-	++ActionID;
-
-	ActionSet.RegisterAction(this, ActionID,
-		TEXT("MeshDisplayFullMaterial"),
-		LOCTEXT("MeshDisplayFullMaterial", "Full Material Display"),
-		LOCTEXT("MeshDisplayFullMaterialTooltip", "Use normal material color to display values"),
-		MaterialColorDisplayInputChord.GetModifierKey(), MaterialColorDisplayInputChord.Key,
-		[this]() { SetColorMode(EDataflowEditorToolColorMode::FullMaterial); });
-	++ActionID;
-
-	//-------------------------------------------------------------------------------------------------------
-	// Component selection mode
-
-	ActionSet.RegisterAction(this, ActionID,
-		TEXT("ComponentSelectionModeVertices"),
-		LOCTEXT("ComponentSelectionModeVertices", "Vertex Selection Mode"),
-		LOCTEXT("ComponentSelectionModeVerticesTooltip", "Enable vertex selection mode"),
-		VerticesSelectionModeInputChord.GetModifierKey(), VerticesSelectionModeInputChord.Key,
-		[this]() { SetComponentSelectionMode(EComponentSelectionMode::Vertices); });
-	++ActionID;
-
-	ActionSet.RegisterAction(this, ActionID,
-		TEXT("ComponentSelectionModeEdges"),
-		LOCTEXT("ComponentSelectionModeEdges", "Edge Selection Mode"),
-		LOCTEXT("ComponentSelectionModeEdgesTooltip", "Enable edge selection mode"),
-		EdgesSelectionModeInputChord.GetModifierKey(), EdgesSelectionModeInputChord.Key,
-		[this]() { SetComponentSelectionMode(EComponentSelectionMode::Edges); });
-	++ActionID;
-
-	ActionSet.RegisterAction(this, ActionID,
-		TEXT("ComponentSelectionModeFaces"),
-		LOCTEXT("ComponentSelectionModeFaces", "Face Selection Mode"),
-		LOCTEXT("ComponentSelectionModeFacesTooltip", "Enable face selection mode"),
-		FacesSelectionModeInputChord.GetModifierKey(), FacesSelectionModeInputChord.Key,
-		[this]() { SetComponentSelectionMode(EComponentSelectionMode::Faces); });
-	++ActionID;
-
-	//-------------------------------------------------------------------------------------------------------
-	// Selection actions
-
-	ActionSet.RegisterAction(this, ActionID,
-		TEXT("SelectionGrow"),
-		LOCTEXT("SelectionGrow", "Grow Selection"),
-		LOCTEXT("SelectionGrowTooltip", "Grow the current selection"),
-		GrowSelectionInputChord.GetModifierKey(), GrowSelectionInputChord.Key,
-		[this]() { GrowSelection(); });
-	++ActionID;
-
-	ActionSet.RegisterAction(this, ActionID,
-		TEXT("SelectionShrink"),
-		LOCTEXT("SelectionShrink", "Shrink Selection"),
-		LOCTEXT("SelectionShrinkTooltip", "Shrink the current selection"),
-		ShrinkSelectionInputChord.GetModifierKey(), ShrinkSelectionInputChord.Key,
-		[this]() { ShrinkSelection(); });
-	++ActionID;
-
-	ActionSet.RegisterAction(this, ActionID,
-		TEXT("SelectionFlood"),
-		LOCTEXT("SelectionFlood", "Flood Selection"),
-		LOCTEXT("SelectionFloodTooltip", "Flood the current selection"),
-		FloodSelectionInputChord.GetModifierKey(), FloodSelectionInputChord.Key,
-		[this]() { FloodSelection(); });
-	++ActionID;
-
-	//-------------------------------------------------------------------------------------------------------
-	// Mirror actions
-
-	ActionSet.RegisterAction(this, ActionID,
-		TEXT("MirrorDirection"),
-		LOCTEXT("MirrorDirection", "Invert Mirror Direction"),
-		LOCTEXT("MirrorDirectionTooltip", "Invert the current mirror direction"),
-		InvertMirrorDirectionInputChord.GetModifierKey(), InvertMirrorDirectionInputChord.Key,
-		[this]() 
-		{ 
-			if (ToolProperties)
-			{
-				ToolProperties->MirrorProperties.MirrorDirection =
-					(ToolProperties->MirrorProperties.MirrorDirection == EDataflowEditorToolMirrorDirection::PositiveToNegative)
-					? EDataflowEditorToolMirrorDirection::NegativeToPositive
-					: EDataflowEditorToolMirrorDirection::PositiveToNegative;
-			}
-		});
-	++ActionID;
-
-	ActionSet.RegisterAction(this, ActionID,
-		TEXT("MirrorValues"),
-		LOCTEXT("MirrorValues", "Mirror Values"),
-		LOCTEXT("MirrorValuesTooltip", "Mirror the current selection values"),
-		MirrorValuesInputChord.GetModifierKey(), MirrorValuesInputChord.Key,
-		[this]() { OnMirrorCommand(); });
-	++ActionID;
-
-	//-------------------------------------------------------------------------------------------------------
-	// Copy Paste
-
-	ActionSet.RegisterAction(this, ActionID,
-		TEXT("CopyValue"),
-		LOCTEXT("CopyValue", "Copy Average Value"),
-		LOCTEXT("CopyValueTooltip", "Copy the current selection average value"),
-		CopyAverageValueInputChord.GetModifierKey(), CopyAverageValueInputChord.Key,
-		[this]() { CopyAverageFromSelectionToClipboard(); });
-	++ActionID;
-
-	ActionSet.RegisterAction(this, ActionID,
-		TEXT("PasteValue"),
-		LOCTEXT("PasteValue", "Paste Value"),
-		LOCTEXT("PasteValueTooltip", "Paste value onto the current selection"),
-		PasteAverageValueInputChord.GetModifierKey(), PasteAverageValueInputChord.Key,
-		[this]() { PasteValueToSelectionFromClipboard(); });
-	++ActionID;
-
-	//-------------------------------------------------------------------------------------------------------
-	// Value picking
-
-	ActionSet.RegisterAction(this, ActionID,
-		TEXT("PickValueUnderCursor"),
-		LOCTEXT("PickValueUnderCursor", "Pick Value"),
-		LOCTEXT("PickValueUnderCursorTooltip", "Set the active value to that currently under the cursor"),
-		ToggleValuePickerInputChord.GetModifierKey(), ToggleValuePickerInputChord.Key,
-		[this]() { EnableValuePicker(!IsValuePickerEnabled()); });
-	++ActionID;
-
-};
-
-void UDataflowEditorVertexAttributePaintTool::OnPropertyModified(UObject* PropertySet, FProperty* Property)
-{
-	CalculateBrushRadius();
-}
-
-FBox UDataflowEditorVertexAttributePaintTool::GetWorldSpaceFocusBox()
-{
-	if (ToolProperties->EditingMode == EDataflowEditorToolEditMode::Mesh)
-	{
-		if (PreviewMesh && PreviewMesh->GetMesh())
-		{
-			const TArray<int32>& SelectedVertices = GetSelectedVertices();
-			if (!SelectedVertices.IsEmpty())
-			{
-				const FDynamicMesh3* Mesh = PreviewMesh->GetMesh();
-				const FTransform3d Transform(PreviewMesh->GetTransform());
-
-				FBox MeshBounds(EForceInit::ForceInit);
-				for (const int32 VertexID : SelectedVertices)
-				{
-					MeshBounds += Transform.TransformPosition(Mesh->GetVertex(VertexID));
-				}
-				return static_cast<FBox>(MeshBounds);
-			}
-			return FBox(MeshLocalBounds);
-		}
-
-	}
-	return Super::GetWorldSpaceFocusBox();
-}
-
-float UDataflowEditorVertexAttributePaintTool::GetBrushMinRadius() const
-{
-	return BrushProperties? BrushProperties->BrushSize.WorldSizeRange.Min * 0.5f: 0.001f;
-}
-
-float UDataflowEditorVertexAttributePaintTool::GetBrushMaxRadius() const
-{
-	return BrushProperties ? BrushProperties->BrushSize.WorldSizeRange.Max * 0.5f: 20.0f;
-}
-
-float UDataflowEditorVertexAttributePaintTool::GetBrushAdaptiveSize() const
-{
-	return BrushProperties ? BrushProperties->BrushSize.AdaptiveSize : 1.0f;
-}
-
-void UDataflowEditorVertexAttributePaintTool::SetBrushAdaptiveSize(float AdaptiveSize)
-{
-	if (BrushProperties)
-	{
-		if (UMeshSculptToolBase::BrushProperties->BrushSize.SizeType == EBrushToolSizeType::Adaptive)
-		{
-			BrushProperties->BrushSize.AdaptiveSize = FMath::Clamp(AdaptiveSize, 0.0f, 1.0f);
-		}
-		else
-		{
-			BrushProperties->BrushSize.WorldRadius = FMath::Clamp(AdaptiveSize * GetBrushMaxRadius(), BrushProperties->BrushSize.WorldSizeRange.Min, BrushProperties->BrushSize.WorldSizeRange.Max);
-		}
-		NotifyOfPropertyChangeByTool(BrushProperties);
-		CalculateBrushRadius();
-		SetFocusInViewport();
-	}
-}
-
-void UDataflowEditorVertexAttributePaintTool::SetBrushRadius(float NewRadius)
-{
-	if (BrushProperties)
-	{
-		if (UMeshSculptToolBase::BrushProperties->BrushSize.SizeType == EBrushToolSizeType::Adaptive)
-		{
-			BrushProperties->BrushSize.AdaptiveSize = FMath::Clamp(NewRadius / FMath::Max(0.0001f, GetBrushMaxRadius()), 0.0f, 1.0f);
-		}
-		else
-		{
-			BrushProperties->BrushSize.WorldRadius = FMath::Clamp(NewRadius, BrushProperties->BrushSize.WorldSizeRange.Min, BrushProperties->BrushSize.WorldSizeRange.Max);
-		}
-		NotifyOfPropertyChangeByTool(BrushProperties);
-		CalculateBrushRadius();
-		SetFocusInViewport();
-	}
-}
-
-float UDataflowEditorVertexAttributePaintTool::GetBrushFalloff() const
-{
-	return BrushProperties ? BrushProperties->BrushFalloffAmount : 1.0f;
-}
-
-void UDataflowEditorVertexAttributePaintTool::SetBrushFalloff(float NewFalloff)
-{
-	if (BrushProperties)
-	{
-		BrushProperties->BrushFalloffAmount = FMath::Clamp(NewFalloff, 0.0f, 1.0f);
-		NotifyOfPropertyChangeByTool(BrushProperties);
-		SetFocusInViewport();
-	}
-}
-
-void UDataflowEditorVertexAttributePaintTool::SaveBrushConfig(EDataflowEditorToolEditOperation BrushMode)
-{
-	if (ToolProperties)
-	{
-		// Save the current brush config 
-		FDataflowEditorVertexAttributePaintToolBrushConfig BrushConfigToSave;
-		BrushConfigToSave.BrushSize = GetBrushAdaptiveSize();
-		BrushConfigToSave.Value = ToolProperties->BrushProperties.AttributeValue;
-		BrushConfigToSave.Falloff = GetBrushFalloff();
-		ToolProperties->SetBrushConfig(BrushMode, BrushConfigToSave);
-	}
-}
-
-void UDataflowEditorVertexAttributePaintTool::LoadBrushConfig(EDataflowEditorToolEditOperation BrushMode)
-{
-	const FDataflowEditorVertexAttributePaintToolBrushConfig& NewBrushConfig = ToolProperties->GetBrushConfig(BrushMode);
-	SetBrushAdaptiveSize(NewBrushConfig.BrushSize);
-	ToolProperties->BrushProperties.AttributeValue = NewBrushConfig.Value;
-	SetBrushFalloff(NewBrushConfig.Falloff);
-
-	FPropertyChangedEvent ValuePropertyChangedEvent(UBrushBaseProperties::StaticClass()->FindPropertyByName(GET_MEMBER_NAME_CHECKED(FDataflowEditorVertexAttributePaintToolBrushProperties, AttributeValue)));
-	ToolProperties->PostEditChangeProperty(ValuePropertyChangedEvent);
-
-}
-
-void UDataflowEditorVertexAttributePaintTool::SetEditingMode(EDataflowEditorToolEditMode Editmode)
-{
-	if (ToolProperties)
-	{
-		ToolProperties->EditingMode = Editmode;
-		if (Editmode == EDataflowEditorToolEditMode::Mesh)
-		{
-			SetMirrorPlaneWidgetVisible(false);
-		}
-		SetFocusInViewport();
-	}
-}
-
-void UDataflowEditorVertexAttributePaintTool::SetBrushMode(EDataflowEditorToolEditOperation NewBrushMode)
-{
-	if (ToolProperties && BrushProperties)
-	{
-		const EDataflowEditorToolEditOperation CurrentBrushMode = ToolProperties->BrushProperties.BrushMode;
-		ToolProperties->BrushProperties.BrushMode = NewBrushMode;
-
-		// Save the current brush config 
-		SaveBrushConfig(CurrentBrushMode);
-
-		// Load the new brush config 
-		LoadBrushConfig(NewBrushMode);
-
-		// Send notification for the changed properties
-		FPropertyChangedEvent BrushModePropertyChangedEvent(UBrushBaseProperties::StaticClass()->FindPropertyByName(GET_MEMBER_NAME_CHECKED(FDataflowEditorVertexAttributePaintToolBrushProperties, BrushMode)));
-		ToolProperties->PostEditChangeProperty(BrushModePropertyChangedEvent);
-
-		ToolProperties->SaveConfig();
-		BrushProperties->SaveConfig();
-
-		SetFocusInViewport();
-	}
-}
-
-bool UDataflowEditorVertexAttributePaintTool::IsInBrushMode() const
-{
-	return (ToolProperties && ToolProperties->EditingMode == EDataflowEditorToolEditMode::Brush);
-}
-
-void UDataflowEditorVertexAttributePaintTool::SetBrushAreaMode(EMeshVertexPaintBrushAreaType BrushAreaMode)
-{
-	if (ToolProperties)
-	{
-		ToolProperties->BrushProperties.BrushAreaMode = BrushAreaMode;
-		SetFocusInViewport();
-	}
-}
-
-bool UDataflowEditorVertexAttributePaintTool::IsVolumetricBrush() const
-{
-	return (ToolProperties && ToolProperties->BrushProperties.BrushAreaMode == EMeshVertexPaintBrushAreaType::Volumetric);
-}
-
-void UDataflowEditorVertexAttributePaintTool::OnBeginStroke(const FRay& WorldRay)
-{
-	if (!VertexData.IsValid())
-	{
-		return;
-	}
-
-	bIsMirrorHiddenBecauseOfBrushStroke = (ToolProperties && ToolProperties->MirrorProperties.bHideOnBrushStroke);
-
-	UpdateBrushPosition(WorldRay);
-
-	if (PaintBrushOpProperties)
-	{
-		PaintBrushOpProperties->AttributeValue = ToolProperties->BrushProperties.AttributeValue;
-		if (GetInInvertStroke())
-		{
-			PaintBrushOpProperties->AttributeValue = FMath::Clamp(1.0f - PaintBrushOpProperties->AttributeValue, 0.f, 1.f);
-		}
-		PaintBrushOpProperties->EditOperation = ToolProperties->BrushProperties.BrushMode;
-		PaintBrushOpProperties->Strength = 1.0f; // Attribute painting only need the attribute value 
-		PaintBrushOpProperties->SetFalloff(UMeshSculptToolBase::BrushProperties->BrushFalloffAmount);
-	}
-
-	// initialize first "Last Stamp", so that we can assume all stamps in stroke have a valid previous stamp
-	LastStamp.WorldFrame = GetBrushFrameWorld();
-	LastStamp.LocalFrame = GetBrushFrameLocal();
-	LastStamp.Radius = GetCurrentBrushRadius();
-	LastStamp.Falloff = GetCurrentBrushFalloff();
-	LastStamp.Direction = GetInInvertStroke() ? -1.0 : 1.0;
-	LastStamp.Depth = GetCurrentBrushDepth();
-	LastStamp.Power = GetActivePressure() * GetCurrentBrushStrength();
-	LastStamp.TimeStamp = FDateTime::Now();
-
-	FSculptBrushOptions SculptOptions;
-	SculptOptions.ConstantReferencePlane = GetCurrentStrokeReferencePlane();
-
-	TUniquePtr<FMeshSculptBrushOp>& UseBrushOp = GetActiveBrushOp();
-	UseBrushOp->ConfigureOptions(SculptOptions);
-	UseBrushOp->BeginStroke(GetSculptMesh(), LastStamp, VertexROI);
-
-	AccumulatedTriangleROI.Reset();
-
-	// begin change here? or wait for first stamp?
-	BeginChange();
-}
-
-void UDataflowEditorVertexAttributePaintTool::OnEndStroke()
-{
-	if (!VertexData.IsValid())
-	{
-		return;
-	}
-
-	bIsMirrorHiddenBecauseOfBrushStroke = false;
-
-	GetActiveBrushOp()->EndStroke(GetSculptMesh(), LastStamp, VertexROI);
-
-	UpdateVertexColorOverlay(&TriangleROI);
-	//DynamicMeshComponent->FastNotifyVertexAttributesUpdated(EMeshRenderAttributeFlags::VertexColors);
-
-	// close change record
-	EndChange();
-
-	if (bPendingPickWeight)
-	{
-		if (GetSculptMesh()->IsTriangle(GetBrushTriangleID()))
-		{
-			ToolProperties->BrushProperties.AttributeValue = GetCurrentWeightValueUnderBrush();
-			NotifyOfPropertyChangeByTool(ToolProperties);
-		}
-	}
-}
-
-void UDataflowEditorVertexAttributePaintTool::OnCancelStroke()
-{
-	if (const TUniquePtr<FMeshSculptBrushOp>& UseBrushOp = GetActiveBrushOp())
-	{
-		UseBrushOp->CancelStroke();
-	}
-	CancelChange();
-}
-
-void UDataflowEditorVertexAttributePaintTool::UpdateROI(const FSculptBrushStamp& BrushStamp, int32 BrushTriangleID)
-{
-	SCOPE_CYCLE_COUNTER(VertexAttributePaintTool_UpdateROI);
-
-	const FVector3d& BrushPos = BrushStamp.LocalFrame.Origin;
-	const FDynamicMesh3* Mesh = GetSculptMesh();
-	const float RadiusSqr = GetCurrentBrushRadius() * GetCurrentBrushRadius();
-	UE::Geometry::FAxisAlignedBox3d BrushBox(
-		BrushPos - GetCurrentBrushRadius() * FVector3d::One(),
-		BrushPos + GetCurrentBrushRadius() * FVector3d::One());
-
-	TriangleROI.Reset();
-
-	int32 CenterTID = BrushTriangleID;
-	if (Mesh->IsTriangle(CenterTID))
-	{
-		TriangleROI.Add(CenterTID);
-	}
-
-	FVector3d CenterNormal = Mesh->IsTriangle(CenterTID) ? TriNormals[CenterTID] : FVector3d::One();		// One so that normal check always passes
-
-	const bool bVolumetric = IsVolumetricBrush();
-	const bool bUseAngleThreshold = ToolProperties->BrushProperties.AngleThreshold < 180.0f;
-	const double DotAngleThreshold = FMathd::Cos(ToolProperties->BrushProperties.AngleThreshold * FMathd::DegToRad);
-	const bool bStopAtUVSeams = ToolProperties->BrushProperties.bUVSeams;
-	const bool bStopAtNormalSeams = ToolProperties->BrushProperties.bNormalSeams;
-
-	auto CheckEdgeCriteria = [&](int32 t1, int32 t2) -> bool
-		{
-			if (bUseAngleThreshold == false || CenterNormal.Dot(TriNormals[t2]) > DotAngleThreshold)
-			{
-				int32 eid = Mesh->FindEdgeFromTriPair(t1, t2);
-				if (bStopAtUVSeams == false || UVSeamEdges[eid] == false)
-				{
-					if (bStopAtNormalSeams == false || NormalSeamEdges[eid] == false)
-					{
-						return true;
-					}
-				}
-			}
-			return false;
-		};
-
-	if (bVolumetric)
-	{
-		Octree.RangeQuery(BrushBox,
-			[&](int TriIdx) {
-				if ((Mesh->GetTriCentroid(TriIdx) - BrushPos).SquaredLength() < RadiusSqr)
-				{
-					TriangleROI.Add(TriIdx);
-				}
-			});
-	}
-	else
-	{
-		if (Mesh->IsTriangle(CenterTID))
-		{
-			TArray<int32> StartROI;
-			StartROI.Add(CenterTID);
-			UE::Geometry::FMeshConnectedComponents::GrowToConnectedTriangles(Mesh, StartROI, TriangleROI, &TempROIBuffer,
-				[&](int t1, int t2)
-				{
-					if ((Mesh->GetTriCentroid(t2) - BrushPos).SquaredLength() < RadiusSqr)
-					{
-						return CheckEdgeCriteria(t1, t2);
-					}
-					return false;
-				});
-		}
-	}
-
-	if (VertexData.HaveDynamicMeshToWeightConversion())
-	{
-		// Find triangles whose vertices map to the same welded vertex as any vertex in VertexSetBuffer and add them to TriangleROI
-		for (const int VertexID : VertexSetBuffer)
-		{
-			for (const int OtherVertexID : VertexData.GetWeightToDynamicMesh()[VertexData.GetDynamicMeshToWeight()[VertexID]])
-			{
-				if (OtherVertexID != VertexID)
-				{
-					Mesh->EnumerateVertexTriangles(OtherVertexID, [this](int32 AdjacentTri)
-						{
-							TriangleROI.Add(AdjacentTri);
-						});
-				}
-			}
-		}
-
-	}
-
-	// Construct ROI vertex set
-	VertexSetBuffer.Reset();
-	for (int32 tid : TriangleROI)
-	{
-		UE::Geometry::FIndex3i Tri = Mesh->GetTriangle(tid);
-		VertexSetBuffer.Add(Tri.A);  VertexSetBuffer.Add(Tri.B);  VertexSetBuffer.Add(Tri.C);
-	}
-
-	// Apply visibility filter
-	if (ToolProperties->BrushProperties.VisibilityFilter != EDataflowEditorToolVisibilityType::None)
-	{
-		TArray<int32> ResultBuffer;
-		ApplyVisibilityFilter(VertexSetBuffer, TempROIBuffer, ResultBuffer);
-	}
-
-	VertexROI.SetNum(0, EAllowShrinking::No);
-	//TODO: If we paint a 2D projection of UVs, these will need to be the 2D vertices not the 3D original mesh vertices
-	BufferUtil::AppendElements(VertexROI, VertexSetBuffer);
-
-	// Construct ROI triangle and weight buffers
-	ROITriangleBuffer.Reserve(TriangleROI.Num());
-	ROITriangleBuffer.SetNum(0, EAllowShrinking::No);
-	for (int32 tid : TriangleROI)
-	{
-		ROITriangleBuffer.Add(tid);
-	}
-	ROIWeightValueBuffer.SetNum(VertexROI.Num(), EAllowShrinking::No);
-	SyncWeightBufferWithMesh(Mesh);
-}
-
-bool UDataflowEditorVertexAttributePaintTool::UpdateStampPosition(const FRay& WorldRay)
-{
-	CalculateBrushRadius();
-
-	TUniquePtr<FMeshSculptBrushOp>& UseBrushOp = GetActiveBrushOp();
-
-	ESculptBrushOpTargetType TargetType = UseBrushOp->GetBrushTargetType();
-	switch (TargetType)
-	{
-	case ESculptBrushOpTargetType::SculptMesh:
-	case ESculptBrushOpTargetType::TargetMesh:
-		UpdateBrushPositionOnSculptMesh(WorldRay, true);
-		break;
-	case ESculptBrushOpTargetType::ActivePlane:
-		check(false);
-		UpdateBrushPositionOnActivePlane(WorldRay);
-		break;
-	}
-
-	if (UseBrushOp->GetAlignStampToView() || IsVolumetricBrush())
-	{
-		AlignBrushToView();
-	}
-
-	CurrentStamp = LastStamp;
-	CurrentStamp.DeltaTime = FMathd::Min((FDateTime::Now() - LastStamp.TimeStamp).GetTotalSeconds(), 1.0);
-	CurrentStamp.WorldFrame = GetBrushFrameWorld();
-	CurrentStamp.LocalFrame = GetBrushFrameLocal();
-	CurrentStamp.Power = GetActivePressure() * GetCurrentBrushStrength();
-	CurrentStamp.Falloff = GetCurrentBrushFalloff();
-
-	CurrentStamp.PrevLocalFrame = LastStamp.LocalFrame;
-	CurrentStamp.PrevWorldFrame = LastStamp.WorldFrame;
-
-	const EAxis::Type MirrorAxis = GetMirrorAxis();
-
-	// update the mirrored stamp
-	MirrorCurrentStamp = CurrentStamp;
-	MirrorCurrentStamp.LocalFrame = ComputeMirroredFrame(CurrentStamp.LocalFrame);
-	MirrorCurrentStamp.WorldFrame = MirrorCurrentStamp.LocalFrame;
-	MirrorCurrentStamp.WorldFrame.Transform(CurTargetTransform);
-	MirrorCurrentStamp.PrevLocalFrame = ComputeMirroredFrame(CurrentStamp.PrevLocalFrame);
-	MirrorCurrentStamp.PrevWorldFrame = MirrorCurrentStamp.PrevLocalFrame;
-	MirrorCurrentStamp.PrevWorldFrame.Transform(CurTargetTransform);
-
-	FVector3d MoveDelta = CurrentStamp.LocalFrame.Origin - CurrentStamp.PrevLocalFrame.Origin;
-
-	if (UseBrushOp->IgnoreZeroMovements() && MoveDelta.SquaredLength() < 0.1 * CurrentBrushRadius)
-	{
-		return false;
-	}
-
-	return true;
-}
-
-bool UDataflowEditorVertexAttributePaintTool::ApplyStamp(const FSculptBrushStamp& InStamp)
-{
-	SCOPE_CYCLE_COUNTER(VertexAttributePaintTool_ApplyStamp);
-
-	TUniquePtr<FMeshSculptBrushOp>& UseBrushOp = GetActiveBrushOp();
-
-	// yuck
-	FDataflowMeshVertexWeightMapEditBrushOp* WeightBrushOp = (FDataflowMeshVertexWeightMapEditBrushOp*)UseBrushOp.Get();
-	WeightBrushOp->bApplyRadiusLimit = IsInBrushMode();
-
-	FDynamicMesh3* Mesh = GetSculptMesh();
-	WeightBrushOp->ApplyStampByVertices(Mesh, InStamp, VertexROI, ROIWeightValueBuffer);
-
-	bool bUpdated = SyncMeshWithWeightBuffer(Mesh);
-
-	return bUpdated;
-}
-
-bool UDataflowEditorVertexAttributePaintTool::SyncMeshWithWeightBuffer(FDynamicMesh3* Mesh)
-{
-	int32 NumModified = 0;
-	if (VertexData.IsValid())
-	{
-		// change update could be async here if we collected array of <idx,orig,new> and dispatched independenlty
-		const int32 NumVertexROI = VertexROI.Num();
-		for (int32 VertexROIIndex = 0; VertexROIIndex < NumVertexROI; ++VertexROIIndex)
-		{
-			const int VertIdx = VertexROI[VertexROIIndex];
-			if (VertIdx != INDEX_NONE)
-			{
-				const float NewValue = ROIWeightValueBuffer[VertexROIIndex];
-				VertexData.SetValue(VertIdx,
-					[NewValue](UE::Geometry::FDynamicMeshWeightAttribute& AttributeMap, int32 AttributeIndex)
-					{
-						return NewValue;
-					});
-				NumModified++;
-			}
-		}
-	}
-	return (NumModified > 0);
-}
-
-bool UDataflowEditorVertexAttributePaintTool::SyncWeightBufferWithMesh(const FDynamicMesh3* Mesh)
-{
-	int32 NumModified = 0;
-	if (VertexData.IsValid())
-	{
-		// change update could be async here if we collected array of <idx,orig,new> and dispatched independenlty
-		const int32 NumVertexROI = VertexROI.Num();
-		for (int32 VertexROIIndex = 0; VertexROIIndex < NumVertexROI; ++VertexROIIndex)
-		{
-			const int VertIdx = VertexROI[VertexROIIndex];
-			const float CurWeight = VertexData.GetValue(VertIdx);
-			if (ROIWeightValueBuffer[VertexROIIndex] != CurWeight)
-			{
-				ROIWeightValueBuffer[VertexROIIndex] = CurWeight;
-				NumModified++;
-			}
-		}
-	}
-	return (NumModified > 0);
-}
-
-// TODO(ccaillaud) : move this to a common  place
-namespace UE::DataflowEditorVertexAttributePaintTool::Private
-{
-	template<typename RealType>
-	static bool FindPolylineSelfIntersection(
-		const TArray<UE::Math::TVector2<RealType>>& Polyline,
-		UE::Math::TVector2<RealType>& IntersectionPointOut,
-		UE::Geometry::FIndex2i& IntersectionIndexOut,
-		bool bParallel = true)
-	{
-		int32 N = Polyline.Num();
-		std::atomic<bool> bSelfIntersects(false);
-		ParallelFor(N - 1, [&](int32 i)
-			{
-				UE::Geometry::TSegment2<RealType> SegA(Polyline[i], Polyline[i + 1]);
-				for (int32 j = i + 2; j < N - 1 && bSelfIntersects == false; ++j)
-				{
-					UE::Geometry::TSegment2<RealType> SegB(Polyline[j], Polyline[j + 1]);
-					if (SegA.Intersects(SegB) && bSelfIntersects == false)
-					{
-						bool ExpectedValue = false;
-						if (std::atomic_compare_exchange_strong(&bSelfIntersects, &ExpectedValue, true))
-						{
-							UE::Geometry::TIntrSegment2Segment2<RealType> Intersection(SegA, SegB);
-							Intersection.Find();
-							IntersectionPointOut = Intersection.Point0;
-							IntersectionIndexOut = UE::Geometry::FIndex2i(i, j);
-							return;
-						}
-					}
-				}
-			}, (bParallel) ? EParallelForFlags::None : EParallelForFlags::ForceSingleThread);
-
-		return bSelfIntersects;
-	}
-
-
-
-	template<typename RealType>
-	static bool FindPolylineSegmentIntersection(
-		const TArray<UE::Math::TVector2<RealType>>& Polyline,
-		const UE::Geometry::TSegment2<RealType>& Segment,
-		UE::Math::TVector2<RealType>& IntersectionPointOut,
-		int& IntersectionIndexOut)
-	{
-
-		int32 N = Polyline.Num();
-		for (int32 i = 0; i < N - 1; ++i)
-		{
-			UE::Geometry::TSegment2<RealType> PolySeg(Polyline[i], Polyline[i + 1]);
-			if (Segment.Intersects(PolySeg))
-			{
-				UE::Geometry::TIntrSegment2Segment2<RealType> Intersection(Segment, PolySeg);
-				Intersection.Find();
-				IntersectionPointOut = Intersection.Point0;
-				IntersectionIndexOut = i;
-				return true;
-			}
-		}
-		return false;
-	}
-
-
-	bool ApproxSelfClipPolyline(TArray<FVector2f>& Polyline)
-	{
-		int32 N = Polyline.Num();
-
-		// handle already-closed polylines
-		if (UE::Geometry::Distance(Polyline[0], Polyline[N - 1]) < 0.0001f)
-		{
-			return true;
-		}
-
-		FVector2f IntersectPoint;
-		UE::Geometry::FIndex2i IntersectionIndex(-1, -1);
-		bool bSelfIntersects = FindPolylineSelfIntersection(Polyline, IntersectPoint, IntersectionIndex);
-		if (bSelfIntersects)
-		{
-			TArray<FVector2f> NewPolyline;
-			NewPolyline.Add(IntersectPoint);
-			for (int32 i = IntersectionIndex.A; i <= IntersectionIndex.B; ++i)
-			{
-				NewPolyline.Add(Polyline[i]);
-			}
-			NewPolyline.Add(IntersectPoint);
-			Polyline = MoveTemp(NewPolyline);
-			return true;
-		}
-
-
-		FVector2f StartDirOut = UE::Geometry::Normalized(Polyline[0] - Polyline[1]);
-		UE::Geometry::FLine2f StartLine(Polyline[0], StartDirOut);
-		FVector2f EndDirOut = UE::Geometry::Normalized(Polyline[N - 1] - Polyline[N - 2]);
-		UE::Geometry::FLine2f EndLine(Polyline[N - 1], EndDirOut);
-		UE::Geometry::FIntrLine2Line2f LineIntr(StartLine, EndLine);
-		bool bIntersects = false;
-		if (LineIntr.Find())
-		{
-			bIntersects = LineIntr.IsSimpleIntersection() && (LineIntr.Segment1Parameter > 0) && (LineIntr.Segment2Parameter > 0);
-			if (bIntersects)
-			{
-				Polyline.Add(StartLine.PointAt(LineIntr.Segment1Parameter));
-				Polyline.Add(StartLine.Origin);
-				return true;
-			}
-		}
-
-
-		UE::Geometry::FAxisAlignedBox2f Bounds;
-		for (const FVector2f& P : Polyline)
-		{
-			Bounds.Contain(P);
-		}
-		float Size = Bounds.DiagonalLength();
-
-		FVector2f StartPos = Polyline[0] + 0.001f * StartDirOut;
-		if (FindPolylineSegmentIntersection(Polyline, UE::Geometry::FSegment2f(StartPos, StartPos + 2 * Size * StartDirOut), IntersectPoint, IntersectionIndex.A))
-		{
-			return true;
-		}
-
-		FVector2f EndPos = Polyline[N - 1] + 0.001f * EndDirOut;
-		if (FindPolylineSegmentIntersection(Polyline, UE::Geometry::FSegment2f(EndPos, EndPos + 2 * Size * EndDirOut), IntersectPoint, IntersectionIndex.A))
-		{
-			return true;
-		}
-
-		return false;
-	}
-}
-
-void UDataflowEditorVertexAttributePaintTool::ComputeGradient()
-{
-	//if (!ensure(ActiveWeightMap))
-	//{
-	//	UE_LOG(LogTemp, Warning, TEXT("No active weight map"));
-	//	return;
-	//}
-
-	//BeginChange();
-
-	//const FDynamicMesh3* const Mesh = DynamicMeshComponent->GetMesh();
-	//TempROIBuffer.SetNum(0, EAllowShrinking::No);
-	//for (int32 VertexIdx : Mesh->VertexIndicesItr())
-	//{
-	//	TempROIBuffer.Add(VertexIdx);
-	//}
-
-	//if (bHaveDynamicMeshToWeightConversion)
-	//{
-	//	for (int32 VertexIdx : TempROIBuffer)
-	//	{
-	//		for (const int32 Idx : WeightToDynamicMesh[DynamicMeshToWeight[VertexIdx]])
-	//		{
-	//			ActiveWeightEditChangeTracker->SaveVertexOneRingTriangles(Idx, true);
-	//		}
-	//	}
-	//}
-	//else
-	//{
-	//	for (int32 VertexIdx : TempROIBuffer)
-	//	{
-	//		ActiveWeightEditChangeTracker->SaveVertexOneRingTriangles(VertexIdx, true);
-	//	}
-	//}
-
-
-	//for (const int32 VertexIndex : TempROIBuffer)
-	//{
-	//	const FVector3d Vert = Mesh->GetVertex(VertexIndex);
-
-	//	// (Copied from FClothPaintTool_Gradient::ApplyGradient)
-
-	//	// Get distances
-	//	// TODO: Look into surface distance instead of 3D distance? May be necessary for some complex shapes
-	//	float DistanceToLowSq = MAX_flt;
-	//	for (const int32& LowIndex : LowValueGradientVertexSelection.SelectedCornerIDs)
-	//	{
-	//		const FVector3d LowPoint = Mesh->GetVertex(LowIndex);
-	//		const float DistanceSq = (LowPoint - Vert).SizeSquared();
-	//		if (DistanceSq < DistanceToLowSq)
-	//		{
-	//			DistanceToLowSq = DistanceSq;
-	//		}
-	//	}
-
-	//	float DistanceToHighSq = MAX_flt;
-	//	for (const int32& HighIndex : HighValueGradientVertexSelection.SelectedCornerIDs)
-	//	{
-	//		const FVector3d HighPoint = Mesh->GetVertex(HighIndex);
-	//		const float DistanceSq = (HighPoint - Vert).SizeSquared();
-	//		if (DistanceSq < DistanceToHighSq)
-	//		{
-	//			DistanceToHighSq = DistanceSq;
-	//		}
-	//	}
-
-	//	const float Value = FMath::LerpStable(FilterProperties->GradientLowValue, FilterProperties->GradientHighValue, DistanceToLowSq / (DistanceToLowSq + DistanceToHighSq));
-	//	if (bHaveDynamicMeshToWeightConversion)
-	//	{
-	//		for (const int32 Idx : WeightToDynamicMesh[DynamicMeshToWeight[VertexIndex]])
-	//		{
-	//			ActiveWeightMap->SetValue(Idx, &Value);
-	//		}
-	//	}
-	//	else
-	//	{
-	//		ActiveWeightMap->SetValue(VertexIndex, &Value);
-	//	}
-	//}
-
-	//// update colors
-	//UpdateVertexColorOverlay();
-	//DynamicMeshComponent->FastNotifyVertexAttributesUpdated(EMeshRenderAttributeFlags::VertexColors);
-	//GetToolManager()->PostInvalidation();
-	//EndChange();
-
-}
-
-bool UDataflowEditorVertexAttributePaintTool::HasSelection() const
-{
-	if (MeshSelector)
-	{
-		return MeshSelector->IsAnyComponentSelected();
-	}
-	return false;
-}
-
-void UDataflowEditorVertexAttributePaintTool::GrowSelection() const
-{
-	if (MeshSelector)
-	{
-		MeshSelector->GrowSelection();
-		SetFocusInViewport();
-	}
-}
-
-void UDataflowEditorVertexAttributePaintTool::ShrinkSelection() const
-{
-	if (MeshSelector)
-	{
-		MeshSelector->ShrinkSelection();
-		SetFocusInViewport();
-	}
-}
-
-void UDataflowEditorVertexAttributePaintTool::InvertSelection() const
-{
-	if (MeshSelector)
-	{
-		MeshSelector->InvertSelection();
-		SetFocusInViewport();
-	}
-}
-
-void UDataflowEditorVertexAttributePaintTool::FloodSelection() const
-{
-	if (MeshSelector)
-	{
-		MeshSelector->FloodSelection();
-		SetFocusInViewport();
-	}
-}
-
-void UDataflowEditorVertexAttributePaintTool::SelectBorder() const
-{
-	if (MeshSelector)
-	{
-		MeshSelector->SelectBorder();
-		SetFocusInViewport();
-	}
-}
-
-
-
-void UDataflowEditorVertexAttributePaintTool::SetComponentSelectionMode(EComponentSelectionMode NewMode)
-{
-	if (ToolProperties)
-	{
-		ToolProperties->SelectionProperties.ComponentSelectionMode = NewMode;
-		if (MeshSelector)
-		{
-			MeshSelector->SetComponentSelectionMode(NewMode);
-		}
-		SetFocusInViewport();
-	}
-}
-
-const TArray<int32>& UDataflowEditorVertexAttributePaintTool::GetSelectedVertices() const
-{
-	static TArray<int32> EmptySelection;
-	if (MeshSelector)
-	{
-		return MeshSelector->GetSelectedVertices();
-	}
-	return EmptySelection;
-}
-
-void UDataflowEditorVertexAttributePaintTool::ApplyValueToSelection(EDataflowEditorToolEditOperation Operation, float InValue, bool bWithTransaction)
-{
-	if (!VertexData.IsValid())
-	{
-		return;
-	}
-
-	VertexROI = UDataflowEditorVertexAttributePaintTool::GetSelectedVertices();
-	if (VertexROI.IsEmpty())
-	{
-		return;
-	}
-
-	if (bWithTransaction)
-	{
-		BeginChange();
-	}
-	if (FDynamicMesh3* Mesh = GetSculptMesh())
-	{
-		ROIWeightValueBuffer.SetNum(VertexROI.Num(), EAllowShrinking::No);
-		SyncWeightBufferWithMesh(Mesh);
-
-		FDataflowVertexAttributePaintBrushOp::ApplyToVerticesStatic(
-			Mesh, VertexROI, ROIWeightValueBuffer, 
-			Operation, InValue, 
-			/*bApplyRadiusLimit*/false, CurrentStamp.LocalFrame.Origin, CurrentStamp.Radius
-		);
-
-		SyncMeshWithWeightBuffer(Mesh);
-	}
-
-	// update colors
-	UpdateVertexColorOverlay();
-	GetToolManager()->PostInvalidation();
-	if (bWithTransaction)
-	{
-		EndChange();
-	}
-	SetFocusInViewport();
-}
-
-void UDataflowEditorVertexAttributePaintTool::PruneSelection(float Threshold)
-{
-	if (!VertexData.IsValid())
-	{
-		return;
-	}
-
-	const TArray<int32>& SelectVertices = UDataflowEditorVertexAttributePaintTool::GetSelectedVertices();
-	if (SelectVertices.IsEmpty())
-	{
-		return;
-	}
-
-	BeginChange();
-
-	auto PruneAttributeValue =
-		[Threshold](UE::Geometry::FDynamicMeshWeightAttribute& AttributeMap, int32 AttributeIndex) -> double
-		{
-			float CurrentValue;
-			AttributeMap.GetValue(AttributeIndex, &CurrentValue);
-			return (CurrentValue >= Threshold) ? CurrentValue : 0;
-		};
-
-	for (int32 VertexIdx : SelectVertices)
-	{
-		VertexData.SetValue(VertexIdx, PruneAttributeValue);
-	}
-
-	UpdateVertexColorOverlay();
-	GetToolManager()->PostInvalidation();
-	EndChange();
-
-	SetFocusInViewport();
-}
-
-namespace UE::DataflowEditorVertexAttributePaintTool::Private
-{
-	const FString CopyAverageFromSelectionToClipboardIdentifier = TEXT("UE_DataflowEditorVertexAttributePaintTool_AverageValue:");
-}
-
-void UDataflowEditorVertexAttributePaintTool::CopyAverageFromSelectionToClipboard()
-{
-	using namespace UE::DataflowEditorVertexAttributePaintTool::Private;
-
-	const TArray<int32>& SelectVertices = GetSelectedVertices();
-	if (SelectVertices.IsEmpty())
-	{
-		const FText NotificationText = FText::FromString(TEXT("No vertices were selected. Nothing was copied to the clipboard."));
-		ShowEditorMessage(ELogVerbosity::Error, NotificationText);
-		return;
-	}
-
-	const float AverageValue = VertexData.GetAverageValue(SelectVertices);
-
-	// copy to clipboard
-	const FString ClipboardString = FString::Format(TEXT("{0}{1}"), { CopyAverageFromSelectionToClipboardIdentifier, FString::SanitizeFloat(AverageValue) });
-	FPlatformApplicationMisc::ClipboardCopy(*ClipboardString);
-
-	// notify user
-	const FText NotificationText = FText::FromString("Selection average value copied to clipboard.");
-	ShowEditorMessage(ELogVerbosity::Log, NotificationText);
-
-	SetFocusInViewport();
-}
-
-void UDataflowEditorVertexAttributePaintTool::PasteValueToSelectionFromClipboard()
-{
-	using namespace UE::DataflowEditorVertexAttributePaintTool::Private;
-
-	if (!HasSelection())
-	{
-		const FText NotificationText = FText::FromString("No vertices were selected. Nothing was pasted.");
-		ShowEditorMessage(ELogVerbosity::Error, NotificationText);
-		return;
-	}
-
-	// get the clipboard content and check if it matches our format
-	FString ClipboardContent;
-	FPlatformApplicationMisc::ClipboardPaste(ClipboardContent);
-	if (!ClipboardContent.StartsWith(CopyAverageFromSelectionToClipboardIdentifier))
-	{
-		const FText NotificationText = FText::FromString("Failed to paste value from clipboard. Format is incompatible.");
-		ShowEditorMessage(ELogVerbosity::Fatal, NotificationText);
-		return;
-	}
-
-	// parse the value from the clipboard string
-	const FString StringValue = ClipboardContent.RightChop(CopyAverageFromSelectionToClipboardIdentifier.Len());
-	float ValueToPaste = 0.0f;
-	if (StringValue.IsNumeric())
-	{
-		LexFromString(ValueToPaste, *StringValue);
-		ValueToPaste = FMath::Clamp(ValueToPaste, 0.0f, 1.0f);
-
-		ApplyValueToSelection(EDataflowEditorToolEditOperation::Replace, ValueToPaste);
-	}
-	else
-	{
-		const FText NotificationText = FText::FromString("Failed to paste value from clipboard, Value from the clipboard is not numeric.");
-		ShowEditorMessage(ELogVerbosity::Warning, NotificationText);
-	}
-	
-	// notify user
-	const FText NotificationText = FText::FromString("Pasted weights.");
-	ShowEditorMessage(ELogVerbosity::Log, NotificationText);
-
-	SetFocusInViewport();
-}
-
-void UDataflowEditorVertexAttributePaintTool::OnSelectionModified()
-{
-	//const bool bToolTypeIsGradient = (FilterProperties->SubToolType == EDataflowEditorWeightMapPaintInteractionType::Gradient);
-	//if (bToolTypeIsGradient && MeshSelector)
-	//{
-	//	const FGroupTopologySelection NewSelection = MeshSelector->GetSelectionMechanic()->GetActiveSelection(); // TODO(ccaillaud) : remove the new to access the mechanic object
-
-	//	const bool bSelectingLowValueGradientVertices = GetCtrlToggle();
-	//	if (bSelectingLowValueGradientVertices)
-	//	{
-	//		LowValueGradientVertexSelection = NewSelection;
-	//	}
-	//	else
-	//	{
-	//		HighValueGradientVertexSelection = NewSelection;
-	//	}
-
-	//	if (LowValueGradientVertexSelection.SelectedCornerIDs.Num() > 0 && HighValueGradientVertexSelection.SelectedCornerIDs.Num() > 0)
-	//	{
-	//		ComputeGradient();
-	//	}
-
-	//	constexpr bool bBroadcast = false;
-	//	MeshSelector->GetSelectionMechanic()->SetSelection(FGroupTopologySelection(), bBroadcast); // TODO(ccaillaud) should probably have a clear selection 
-	//}
-}
-
-
-void UDataflowEditorVertexAttributePaintTool::SetVerticesToWeightMap(const TSet<int32>& Vertices, double WeightValue)
-{
-	BeginChange();
-
-	TempROIBuffer.SetNum(0, EAllowShrinking::No);
-	for (int32 VertexIdx : Vertices)
-	{
-		TempROIBuffer.Add(VertexIdx);
-	}
-
-	if (HaveVisibilityFilter())
-	{
-		TArray<int32> VisibleVertices;
-		VisibleVertices.Reserve(TempROIBuffer.Num());
-		ApplyVisibilityFilter(TempROIBuffer, VisibleVertices);
-		TempROIBuffer = MoveTemp(VisibleVertices);
-	}
-
-	auto UpdateAttributeValue =
-		[WeightValue](UE::Geometry::FDynamicMeshWeightAttribute& AttributeMap, int32 AttributeIndex) -> double
-		{
-			return WeightValue;
-		};
-
-	for (int32 VertexIdx : TempROIBuffer)
-	{
-		VertexData.SetValue(VertexIdx, UpdateAttributeValue);
-	}
-
-	// update colors
-	UpdateVertexColorOverlay();
-	//DynamicMeshComponent->FastNotifyVertexAttributesUpdated(EMeshRenderAttributeFlags::VertexColors);
-	GetToolManager()->PostInvalidation();
-
-	EndChange();
-
-}
-
-bool UDataflowEditorVertexAttributePaintTool::HaveVisibilityFilter() const
-{
-	return (ToolProperties && ToolProperties->BrushProperties.VisibilityFilter != EDataflowEditorToolVisibilityType::None);
-}
-
-void UDataflowEditorVertexAttributePaintTool::ApplyVisibilityFilter(TSet<int32>& Vertices, TArray<int32>& ROIBuffer, TArray<int32>& OutputBuffer)
-{
-	ROIBuffer.SetNum(0, EAllowShrinking::No);
-	ROIBuffer.Reserve(Vertices.Num());
-	for (int32 VertexIdx : Vertices)
-	{
-		ROIBuffer.Add(VertexIdx);
-	}
-
-	OutputBuffer.Reset();
-	ApplyVisibilityFilter(ROIBuffer, OutputBuffer);
-
-	Vertices.Reset();
-	for (int32 VertexIdx : OutputBuffer)
-	{
-		Vertices.Add(VertexIdx);
-	}
-}
-
-void UDataflowEditorVertexAttributePaintTool::ApplyVisibilityFilter(const TArray<int32>& Vertices, TArray<int32>& VisibleVertices)
-{
-	if (!HaveVisibilityFilter())
-	{
-		VisibleVertices = Vertices;
-		return;
-	}
-
-	FViewCameraState StateOut;
-	GetToolManager()->GetContextQueriesAPI()->GetCurrentViewState(StateOut);
-	FTransform3d LocalToWorld = UE::ToolTarget::GetLocalToWorldTransform(Target);
-	FVector3d LocalEyePosition(LocalToWorld.InverseTransformPosition(StateOut.Position));
-
-	const FDynamicMesh3* Mesh = GetSculptMesh();
-
-	int32 NumVertices = Vertices.Num();
-
-	VisibilityFilterBuffer.SetNum(NumVertices, EAllowShrinking::No);
-	ParallelFor(NumVertices, [&](int32 idx)
-		{
-			VisibilityFilterBuffer[idx] = true;
-			UE::Geometry::FVertexInfo VertexInfo;
-			Mesh->GetVertex(Vertices[idx], VertexInfo, true, false, false);
-			const FVector3d Centroid = VertexInfo.Position;
-			const FVector3d FaceNormal = (FVector3d)VertexInfo.Normal;
-			const FVector3d CameraToVertexNormal = UE::Geometry::Normalized(Centroid - LocalEyePosition);
-			if (!FaceNormal.IsZero() && FaceNormal.Dot(CameraToVertexNormal) > 0)
-			{
-				VisibilityFilterBuffer[idx] = false;
-			}
-			else if (ToolProperties->BrushProperties.VisibilityFilter == EDataflowEditorToolVisibilityType::Unoccluded)
-			{
-				const FVector RayDirection = CameraToVertexNormal;
-				const FVector RayOrigin = LocalEyePosition;
-				const double MaxDistance = (LocalEyePosition - Centroid).Size() - UE_SMALL_NUMBER;
-				const int32 HitTID = Octree.FindNearestHitObject(FRay3d(RayOrigin, RayDirection), MaxDistance);
-				if (HitTID != IndexConstants::InvalidID && Mesh->IsTriangle(HitTID))
-				{
-					VisibilityFilterBuffer[idx] = false;
-				}
-			}
-		});
-
-	VisibleVertices.Reset();
-	for (int32 k = 0; k < NumVertices; ++k)
-	{
-		if (VisibilityFilterBuffer[k])
-		{
-			VisibleVertices.Add(Vertices[k]);
-		}
-	}
-}
-
-
-
-int32 UDataflowEditorVertexAttributePaintTool::FindHitSculptMeshTriangleConst(const FRay3d& LocalRay) const
-{
-	if (!IsInBrushMode())
-	{
-		return IndexConstants::InvalidID;
-	}
-
-	if (GetBrushCanHitBackFaces())
-	{
-		return Octree.FindNearestHitObject(LocalRay);
-	}
-	else
-	{
-		const FDynamicMesh3* Mesh = GetSculptMesh();
-
-		FViewCameraState StateOut;
-		GetToolManager()->GetContextQueriesAPI()->GetCurrentViewState(StateOut);
-		FVector3d LocalEyePosition(CurTargetTransform.InverseTransformPosition((FVector3d)StateOut.Position));
-		int HitTID = Octree.FindNearestHitObject(LocalRay,
-			[this, Mesh, &LocalEyePosition](int TriangleID) {
-				FVector3d Normal, Centroid;
-				double Area;
-				Mesh->GetTriInfo(TriangleID, Normal, Area, Centroid);
-				return Normal.Dot((Centroid - LocalEyePosition)) < 0;
-			});
-		return HitTID;
-	}
-}
-
-int32 UDataflowEditorVertexAttributePaintTool::FindHitTargetMeshTriangleConst(const FRay3d& LocalRay) const
-{
-	check(false);
-	return IndexConstants::InvalidID;
-}
-
-void UDataflowEditorVertexAttributePaintTool::UpdateHitSculptMeshTriangle(int32 TriangleID, const FRay3d& LocalRay)
-{
-	using namespace UE::Geometry;
-
-	CurrentBaryCentricCoords = FVector3d::Zero();
-
-	const FDynamicMesh3* const Mesh = GetSculptMesh();
-
-	if (Mesh->IsTriangle(TriangleID))
-	{
-		FTriangle3d Triangle;
-		Mesh->GetTriVertices(TriangleID, Triangle.V[0], Triangle.V[1], Triangle.V[2]);
-		FIntrRay3Triangle3d Query(LocalRay, Triangle);
-		Query.Find();
-		CurrentBaryCentricCoords = Query.TriangleBaryCoords;
-	}
-}
-
-bool UDataflowEditorVertexAttributePaintTool::UpdateBrushPosition(const FRay& WorldRay)
-{
-	TUniquePtr<FMeshSculptBrushOp>& UseBrushOp = GetActiveBrushOp();
-
-	bool bHit = false;
-	ESculptBrushOpTargetType TargetType = UseBrushOp->GetBrushTargetType();
-	switch (TargetType)
-	{
-	case ESculptBrushOpTargetType::SculptMesh:
-	case ESculptBrushOpTargetType::TargetMesh:
-		bHit = UpdateBrushPositionOnSculptMesh(WorldRay, false);
-		break;
-	case ESculptBrushOpTargetType::ActivePlane:
-		check(false);
-		bHit = UpdateBrushPositionOnSculptMesh(WorldRay, false);
-		break;
-	}
-
-	if (bHit && (UseBrushOp->GetAlignStampToView() || IsVolumetricBrush()))
-	{
-		AlignBrushToView();
-	}
-
-	return bHit;
-}
-
-FInputRayHit UDataflowEditorVertexAttributePaintTool::BeginHoverSequenceHitTest(const FInputDeviceRay& Ray)
-{
-	FInputRayHit RayHit = Super::BeginHoverSequenceHitTest(Ray);
-	bBrushIsHoverMesh = RayHit.bHit;
-	return RayHit;
-}
-
-bool UDataflowEditorVertexAttributePaintTool::OnUpdateHover(const FInputDeviceRay& DevicePos)
-{
-	if (ensure(InStroke() == false))
-	{
-		UpdateBrushPosition(DevicePos.WorldRay);
-	}
-	return true;
-}
-
-
-void UDataflowEditorVertexAttributePaintTool::DrawHUD(FCanvas* Canvas, IToolsContextRenderAPI* RenderAPI)
-{
-	if (ToolProperties->EditingMode == EDataflowEditorToolEditMode::Mesh)
-	{
-		if (MeshSelector)
-		{
-			MeshSelector->DrawHUD(Canvas, RenderAPI);
-		}
-	}
-
-	if (BrushEditBehavior.IsValid())
-	{
-		BrushEditBehavior->DrawHUD(Canvas, RenderAPI);
-	}
-
-	const bool bInteractiveBrushIsAdjusting = (BrushEditBehavior.IsValid() && BrushEditBehavior->IsEditing());
-
-	// when the user adjust the brush interactively we hide the value on brush
-	if (!bInteractiveBrushIsAdjusting && IsInBrushMode())
-	{
-		if (Canvas && RenderAPI && BrushIndicator && bBrushIsHoverMesh)
-		{
-			// Display the value under the brush
-			const float BrushRadius = BrushIndicator->BrushRadius;
-			FVector Offset{ 0 };
-			const FVector HoverPosition = HoverStamp.WorldFrame.ToFTransform().GetTranslation();
-
-			const FVector BrushPos = HoverPosition + Offset;
-
-			FVector4 ScreenBrushPos = RenderAPI->GetSceneView()->WorldToScreen(BrushPos);
-			FVector2D ScreenBrushPos2D = FVector2D(ScreenBrushPos.X, ScreenBrushPos.Y);
-
-			RenderAPI->GetSceneView()->ScreenToPixel(ScreenBrushPos, ScreenBrushPos2D);
-
-			// Render value under at brush even in mesh mode
-			const FText BaseText = bPendingPickWeight
-				? LOCTEXT("DataflowWeightMapPaintTool_PickValueAtBrush_TextFormat", "Pick Value: {0}")
-				: LOCTEXT("DataflowWeightMapPaintTool_ValueAtBrush_TextFormat", "Value at brush: {0}");
-			const FText ValueAtBrushText = FText::Format(BaseText, FText::AsNumber(ToolProperties->BrushProperties.ValueAtBrush));
-			FCanvasTextItem TextItem(ScreenBrushPos2D, ValueAtBrushText, GEngine->GetLargeFont(), FLinearColor::White);
-			TextItem.EnableShadow(FLinearColor::Black);
-			TextItem.bCentreX = true;
-			Canvas->DrawItem(TextItem);
-		}
-	}
-
-	if (ToolProperties && ToolProperties->DisplayProperties.bShowValues)
-	{
-		if (Canvas && RenderAPI)
-		{
-			if (const FDynamicMesh3* Mesh = GetSculptMesh())
-			{
-				FViewCameraState StateOut;
-				GetToolManager()->GetContextQueriesAPI()->GetCurrentViewState(StateOut);
-				const FTransform3d LocalToWorld = UE::ToolTarget::GetLocalToWorldTransform(Target);
-				const FVector3d LocalEyePosition(LocalToWorld.InverseTransformPosition(StateOut.Position));
-
-				auto IsVertexVisible =
-					[&LocalEyePosition, &Octree = Octree](const FDynamicMesh3& Mesh, const int32 VtxIndex)
-					{
-						UE::Geometry::FVertexInfo VertexInfo;
-						Mesh.GetVertex(VtxIndex, VertexInfo, true, false, false);
-						const FVector3d Centroid = VertexInfo.Position;
-						const FVector3d FaceNormal = FVector3d(VertexInfo.Normal);
-						const FVector3d CameraToVertexNormal = UE::Geometry::Normalized(Centroid - LocalEyePosition);
-						if (!FaceNormal.IsZero() && FaceNormal.Dot(CameraToVertexNormal) > 0)
-						{
-							return false;
-						}
-						const FVector RayDirection = CameraToVertexNormal;
-						const FVector RayOrigin = LocalEyePosition;
-						const double MaxDistance = (LocalEyePosition - Centroid).Size() - UE_SMALL_NUMBER;
-						const int32 HitTID = Octree.FindNearestHitObject(FRay3d(RayOrigin, RayDirection), MaxDistance);
-						if (HitTID != IndexConstants::InvalidID && Mesh.IsTriangle(HitTID))
-						{
-							return false;
-						}
-						return true;
-					};
-
-				FNumberFormattingOptions NumberFormattingOptions;
-				NumberFormattingOptions.AlwaysSign = false;
-				NumberFormattingOptions.UseGrouping = false;
-				NumberFormattingOptions.RoundingMode = ERoundingMode::HalfFromZero;
-				NumberFormattingOptions.MinimumIntegralDigits = 1;
-				NumberFormattingOptions.MaximumIntegralDigits = 6;
-				NumberFormattingOptions.MinimumFractionalDigits = 2;
-				NumberFormattingOptions.MaximumFractionalDigits = 2;
-
-				const float MinValue = ToolProperties->DisplayProperties.MinValue;
-				const float MaxValue = ToolProperties->DisplayProperties.MaxValue;
-				const float ValueRange = MaxValue - MinValue;
-
-				UFont* EngineLargeFont = GEngine->GetLargeFont();
-
-				const int32 MaxVertexId = Mesh->MaxVertexID();
-				TArray<bool> VertexVisibility;
-				VertexVisibility.Init(false, MaxVertexId);
-
-				ParallelFor(Mesh->MaxVertexID(), [this, &VertexVisibility, &Mesh, IsVertexVisible](int32 VertexID)
-					{
-						if (Mesh->IsVertex(VertexID))
-						{
-							if (IsVertexVisible(*Mesh, VertexID))
-							{
-								VertexVisibility[VertexID] = true;
-							}
-						}
-					});
-
-				auto DrawValue = 
-					[this, &VertexVisibility, &MinValue, &MaxValue, &ValueRange, &NumberFormattingOptions, &Canvas, &RenderAPI, &Mesh]
-					(int32 VertexID)
-					{
-						if (VertexVisibility.IsValidIndex(VertexID) && VertexVisibility[VertexID])
-						{
-							const float NormalizedValue = VertexData.GetValue(VertexID);
-							const float RemappedValue = MinValue + NormalizedValue * ValueRange;
-							const FText ValueText = FText::AsNumber(RemappedValue, &NumberFormattingOptions);
-
-							const FVector3d VtxPosition = Mesh->GetVertex(VertexID);
-							const FVector4 ScreenBrushPos = RenderAPI->GetSceneView()->WorldToScreen(VtxPosition);
-							FVector2D ScreenBrushPos2D = FVector2D(ScreenBrushPos.X, ScreenBrushPos.Y);
-							RenderAPI->GetSceneView()->ScreenToPixel(ScreenBrushPos, ScreenBrushPos2D);
-
-							FCanvasTextItem TextItem(ScreenBrushPos2D, ValueText, GEngine->GetLargeFont(), FLinearColor::White);
-							TextItem.EnableShadow(FLinearColor::Black);
-							TextItem.bCentreX = true;
-							Canvas->DrawItem(TextItem);
-						}
-					};
-
-				if (IsInBrushMode() || ToolProperties->DisplayProperties.bShowValuesOnlySelected == false)
-				{
-					for (int32 VertexID = 0; VertexID < MaxVertexId; ++VertexID)
-					{
-						DrawValue(VertexID);
-					}
-				}
-				else
-				{
-					const TArray<int32>& SelectedVertices = GetSelectedVertices();
-					for (const int32 VertexID : SelectedVertices)
-					{
-						DrawValue(VertexID);
-					}
-				}
-			}
-		}
-	}
-}
-
-void UDataflowEditorVertexAttributePaintTool::Render(IToolsContextRenderAPI* RenderAPI)
-{
-	UMeshSculptToolBase::Render(RenderAPI);
-
-	if (ToolProperties->EditingMode == EDataflowEditorToolEditMode::Mesh)
-	{
-		if (MeshSelector)
-		{
-			MeshSelector->Render(RenderAPI);
-		}
-	}
-}
-
-
-void UDataflowEditorVertexAttributePaintTool::UpdateMaterialMode(EMeshEditingMaterialModes MaterialMode)
-{
-	if (MaterialMode == EMeshEditingMaterialModes::VertexColor)
-	{
-		constexpr bool bUseTwoSidedMaterial = true;
-		ActiveOverrideMaterial = ToolSetupUtil::GetVertexColorMaterial(GetToolManager(), bUseTwoSidedMaterial);
-		if (ensure(ActiveOverrideMaterial != nullptr))
-		{
-			GetSculptMeshComponent()->SetOverrideRenderMaterial(ActiveOverrideMaterial);
-			ActiveOverrideMaterial->SetScalarParameterValue(TEXT("FlatShading"), (ViewProperties->bFlatShading) ? 1.0f : 0.0f);
-		}
-		GetSculptMeshComponent()->SetShadowsEnabled(false);
-	}
-	else
-	{
-		UMeshSculptToolBase::UpdateMaterialMode(MaterialMode);
-	}
-}
-
-void UDataflowEditorVertexAttributePaintTool::OnTick(float DeltaTime)
-{
-	UMeshSculptToolBase::OnTick(DeltaTime);
-	MeshElementsDisplay->OnTick(DeltaTime);
-
-	const bool bIsMeshEditMode = (ToolProperties->EditingMode == EDataflowEditorToolEditMode::Mesh);
-	if (MeshSelector)
-	{
-		MeshSelector->SetIsEnabled(bIsMeshEditMode);
-		MeshSelector->Tick(DeltaTime);
-	}
-
-	ConfigureIndicator(IsVolumetricBrush());
-	SetIndicatorVisibility(!bIsMeshEditMode && bBrushIsHoverMesh && !bPendingPickWeight);
-	UpdateMirrorIndicator();
-
-	SCOPE_CYCLE_COUNTER(VertexAttributePaintTool_Tick);
-
-	// process the undo update
-	if (bUndoUpdatePending)
-	{
-		// wait for updates
-		WaitForPendingUndoRedoUpdate();
-
-		// post rendering update
-		//DynamicMeshComponent->FastNotifyTriangleVerticesUpdated(AccumulatedTriangleROI, EMeshRenderAttributeFlags::VertexColors);
-		GetToolManager()->PostInvalidation();
-
-		// ignore stamp and wait for next tick to do anything else
-		bUndoUpdatePending = false;
-		return;
-	}
-
-	// Get value at brush location
-	if (IsInBrushMode())
-	{
-		if (GetSculptMesh()->IsTriangle(GetBrushTriangleID()))
-		{
-			const double HitWeightValue = GetCurrentWeightValueUnderBrush();
-			ToolProperties->BrushProperties.ValueAtBrush = HitWeightValue;
-		}
-	}
-	else
-	{
-		ToolProperties->BrushProperties.ValueAtBrush = GetAverageWeightValueFromSelection();
-	}
-
-	auto ExecuteStampOperation = [this](int StampIndex, const FRay& StampRay)
-		{
-			SCOPE_CYCLE_COUNTER(VertexAttributePaintTool_Tick_ApplyStampBlock);
-
-			FDynamicMesh3* Mesh = GetSculptMesh();
-
-			// update sculpt ROI
-			UpdateROI(CurrentStamp, GetBrushTriangleID());
-
-			// append updated ROI to modified region (async)
-			TFuture<void> AccumulateROI = Async(UE::DataflowEditorVertexAttributePaintTool::Private::WeightPaintToolAsyncExecTarget, [&]()
-				{
-					UE::Geometry::VertexToTriangleOneRing(Mesh, VertexROI, AccumulatedTriangleROI);
-				});
-
-			// apply the stamp
-			bool bWeightsModified = ApplyStamp(CurrentStamp);
-
-			if (bWeightsModified)
-			{
-				SCOPE_CYCLE_COUNTER(VertexAttributePaintTool_Tick_UpdateMeshBlock);
-				UpdateVertexColorOverlay(&TriangleROI);
-			}
-
-			// we don't really need to wait for these to happen to end Tick()...
-			AccumulateROI.Wait();
-
-			LastStamp = CurrentStamp;
-			LastStamp.TimeStamp = FDateTime::Now();
-
-			// now mirror stamp 
-			if (bBrushMirroringEnabled)
-			{
-				// update sculpt ROI
-				UpdateROI(MirrorCurrentStamp, MirrorBrushTriangleID);
-
-				// apply the stamp
-				bWeightsModified |= ApplyStamp(MirrorCurrentStamp);
-
-				if (bWeightsModified)
-				{
-					SCOPE_CYCLE_COUNTER(VertexAttributePaintTool_Tick_UpdateMeshBlock);
-					UpdateVertexColorOverlay(&TriangleROI);
-				}
-			}
-
-			if (bWeightsModified)
-			{
-				GetToolManager()->PostInvalidation();
-			}
-		};
-
-	if (IsInBrushMode() && !bPendingPickWeight)
-	{
-		ProcessPerTickStamps(
-			[this](const FRay& StampRay) -> bool {
-				return UpdateStampPosition(StampRay);
-			}, ExecuteStampOperation);
-	}
-
-	if (UToolsContextCursorAPI* ToolsContextCursorAPI = GetToolManager()->GetContextObjectStore()->FindContext<UToolsContextCursorAPI>())
-	{
-		if (bPendingPickWeight)
-		{
-			ToolsContextCursorAPI->SetCursorOverride(EMouseCursor::EyeDropper);
-		}
-		else
-		{
-			ToolsContextCursorAPI->ClearCursorOverride();
-		}
-	}
-
-}
-
-bool UDataflowEditorVertexAttributePaintTool::CanAccept() const
-{
-	return bAnyChangeMade;
-}
-
-void UDataflowEditorVertexAttributePaintTool::BeginChange()
-{
-	VertexData.BeginChange(Cast<UDynamicMeshComponent>(GetSculptMeshComponent()));
-	LongTransactions.Open(LOCTEXT("WeightPaintChange", "Weight Stroke"), GetToolManager());
-}
-
-void UDataflowEditorVertexAttributePaintTool::EndChange()
-{
-	bAnyChangeMade = true;
-
-	using namespace UE::DataflowEditorVertexAttributePaintTool;
-
-	TUniquePtr<Private::FMeshChange> MeshChange = VertexData.EndChange(Cast<UDynamicMeshComponent>(GetSculptMeshComponent()));
-
-	TUniquePtr<TWrappedToolCommandChange<Private::FMeshChange>> NewChange = MakeUnique<TWrappedToolCommandChange<Private::FMeshChange>>();
-	NewChange->WrappedChange = MoveTemp(MeshChange);
-	NewChange->BeforeModify = [this](bool bRevert)
-		{
-			this->WaitForPendingUndoRedoUpdate();
-		};
-
-	GetToolManager()->EmitObjectChange(GetSculptMeshComponent(), MoveTemp(NewChange), LOCTEXT("WeightPaintChange", "Weight Stroke"));
-	LongTransactions.Close(GetToolManager());
-}
-
-void UDataflowEditorVertexAttributePaintTool::CancelChange()
-{
-	VertexData.CancelChange();
-	LongTransactions.Close(GetToolManager());
-}
-
-
-void UDataflowEditorVertexAttributePaintTool::WaitForPendingUndoRedoUpdate()
-{
-	if (bUndoUpdatePending)
-	{
-		bUndoUpdatePending = false;
-	}
-}
-
-void UDataflowEditorVertexAttributePaintTool::OnDynamicMeshComponentChanged(UDynamicMeshComponent* Component, const FMeshVertexChange* Change, bool bRevert)
-{
-	// update octree
-	FDynamicMesh3* Mesh = GetSculptMesh();
-
-	// make sure any previous async computations are done, and update the undo ROI
-	if (bUndoUpdatePending)
-	{
-		// we should never hit this anymore, because of pre-change calling WaitForPendingUndoRedoUpdate()
-		WaitForPendingUndoRedoUpdate();
-
-		// this is not right because now we are going to do extra recomputation, but it's very messy otherwise...
-		UE::Geometry::VertexToTriangleOneRing(Mesh, Change->Vertices, AccumulatedTriangleROI);
-	}
-	else
-	{
-		AccumulatedTriangleROI.Reset();
-		UE::Geometry::VertexToTriangleOneRing(Mesh, Change->Vertices, AccumulatedTriangleROI);
-	}
-
-	// note that we have a pending update
-	bUndoUpdatePending = true;
-}
-
-
-void UDataflowEditorVertexAttributePaintTool::PrecomputeFilterData()
-{
-	const FDynamicMesh3* Mesh = GetSculptMesh();
-
-	TriNormals.SetNum(Mesh->MaxTriangleID());
-	ParallelFor(Mesh->MaxTriangleID(), [&](int32 tid)
-		{
-			if (Mesh->IsTriangle(tid))
-			{
-				TriNormals[tid] = Mesh->GetTriNormal(tid);
-			}
-		});
-
-	const UE::Geometry::FDynamicMeshNormalOverlay* Normals = Mesh->Attributes()->PrimaryNormals();
-	const UE::Geometry::FDynamicMeshUVOverlay* UVs = Mesh->Attributes()->PrimaryUV();
-	UVSeamEdges.SetNum(Mesh->MaxEdgeID());
-	NormalSeamEdges.SetNum(Mesh->MaxEdgeID());
-	ParallelFor(Mesh->MaxEdgeID(), [&](int32 eid)
-		{
-			if (Mesh->IsEdge(eid))
-			{
-				UVSeamEdges[eid] = UVs->IsSeamEdge(eid);
-				NormalSeamEdges[eid] = Normals->IsSeamEdge(eid);
-			}
-		});
-}
-
-float UDataflowEditorVertexAttributePaintTool::GetCurrentWeightValueUnderBrush() const
-{
-	using namespace UE::Geometry;
-
-	float Value = 0;
-	if (ToolProperties)
-	{
-		switch (ToolProperties->BrushProperties.ValueQueryType)
-		{
-
-		case EDataflowEditorToolValueQueryType::Interpolated:
-		{
-			const int32 Tid = GetBrushTriangleID();
-			if (Tid != IndexConstants::InvalidID)
-			{
-				const FDynamicMesh3* Mesh = GetSculptMesh();
-				const FIndex3i Vertices = Mesh->GetTriangle(Tid);
-				Value = 0.0f;
-				for (int32 TriangleVertexIndex = 0; TriangleVertexIndex < 3; ++TriangleVertexIndex)
-				{
-					const float VertexWeight = VertexData.GetValue(Vertices[TriangleVertexIndex]);
-					Value += CurrentBaryCentricCoords[TriangleVertexIndex] * VertexWeight;
-				}
-			}
-		}
-		break;
-
-		case EDataflowEditorToolValueQueryType::NearestVertexFast:
-		{
-			const int32 VertexID = GetBrushNearestVertex();
-			Value = VertexData.GetValue(VertexID);
-		}
-		break;
-
-		case EDataflowEditorToolValueQueryType::NearestVertexAccurate:
-		{
-			const int32 VertexID = GetBrushNearestVertexAccurate();
-			Value = VertexData.GetValue(VertexID);
-		}
-
-		break;
-		}
-	}
-	return Value;
-}
-
-int32 UDataflowEditorVertexAttributePaintTool::GetBrushNearestVertex() const
-{
-	int TriangleVertex = 0;
-
-	if (CurrentBaryCentricCoords.X >= CurrentBaryCentricCoords.Y && CurrentBaryCentricCoords.X >= CurrentBaryCentricCoords.Z)
-	{
-		TriangleVertex = 0;
-	}
-	else
-	{
-		if (CurrentBaryCentricCoords.Y >= CurrentBaryCentricCoords.X && CurrentBaryCentricCoords.Y >= CurrentBaryCentricCoords.Z)
-		{
-			TriangleVertex = 1;
-		}
-		else
-		{
-			TriangleVertex = 2;
-		}
-	}
-	const FDynamicMesh3* Mesh = GetSculptMesh();
-
-	int32 Tid = GetBrushTriangleID();
-	if (Tid == IndexConstants::InvalidID)
-	{
-		return IndexConstants::InvalidID;
-	}
-
-	UE::Geometry::FIndex3i Vertices = Mesh->GetTriangle(Tid);
-	return Vertices[TriangleVertex];
-}
-
-
-int32 UDataflowEditorVertexAttributePaintTool::GetBrushNearestVertexAccurate() const
-{
-	using namespace UE::Geometry;
-
-	int32 NearestVertexIndex = INDEX_NONE;
-	const int32 Tid = GetBrushTriangleID();
-
-	if (Tid != IndexConstants::InvalidID)
-	{
-		const FDynamicMesh3* const Mesh = GetSculptMesh();
-
-		FVector3d PointOnSurface(0, 0, 0);
-		const FIndex3i Vertices = Mesh->GetTriangle(Tid);
-		for (int32 TriangleVertexIndex = 0; TriangleVertexIndex < 3; ++TriangleVertexIndex)
-		{
-			PointOnSurface += CurrentBaryCentricCoords[TriangleVertexIndex] * Mesh->GetVertex(Vertices[TriangleVertexIndex]);
-		}
-
-		double MinDist = UE_BIG_NUMBER;
-		for (int32 TriangleVertexIndex = 0; TriangleVertexIndex < 3; ++TriangleVertexIndex)
-		{
-			const int32 VertexIndex = Vertices[TriangleVertexIndex];
-			const FVector3d& VertexPosition = Mesh->GetVertex(VertexIndex);
-			const double CurrDist = FVector3d::Distance(VertexPosition, PointOnSurface);
-			if (CurrDist < MinDist)
-			{
-				MinDist = CurrDist;
-				NearestVertexIndex = VertexIndex;
-			}
-		}
-	}
-
-	return NearestVertexIndex;
-}
-
-float UDataflowEditorVertexAttributePaintTool::GetAverageWeightValueFromSelection() const
-{
-	return VertexData.GetAverageValue(GetSelectedVertices());
-}
-
-void UDataflowEditorVertexAttributePaintTool::UpdateBrushType(EDataflowEditorToolEditOperation BrushMode)
-{
-	static const FText BaseMessage = LOCTEXT("OnStartTool", "Hold [Shift] to Relax values, Hold [Ctrl] to Invert Value, [/] and S/D change Size (+Shift to small-step)");
-	FTextBuilder Builder;
-	Builder.AppendLine(BaseMessage);
-
-	if (ToolProperties)
-	{
-		ToolProperties->BrushProperties.BrushMode = BrushMode;
-	}
-	SetActivePrimaryBrushType(PaintBrushId);
-	SetActiveSecondaryBrushType(SmoothBrushId);
-
-	SetToolPropertySourceEnabled(PaintBrushOpProperties, false);
-	SetToolPropertySourceEnabled(SmoothBrushOpProperties, false);
-	SetToolPropertySourceEnabled(GizmoProperties, false);
-
-	GetToolManager()->DisplayMessage(Builder.ToText(), EToolMessageLevel::UserNotification);
-}
-
-void UDataflowEditorVertexAttributePaintTool::OnColorRampChanged(TArray<FRichCurve*> Curves)
-{
-	if (ToolProperties)
-	{
-		ToolProperties->SaveConfig();
-	}
-	UpdateVertexColorOverlay();
-	GetToolManager()->PostInvalidation();
-}
-
-void UDataflowEditorVertexAttributePaintTool::SetColorMode(EDataflowEditorToolColorMode NewColorMode)
-{
-	if (ToolProperties)
-	{
-		ToolProperties->DisplayProperties.ColorMode = NewColorMode;
-		ToolProperties->SaveConfig();
-		UpdateVertexColorOverlay();
-		GetToolManager()->PostInvalidation();
-
-		FPropertyChangedEvent ColorModePropertyChangedEvent(UBrushBaseProperties::StaticClass()->FindPropertyByName(GET_MEMBER_NAME_CHECKED(FDataflowEditorVertexAttributePaintToolDisplayProperties, ColorMode)));
-		ToolProperties->PostEditChangeProperty(ColorModePropertyChangedEvent);
-
-		SetFocusInViewport();
-	}
-}
-
-void UDataflowEditorVertexAttributePaintTool::UpdateVertexColorOverlay(const TSet<int>* TrianglesToUpdate)
-{
-	auto ValueToColor = [this](float Value) -> FVector4f
-		{
-			switch (ToolProperties->DisplayProperties.ColorMode)
-			{
-			case EDataflowEditorToolColorMode::Greyscale:
-				return FMath::Lerp(FLinearColor::Black, FLinearColor::White, Value);
-			case EDataflowEditorToolColorMode::Ramp:
-				return ToolProperties->DisplayProperties.ColorRamp.GetLinearColorValue(Value);
-			case EDataflowEditorToolColorMode::FullMaterial:
-				return FLinearColor::White;
-			}
-			return FLinearColor::White;
-		};
-
-	auto SetColorsFromWeights = [&](FDynamicMesh3& Mesh, UE::Geometry::FDynamicMeshColorOverlay& ColorOverlay, int TriangleID)
-		{
-			const UE::Geometry::FIndex3i Tri = Mesh.GetTriangle(TriangleID);
-			const UE::Geometry::FIndex3i ColorElementTri = ColorOverlay.GetTriangle(TriangleID);
-
-			for (int TriVertIndex = 0; TriVertIndex < 3; ++TriVertIndex)
-			{
-				const float Value = VertexData.GetValue(Tri[TriVertIndex]);
-				ColorOverlay.SetElement(ColorElementTri[TriVertIndex], ValueToColor(Value));
-			}
-		};
-
-	// update mesh with new value colors
-	PreviewMesh->DeferredEditMesh(
-		[this, &TrianglesToUpdate, &SetColorsFromWeights](FDynamicMesh3& Mesh)
-		{
-			check(Mesh.HasAttributes());
-			check(Mesh.Attributes()->PrimaryColors());
-			
-			if (UE::Geometry::FDynamicMeshColorOverlay* ColorOverlay = Mesh.Attributes()->PrimaryColors())
-			{
-				if (TrianglesToUpdate)
-				{
-					for (const int TriangleID : *TrianglesToUpdate)
-					{
-						SetColorsFromWeights(Mesh, *ColorOverlay, TriangleID);
-					}
-				}
-				else
-				{
-					for (const int TriangleID : Mesh.TriangleIndicesItr())
-					{
-						SetColorsFromWeights(Mesh, *ColorOverlay, TriangleID);
-					}
-				}
-			}
-		}, false);
-	PreviewMesh->NotifyDeferredEditCompleted(UPreviewMesh::ERenderUpdateMode::FastUpdate, EMeshRenderAttributeFlags::VertexColors, false);
-}
-
-void UDataflowEditorVertexAttributePaintTool::OnMirrorCommand()
-{
-	if (IsInBrushMode())
-	{
-		EnableBrushMirroring(!IsBrushMirroringEnabled());
-	}
-	else
-	{
-		MirrorValues();
-	}
-}
-
-void UDataflowEditorVertexAttributePaintTool::MirrorValues()
-{
-	if (ToolProperties)
-	{
-		if (!VertexData.IsValid())
-		{
-			return;
-		}
-
-		VertexROI = UDataflowEditorVertexAttributePaintTool::GetSelectedVertices();
-		if (VertexROI.IsEmpty())
-		{
-			return;
-		}
-
-		BeginChange();
-
-		if (FDynamicMesh3* Mesh = GetSculptMesh())
-		{
-			MirrorData.EnsureMirrorDataIsUpdated(*Mesh, ToolProperties->MirrorProperties.MirrorAxis, ToolProperties->MirrorProperties.MirrorDirection, ToolProperties->MirrorProperties.bObjectSpace);
-
-			TArray<int32> VerticesToUpdate;
-			MirrorData.FindMirroredIndices(*Mesh, VertexROI, VerticesToUpdate);
-
-			// get the selected vertices values
-			ROIWeightValueBuffer.SetNum(VertexROI.Num(), EAllowShrinking::No);
-			SyncWeightBufferWithMesh(Mesh);
-
-			// now swap VertexROI with VerticesToUpdate ( ROIWeightValueBuffer map 1:1 to VerticesToUpdate)
-			Swap(VertexROI, VerticesToUpdate);
-
-			// apply the values back to the mesh with the mirrored indices
-			SyncMeshWithWeightBuffer(Mesh);
-		}
-
-		// update colors
-		UpdateVertexColorOverlay();
-		GetToolManager()->PostInvalidation();
-		EndChange();
-
-		SetFocusInViewport();
-	}
-}
-
-void UDataflowEditorVertexAttributePaintTool::AddReferencedObjects(UObject* InThis, FReferenceCollector& Collector)
+void UDataflowEditorVertexAttributePaintTool::AddReferencedObjects(UObject* InThis, FReferenceCollector& Collector)
 {
 	UDataflowEditorVertexAttributePaintTool* This = CastChecked<UDataflowEditorVertexAttributePaintTool>(InThis);
-	Collector.AddReferencedObject(This->PreviewMesh);
-	Collector.AddReferencedObject(This->MeshSelector);
-	Collector.AddReferencedObject(This->MeshElementsDisplay);
 	Collector.AddReferencedObject(This->DataflowEditorContextObject);
 	Super::AddReferencedObjects(InThis, Collector);
 }
 
-/////////////////////////////////////////////////////////////////////////////////////////////////////////////////////////////////////////////////////////////////
-// 
-// MirrorData
-// 
-/////////////////////////////////////////////////////////////////////////////////////////////////////////////////////////////////////////////////////////////////
-
-bool UDataflowEditorVertexAttributePaintTool::FMirrorData::IsPointOnTargetMirrorSide(const FVector& InPoint) const
-{
-	if (Axis == EAxis::None)
-	{
-		return false;
-	}
-	const int32 AxisIndex = (Axis - 1);
-	if (Direction == EDataflowEditorToolMirrorDirection::PositiveToNegative && InPoint[AxisIndex] >= PlaneOffsets[AxisIndex])
-	{
-		return false; // target is negative side, but point is on positive side
-	}
-	if (Direction == EDataflowEditorToolMirrorDirection::NegativeToPositive && InPoint[AxisIndex] <= PlaneOffsets[AxisIndex])
-	{
-		return false; // target is positive side, but vertex is on negative side
-	}
-	return true;
-}
-
-void UDataflowEditorVertexAttributePaintTool::FMirrorData::EnsureMirrorDataIsUpdated(const FDynamicMesh3& Mesh, EAxis::Type InMirrorAxis, EDataflowEditorToolMirrorDirection InMirrorDirection, bool bInObjectSpace)
-{
-	if (VertexMap.Num() > 0 && InMirrorAxis == Axis && InMirrorDirection == Direction && bObjectSpace == bInObjectSpace)
-	{
-		// already initialized, just re-use cached data
-		return;
-	}
-
-	// need to re-initialize
-	Axis = InMirrorAxis;
-	Direction = InMirrorDirection;
-	bObjectSpace = bInObjectSpace;
-	VertexMap.Reset();
-
-	const int32 NumVertices = Mesh.VertexCount();
-
-	// build a spatial hash grid
-	constexpr float HashGridCellSize = 2.0f; // the length of the cell size in the point hash grid
-	UE::Geometry::TPointHashGrid3f<int32> VertHash(HashGridCellSize, INDEX_NONE);
-	VertHash.Reserve(NumVertices);
-
-	PlaneOffsets = FVector{ 0, 0, 0 };
-	if (bObjectSpace)
-	{
-		const UE::Geometry::FAxisAlignedBox3d MeshBounds = Mesh.GetBounds(true);
-		PlaneOffsets = MeshBounds.Center();
-	}
-
-	for (int32 VertexID : Mesh.VertexIndicesItr())
-	{
-		const FVector VertexPos = Mesh.GetVertex(VertexID);
-		VertHash.InsertPointUnsafe(VertexID, static_cast<FVector3f>(VertexPos));
-	}
-
-	// generate a map of point IDs on the target side, to their equivalent vertex ID on the source side 
-	for (int32 TargetVertexID : Mesh.VertexIndicesItr())
-	{
-		const FVector SourcePosition = Mesh.GetVertex(TargetVertexID);
-
-		// flip position across the mirror axis
-		FVector3f MirroredPosition = FVector3f(SourcePosition);
-		const int32 AxisIndex = (Axis - 1);
-		MirroredPosition[AxisIndex] = ((MirroredPosition[AxisIndex] - PlaneOffsets[AxisIndex]) *  (-1.0f)) + PlaneOffsets[AxisIndex];
-
-		// Query spatial hash near mirrored position, gradually increasing search radius until at least 1 point is found
-		TPair<int32, double> ClosestMirroredPoint = { INDEX_NONE, TNumericLimits<double>::Max() };
-		float SearchRadius = HashGridCellSize;
-		while (ClosestMirroredPoint.Key == INDEX_NONE)
-		{
-			ClosestMirroredPoint = VertHash.FindNearestInRadius(
-				MirroredPosition,
-				SearchRadius,
-				[&Mesh, MirroredPosition](int32 VID)
-				{
-					return FVector3f::DistSquared(FVector3f(Mesh.GetVertex(VID)), MirroredPosition);
-				});
-
-			SearchRadius += HashGridCellSize;
-
-			// forcibly break out if search radius gets bigger than the maximum search radius
-			static float MaxSearchRadius = 15.f; // TODO we may want to expose this value to the user...
-			if (SearchRadius >= MaxSearchRadius)
-			{
-				break;
-			}
-		}
-
-		if (ClosestMirroredPoint.Key != INDEX_NONE)
-		{
-			VertexMap.FindOrAdd(TargetVertexID, ClosestMirroredPoint.Key);
-		}
-	}
-}
-
-void UDataflowEditorVertexAttributePaintTool::FMirrorData::FindMirroredIndices(const FDynamicMesh3& Mesh, const TArray<int32>& SelectedVertices, TArray<int32>& OutVerticesToUpdate)
-{
-	check(Axis != EAxis::None);
-
-	// results have a 1:1 mapping to the input vertices
-	OutVerticesToUpdate.SetNum(SelectedVertices.Num());
-
-	// we need to convert selection to the equivalent target vertex indices (on the target side of the mirror plane)
-	// if a vertex is already on the target side, great
-	// if the user selected vertices on the source side, we convert them to the mirrored equivalent on the target side
-	TSet<VertexIndex> TargetVertices;
-	TArray<VertexIndex> MissingVertices;
-	for (int32 Index = 0; Index < SelectedVertices.Num(); ++Index)
-	{
-		const int32 SelectedVertexID = SelectedVertices[Index];
-		int32& TargetVertexID = OutVerticesToUpdate[Index];
-
-		// get the mirrored index or the invaliod index if not found
-		OutVerticesToUpdate[Index] = VertexMap.FindRef(SelectedVertexID, INDEX_NONE);
-
-		// if the selected vertex is on the wrong side, reset to an invalid index
-		if (Mesh.IsVertex(TargetVertexID))
-		{
-			const FVector TargetPosition = Mesh.GetVertex(TargetVertexID);
-			if (!IsPointOnTargetMirrorSide(TargetPosition))
-			{
-				OutVerticesToUpdate[Index] = INDEX_NONE;
-			}
-		}
-	}
-}
-
 
 #undef LOCTEXT_NAMESPACE
 
