[ModelingTools Exp] Extract most of data flow vert paint tool into a sharable base class

diff --git a/Engine/Plugins/Experimental/Dataflow/Source/DataflowEditor/Public/DataflowEditorTools/DataflowEditorVertexAttributePaintTool.h b/Engine/Plugins/Experimental/Dataflow/Source/DataflowEditor/Public/DataflowEditorTools/DataflowEditorVertexAttributePaintTool.h
--- a/Engine/Plugins/Experimental/Dataflow/Source/DataflowEditor/Public/DataflowEditorTools/DataflowEditorVertexAttributePaintTool.h
+++ b/Engine/Plugins/Experimental/Dataflow/Source/DataflowEditor/Public/DataflowEditorTools/DataflowEditorVertexAttributePaintTool.h
@@ -8,63 +8,34 @@
 #include "BaseTools/BaseBrushTool.h"
 #include "BaseTools/MeshSurfacePointMeshEditingTool.h"
 #include "Components/DynamicMeshComponent.h"
-#include "PropertySets/WeightMapSetProperties.h"
 
 #include "Sculpting/MeshSculptToolBase.h"
-#include "Sculpting/MeshBrushOpBase.h"
 
-#include "DynamicMesh/DynamicMeshAABBTree3.h"
 #include "DynamicMesh/DynamicMeshOctree3.h"
 #include "DynamicMesh/DynamicMeshAttributeSet.h"
-#include "DynamicMesh/DynamicMeshChangeTracker.h"
-#include "DynamicMesh/MeshNormals.h"
-#include "TransformTypes.h"
+
 #include "ToolDataVisualizer.h"
 #include "GroupTopology.h"
-#include "MeshVertexPaintTool.h"
+#include "MeshVertexAttributePaintToolBase.h"
+
 #include "DataflowEditorTools/DataflowEditorToolBuilder.h"
-#include "Curves/LinearColorRamp.h"
-#include "DataflowEditorTools/DataflowEditorToolEnums.h"
-#include "ToolMeshSelector.h"
 
-#include "Framework/Commands/InputChord.h"
 
 #include "DataflowEditorVertexAttributePaintTool.generated.h"
 
 
 #define UE_API DATAFLOWEDITOR_API
 
-class IToolsContextTransactionsAPI;
-class UMeshElementsVisualizer;
-class UDataflowWeightMapSmoothBrushOpProps;
-class UDataflowVertexAttributePaintBrushOpProps;
 class UDataflowContextObject;
 struct FDataflowVertexAttributeEditableNode;
 struct FDataflowVertexAttributeProviderNode;
 class UDataflowEditorMode;
 
-namespace UE::DataflowEditorVertexAttributePaintTool::Private
-{
-	class FMeshChange;
-}
-
 namespace UE::DataflowEditorVertexAttributePaintTool::CVars
 {
 	extern bool DataflowEditorUseNewWeightMapTool;
 }
 
-namespace UE::Dataflow::Editor
-{
-	class FVertexAttributePaintToolDetailCustomization;
-}
-
-DECLARE_STATS_GROUP(TEXT("Dataflow_VertexAttributePaintTool"), STATGROUP_VertexAttributePaintTool, STATCAT_Advanced);
-DECLARE_CYCLE_STAT(TEXT("Dataflow_VertexAttributePaintTool_UpdateROI"), VertexAttributePaintTool_UpdateROI, STATGROUP_VertexAttributePaintTool);
-DECLARE_CYCLE_STAT(TEXT("Dataflow_VertexAttributePaintTool_ApplyStamp"), VertexAttributePaintTool_ApplyStamp, STATGROUP_VertexAttributePaintTool);
-DECLARE_CYCLE_STAT(TEXT("Dataflow_VertexAttributePaintTool_Tick"), VertexAttributePaintTool_Tick, STATGROUP_VertexAttributePaintTool);
-DECLARE_CYCLE_STAT(TEXT("Dataflow_VertexAttributePaintTool_Tick_ApplyStampBlock"), VertexAttributePaintTool_Tick_ApplyStampBlock, STATGROUP_VertexAttributePaintTool);
-DECLARE_CYCLE_STAT(TEXT("Dataflow_VertexAttributePaintTool_Tick_UpdateMeshBlock"), VertexAttributePaintTool_Tick_UpdateMeshBlock, STATGROUP_VertexAttributePaintTool);
-
 ///////////////////////////////////////////////////////////////////////////////////////////////////////////////////////////////////////////////////
 
 /**
@@ -91,303 +62,17 @@ private:
 
 ///////////////////////////////////////////////////////////////////////////////////////////////////////////////////////////////////////////////////
 
-USTRUCT()
-struct FDataflowEditorVertexAttributePaintToolBrushProperties
-{
-	GENERATED_BODY()
-
-	UPROPERTY()
-	EDataflowEditorToolEditOperation BrushMode = EDataflowEditorToolEditOperation::Replace;
-
-	/** Value to paint on the mesh (dependent on the BrushMode) */
-	UPROPERTY()
-	float AttributeValue = 1.0;
-
-	/** Area Mode specifies the shape of the brush and which triangles will be included relative to the cursor */
-	UPROPERTY(EditAnywhere, Category = Filters, meta = (DisplayName = "Brush Area Mode"))
-	EMeshVertexPaintBrushAreaType BrushAreaMode = EMeshVertexPaintBrushAreaType::Connected;
-
-	/** The Region affected by the current operation will be bounded by edge angles larger than this threshold */
-	UPROPERTY(EditAnywhere, Category = Filters, meta = (UIMin = "0.0", ClampMin = "0.0", UIMax = "180.0", ClampMax = "180.0"))
-	float AngleThreshold = 180.0f;
-
-	/** The Region affected by the current operation will be bounded by UV borders/seams */
-	UPROPERTY(EditAnywhere, Category = Filters)
-	bool bUVSeams = false;
-
-	/** The Region affected by the current operation will be bounded by Hard Normal edges/seams */
-	UPROPERTY(EditAnywhere, Category = Filters)
-	bool bNormalSeams = false;
-
-	/** Control which triangles can be affected by the current operation based on visibility. Applied after all other filters. */
-	UPROPERTY(EditAnywhere, Category = Filters)
-	EDataflowEditorToolVisibilityType VisibilityFilter = EDataflowEditorToolVisibilityType::None;
-
-	/** The weight value at the brush indicator */
-	UPROPERTY(VisibleAnywhere, Transient, Category = Query, meta = (NoResetToDefault))
-	double ValueAtBrush = 0;
-
-	/** method used for querying the value at brush */
-	UPROPERTY(EditAnywhere, Category = Filters)
-	EDataflowEditorToolValueQueryType ValueQueryType = EDataflowEditorToolValueQueryType::NearestVertexFast;
-};
-
-USTRUCT()
-struct FDataflowEditorVertexAttributePaintToolBrushConfig
+struct FDataflowEditorVertexAttributePaintToolDataAdapter
 {
-	GENERATED_BODY()
-
 public:
-	/** Adaptive size of brush - Relative to the model size */
-	UPROPERTY(Config)
-	float BrushSize = 0.25f;
-
-	/** Value to paint on the mesh (dependent on the BrushMode) */
-	UPROPERTY(Config)
-	float Value = 1.0;
-
-	UPROPERTY(Config)
-	float Falloff = 0.0;
-};
-
-///////////////////////////////////////////////////////////////////////////////////////////////////////////////////////////////////////////////////
-
-USTRUCT()
-struct FDataflowEditorVertexAttributePaintToolGradientProperties
-{
-	GENERATED_BODY()
-
-	// Mesh selection mode (vertex, edge face)
-	/** The Gradient upper limit value */
-	UPROPERTY(EditAnywhere, Category = Gradient, meta = (UIMin = 0, ClampMin = 0, UIMax = 1, ClampMax = 1))
-	double GradientHighValue = 1.0;
-
-	/** The Gradient lower limit value */
-	UPROPERTY(EditAnywhere, Category = Gradient, meta = (UIMin = 0, ClampMin = 0, UIMax = 1, ClampMax = 1))
-	double GradientLowValue = 0.0;
-};
-
-///////////////////////////////////////////////////////////////////////////////////////////////////////////////////////////////////////////////////
-
-USTRUCT()
-struct FDataflowEditorVertexAttributePaintToolSelectionProperties
-{
-	GENERATED_BODY()
-
-	// Mesh selection mode (vertex, edge face)
-	UPROPERTY()
-	EComponentSelectionMode ComponentSelectionMode = EComponentSelectionMode::Vertices;
-};
-
-///////////////////////////////////////////////////////////////////////////////////////////////////////////////////////////////////////////////////
-
-USTRUCT()
-struct FDataflowEditorVertexAttributePaintToolDisplayProperties
-{
-	GENERATED_BODY()
-
-	// Vertex color display mode (Greyscale, Ramp, ... )
-	UPROPERTY()
-	EDataflowEditorToolColorMode ColorMode = EDataflowEditorToolColorMode::Greyscale;
-
-	// Color ramp to use when the color mode is set to "Ramp"
-	UPROPERTY()
-	FLinearColorRamp ColorRamp;
-
-	UPROPERTY()
-	TArray<FLinearColorRamp> ColorRampPresets;
-
-	// Used by UI customization 
-	FLinearColorRamp GreyScaleColorRamp;
-	FLinearColorRamp WhiteColorRamp;
-
-	UPROPERTY()
-	bool bShowValues = false;
-
-	UPROPERTY()
-	bool bShowValuesOnlySelected = false;
-
-	UPROPERTY()
-	float MinValue = 0;
-
-	UPROPERTY()
-	float MaxValue = 100;
-
-	/** Should mesh boundary edges be shown */
-	UPROPERTY()
-	bool bShowBorders = true;
-
-	/** Should mesh uv seam edges be shown */
-	UPROPERTY()
-	bool bShowUVSeams = true;
-
-	/** Should mesh normal seam edges be shown */
-	UPROPERTY()
-	bool bShowNormalSeams = false;
-
-	/** Should mesh tangent seam edges be shown */
-	UPROPERTY()
-	bool bShowTangentSeams = true;
-
-	/** Should mesh color seam edges be shown */
-	UPROPERTY()
-	bool bShowColorSeams = true;
-
-	/** Multiplier on edge thicknesses */
-	UPROPERTY()
-	float ThicknessScale = 0.5f;
-
-	/** Color of mesh wireframe */
-	UPROPERTY()
-	FColor WireframeColor = FColor(128, 128, 128);
-
-	/** Color of mesh boundary edges */
-	UPROPERTY()
-	FColor BoundaryEdgeColor = FColor(164, 73, 164);
-
-	/** Color of mesh UV seam edges */
-	UPROPERTY()
-	FColor UVSeamColor = FColor(240, 160, 15);
-
-	/** Color of mesh normal seam edges */
-	UPROPERTY()
-	FColor NormalSeamColor = FColor(128, 128, 240);
-
-	/** Color of mesh tangent seam edges */
-	UPROPERTY()
-	FColor TangentSeamColor = FColor(64, 240, 240);
-
-	/** Color of mesh color seam edges */
-	UPROPERTY()
-	FColor ColorSeamColor = FColor(46, 204, 113);
-
-	/** Depth bias used to slightly shift depth of lines */
-	UPROPERTY()
-	float DepthBias = 0.2;
-};
-
-///////////////////////////////////////////////////////////////////////////////////////////////////////////////////////////////////////////////////
-
-USTRUCT()
-struct FDataflowEditorVertexAttributePaintToolMirrorProperties
-{
-	GENERATED_BODY()
-
-	/* if true, the mirror planes will be centered on the bounds of the object */
-	UPROPERTY()
-	bool bObjectSpace = true;
-
-	/* if true, the mirror planes will be hidden when actively painting */
-	UPROPERTY()
-	bool bHideOnBrushStroke = false;
-
-	UPROPERTY()
-	TEnumAsByte<EAxis::Type> MirrorAxis = EAxis::X;
-
-	UPROPERTY()
-	EDataflowEditorToolMirrorDirection MirrorDirection = EDataflowEditorToolMirrorDirection::PositiveToNegative;
-};
-
-///////////////////////////////////////////////////////////////////////////////////////////////////////////////////////////////////////////////////
-
-UCLASS(MinimalAPI, config = EditorSettings)
-class UDataflowEditorVertexAttributePaintToolProperties : public UInteractiveToolPropertySet
-{
-	GENERATED_BODY()
-
-	UE_API UDataflowEditorVertexAttributePaintToolProperties();
-
-public:
-
-	// Edit mode (brush,mesh...)
-	UPROPERTY()
-	EDataflowEditorToolEditMode EditingMode = EDataflowEditorToolEditMode::Brush;
+	void Setup(
+		const FDynamicMesh3& InMesh,
+		UE::Geometry::FDynamicMeshWeightAttribute* InActiveWeightMap,
+		FDataflowVertexAttributeEditableNode* InNodeToUpdate,
+		TObjectPtr<UDataflowContextObject> InDataflowEditorContextObject
+		);
 	
-	UPROPERTY(config)
-	FDataflowEditorVertexAttributePaintToolBrushProperties BrushProperties;
-
-	UPROPERTY()
-	FDataflowEditorVertexAttributePaintToolGradientProperties GradientProperties;
-
-	UPROPERTY()
-	FDataflowEditorVertexAttributePaintToolSelectionProperties SelectionProperties;
-
-	UPROPERTY(config)
-	FDataflowEditorVertexAttributePaintToolDisplayProperties DisplayProperties;
-
-	UPROPERTY(config)
-	FDataflowEditorVertexAttributePaintToolMirrorProperties MirrorProperties;
-
-	const FDataflowEditorVertexAttributePaintToolBrushConfig& GetBrushConfig(EDataflowEditorToolEditOperation BrushMode) const;
-	void SetBrushConfig(EDataflowEditorToolEditOperation BrushMode, const FDataflowEditorVertexAttributePaintToolBrushConfig& BrushConfig);
-
-	UPROPERTY(config)
-	FDataflowEditorVertexAttributePaintToolBrushConfig BrushConfigAdd;
-
-	UPROPERTY(config)
-	FDataflowEditorVertexAttributePaintToolBrushConfig BrushConfigReplace;
-
-	UPROPERTY(config)
-	FDataflowEditorVertexAttributePaintToolBrushConfig BrushConfigMultiply;
-
-	UPROPERTY(config)
-	FDataflowEditorVertexAttributePaintToolBrushConfig BrushConfigRelax;
-
-	void ResetColorRamp();
-
-private:
-	void CreateDefaultColorRamp();
-
-};
-
-///////////////////////////////////////////////////////////////////////////////////////////////////////////////////////////////////////////////////
-
-struct FDataflowEditorVertexAttributePaintToolData
-{
-public:
-	void Setup(FDynamicMesh3& InMesh, FDataflowVertexAttributeEditableNode* InNodeToUpdate, TObjectPtr<UDataflowContextObject> InDataflowEditorContextObject);
-	void Shutdown();
-
-	void BeginChange(UDynamicMeshComponent* Component);
-	TUniquePtr<UE::DataflowEditorVertexAttributePaintTool::Private::FMeshChange> EndChange(UDynamicMeshComponent* Component);
-	void CancelChange();
-
-	bool IsValid() const { return (ActiveWeightMap != nullptr); }
-
-	float GetValue(int32 VertexIdx) const;
-	float GetAverageValue(const TArray<int32>& Vertices) const;
-
-	// lambda signature : float GetValue(UE::Geometry::FDynamicMeshWeightAttribute& AttributeMap, int32 AttributeIndex)
-	template <typename TLambda>
-	void SetValue(int32 VertexIdx, TLambda Lambda)
-	{
-		if (ActiveWeightMap)
-		{
-			if (bHaveDynamicMeshToWeightConversion)
-			{
-				for (const int32 Idx : WeightToDynamicMesh[DynamicMeshToWeight[VertexIdx]])
-				{
-					if (ActiveWeightEditChangeTracker)
-					{
-						ActiveWeightEditChangeTracker->SaveVertexOneRingTriangles(Idx, true);
-					}
-					const float Value = Lambda(*ActiveWeightMap, Idx);
-					ActiveWeightMap->SetValue(Idx, &Value);
-				}
-			}
-			else
-			{
-				if (ActiveWeightEditChangeTracker)
-				{
-					ActiveWeightEditChangeTracker->SaveVertexOneRingTriangles(VertexIdx, true);
-				}
-				const float Value = Lambda(*ActiveWeightMap, VertexIdx);
-				ActiveWeightMap->SetValue(VertexIdx, &Value);
-			}
-		}
-	}
-
-	void UpdateNode(TObjectPtr<UDataflowContextObject> InDataflowEditorContextObject, IToolsContextTransactionsAPI* TransactionAPI);
+	void CommitToNode(TObjectPtr<UDataflowContextObject> InDataflowEditorContextObject, IToolsContextTransactionsAPI* TransactionAPI);
 
 	bool HaveDynamicMeshToWeightConversion() const { return bHaveDynamicMeshToWeightConversion; }
 	const TArray<int32>& GetDynamicMeshToWeight() const { return  DynamicMeshToWeight; }
@@ -403,8 +88,6 @@ private:
 	FDataflowVertexAttributeEditableNode* NodeToUpdate = nullptr;
 
 	UE::Geometry::FDynamicMeshWeightAttribute* ActiveWeightMap = nullptr;
-
-	TUniquePtr<UE::Geometry::FDynamicMeshChangeTracker> ActiveWeightEditChangeTracker;
 };
 
 ///////////////////////////////////////////////////////////////////////////////////////////////////////////////////////////////////////////////////
@@ -413,300 +96,43 @@ private:
  * Dataflow Editor tool to paint a vertex attribute as vertex colors
  */
 UCLASS(MinimalAPI)
-class UDataflowEditorVertexAttributePaintTool : public UMeshSculptToolBase
+class UDataflowEditorVertexAttributePaintTool : public UMeshVertexAttributePaintToolBase
 {
 	GENERATED_BODY()
 
 public:
-	UE_API virtual void RegisterActions(FInteractiveToolActionSet& ActionSet) override;
+	UE_API virtual bool SetupToolMesh(FDynamicMesh3& InOutToolMesh, int32& OutInitialAttributeIndex) override;
+	UE_API virtual void CommitToolMesh(FDynamicMesh3& InToolMesh) override;
+	
+public:
 
 	UE_API virtual void Setup() override;
-	UE_API virtual void Shutdown(EToolShutdownType ShutdownType) override;
-	UE_API virtual void DrawHUD(FCanvas* Canvas, IToolsContextRenderAPI* RenderAPI) override;
-	UE_API virtual void Render(IToolsContextRenderAPI* RenderAPI) override;
-	UE_API virtual void UpdateMaterialMode(EMeshEditingMaterialModes MaterialMode) override;
-	UE_API virtual void OnTick(float DeltaTime) override;
-
-	virtual bool HasCancel() const override { return true; }
-	virtual bool HasAccept() const override { return true; }
-	UE_API virtual bool CanAccept() const override;
-
-	UE_API virtual FInputRayHit BeginHoverSequenceHitTest(const FInputDeviceRay& Ray)  override;
-	UE_API virtual bool OnUpdateHover(const FInputDeviceRay& DevicePos) override;
-
-	UE_API virtual void OnPropertyModified(UObject* PropertySet, FProperty* Property) override;
-
-	UE_API virtual FBox GetWorldSpaceFocusBox() override;
-
-	UE_API void SetEditingMode(EDataflowEditorToolEditMode Editmode);
-
-	UE_API void SetBrushMode(EDataflowEditorToolEditOperation BrushMode);
-	UE_API bool IsInBrushMode() const;
-
-	UE_API void SetBrushAreaMode(EMeshVertexPaintBrushAreaType BrushAreaMode);
-	UE_API bool IsVolumetricBrush() const;
-
-	UE_API virtual void CommitResult(UBaseDynamicMeshComponent* Component, bool bModifiedTopology) override;
 
 	UE_API void SetDataflowEditorContextObject(TObjectPtr<UDataflowContextObject> InDataflowEditorContextObject);
 	void SetEditorMode(UDataflowEditorMode* InMode) { Mode = InMode; }
 
-	virtual void SetupBrushEditBehaviorSetup(ULocalTwoAxisPropertyEditInputBehavior& OutBehavior) override;
-
 	//~ UObject interface
-
 	static UE_API void AddReferencedObjects(UObject* InThis, FReferenceCollector& Collector);
 
 public:
-	/** properties of the tool */
-	UPROPERTY()
-	TObjectPtr<UDataflowEditorVertexAttributePaintToolProperties> ToolProperties;
-
-	UPROPERTY()
-	TObjectPtr<UDataflowVertexAttributePaintBrushOpProps> PaintBrushOpProperties;
-
-	UPROPERTY()
-	TObjectPtr<UDataflowWeightMapSmoothBrushOpProps> SmoothBrushOpProperties;
-
-	UE_API bool HaveVisibilityFilter() const;
-	UE_API void ApplyVisibilityFilter(const TArray<int32>& Vertices, TArray<int32>& VisibleVertices);
-	UE_API void ApplyVisibilityFilter(TSet<int32>& Triangles, TArray<int32>& ROIBuffer, TArray<int32>& OutputBuffer);
-
-	using UMeshSculptToolBase::GetCurrentBrushRadius;
-	UE_API float GetBrushAdaptiveSize() const;
-	UE_API float GetBrushMinRadius() const;
-	UE_API float GetBrushMaxRadius() const;
-	UE_API void SetBrushRadius(float NewRadius);
-	UE_API void SetBrushAdaptiveSize(float AdaptiveSize);
-
-	UE_API float GetBrushFalloff() const;
-	UE_API void SetBrushFalloff(float NewFalloff);
-
-	UE_API void SetColorMode(EDataflowEditorToolColorMode NewColorMode);
-
-	UE_API bool HasSelection() const;
-	UE_API void GrowSelection() const;
-	UE_API void ShrinkSelection() const;
-	UE_API void InvertSelection() const;
-	UE_API void FloodSelection() const;
-	UE_API void SelectBorder() const;
-	UE_API void SetComponentSelectionMode(EComponentSelectionMode NewMode);
-
-	UE_API void CopyAverageFromSelectionToClipboard();
-	UE_API void PasteValueToSelectionFromClipboard();
-	UE_API void PruneSelection(float Threshold);
-
-	// selection operations
-	UE_API void ApplyValueToSelection(EDataflowEditorToolEditOperation Operation, float InValue, bool bWithTransaction = true);
-	UE_API void MirrorValues();
-
-	UE_API const TArray<int32>& GetSelectedVertices() const;
-
-	UE_API void SetFocusInViewport() const;
-
-public:
-	static FInputChord GreyscaleDisplayInputChord;
-	static FInputChord ColorMapDisplayInputChord;
-	static FInputChord MaterialColorDisplayInputChord;
+	UE_API virtual void SetFocusInViewport() const override;
 
-	static FInputChord VerticesSelectionModeInputChord;
-	static FInputChord EdgesSelectionModeInputChord;
-	static FInputChord FacesSelectionModeInputChord;
-
-	static FInputChord GrowSelectionInputChord;
-	static FInputChord ShrinkSelectionInputChord;
-	static FInputChord FloodSelectionInputChord;
-
-	static FInputChord InvertMirrorDirectionInputChord;
-	static FInputChord MirrorValuesInputChord;
-
-	static FInputChord CopyAverageValueInputChord;
-	static FInputChord PasteAverageValueInputChord;
-
-	static FInputChord ToggleValuePickerInputChord;
-
-protected:
-	// UMeshSculptToolBase API
-	virtual UBaseDynamicMeshComponent* GetSculptMeshComponent() override;
-	virtual FDynamicMesh3* GetBaseMesh()  override { check(false); return nullptr; }
-	virtual const FDynamicMesh3* GetBaseMesh() const  override { check(false); return nullptr; }
-
-	UE_API virtual int32 FindHitSculptMeshTriangleConst(const FRay3d& LocalRay) const override;
-	UE_API virtual int32 FindHitTargetMeshTriangleConst(const FRay3d& LocalRay) const override;
-	UE_API virtual void UpdateHitSculptMeshTriangle(int32 TriangleID, const FRay3d& LocalRay) override;
-
-	UE_API virtual void OnBeginStroke(const FRay& WorldRay) override;
-	UE_API virtual void OnEndStroke() override;
-	UE_API virtual void OnCancelStroke() override;
-
-	UE_API void OnColorRampChanged(TArray<FRichCurve*> Curves);
-
-	void SaveBrushConfig(EDataflowEditorToolEditOperation BrushMode);
-	void LoadBrushConfig(EDataflowEditorToolEditOperation BrushMode);
-
-	virtual bool SharesBrushPropertiesChanges() const override { return false; }
-	virtual void InitializeIndicator() override;
-	// end UMeshSculptToolBase API
-
-
-	void EnableValuePicker(bool bEnable);
-	bool IsValuePickerEnabled() const;
-
-	void EnableBrushMirroring(bool bEnable);
-	bool IsBrushMirroringEnabled() const;
-
-	void SetMirrorPlaneWidgetVisible(bool bVisible);
-	bool IsMirrorPlaneWidgetVisible();
-
-	bool IsBrushMirroringInObjectSpace() const;
-
-	void SetMirrorAxis(EAxis::Type InAxis);
-	EAxis::Type GetMirrorAxis() const;
-
-	void InitializeMirrorIndicator();
-	void UpdateMirrorIndicator();
-	void UpdateMirrorPlaneMesh();
-	UE::Geometry::FFrame3d ComputeMirroredFrame(const UE::Geometry::FFrame3d& InLocalSpaceFrame) const;
-	void OnMirrorCommand();
-
-	void SetMeshSelectionTool(EMeshSelectorTool SelectorTool);
-	EMeshSelectorTool GetMeshSelectionTool() const;
-
-protected:
-
-	// 
-	// Gradient Support
-	//
-	FToolDataVisualizer GradientSelectionRenderer;
-
-	UE::Geometry::FGroupTopologySelection LowValueGradientVertexSelection;
-	UE::Geometry::FGroupTopologySelection HighValueGradientVertexSelection;
-
-	UE_API void ComputeGradient();
-	UE_API void OnSelectionModified();
 
 	//
 	// Internals
 	//
 
 protected:
-
-	UPROPERTY()
-	TObjectPtr<UMeshElementsVisualizer> MeshElementsDisplay = nullptr;
-
 	UPROPERTY()
 	TObjectPtr<UDataflowContextObject> DataflowEditorContextObject = nullptr;
 
-	UPROPERTY()
-	TObjectPtr<UPreviewMesh> PreviewMesh = nullptr;
-	FMeshDescription PreviewMeshDescription;
-	UE::Geometry::FAxisAlignedBox3d MeshLocalBounds;
-
-	UPROPERTY()
-	TObjectPtr<UToolMeshSelector> MeshSelector = nullptr;
-
-	UPROPERTY()
-	TObjectPtr<UBrushStampIndicator> MirrorBrushIndicator = nullptr;
-
-	UPROPERTY()
-	TObjectPtr<UPreviewMesh> MirrorBrushIndicatorMesh = nullptr;
-
-	UPROPERTY()
-	TObjectPtr<UPreviewMesh> MirrorPlaneMesh = nullptr;
-
-	FSculptBrushStamp MirrorCurrentStamp;
-	int32 MirrorBrushTriangleID = INDEX_NONE;
-
-	// realtime visualization
-	UE_API void OnDynamicMeshComponentChanged(UDynamicMeshComponent* Component, const FMeshVertexChange* Change, bool bRevert);
-	FDelegateHandle OnDynamicMeshComponentChangedHandle;
-
-	UE_API float GetCurrentWeightValueUnderBrush() const;
-	FVector3d CurrentBaryCentricCoords;
-	UE_API int32 GetBrushNearestVertex() const;
-	UE_API int32 GetBrushNearestVertexAccurate() const;
-	UE_API float GetAverageWeightValueFromSelection() const;
-
-	TSet<int32> AccumulatedTriangleROI;
-	bool bUndoUpdatePending = false;
-	TArray<int> NormalsBuffer;
-	UE_API void WaitForPendingUndoRedoUpdate();
-
-	TArray<int> TempROIBuffer;
-	TArray<int> VertexROI;
-	TArray<bool> VisibilityFilterBuffer;
-	TSet<int> VertexSetBuffer;
-	TSet<int> TriangleROI;
-	UE_API void UpdateROI(const FSculptBrushStamp& CurrentStamp, int32 BrushTriangleID);
-
-	UE_API bool UpdateStampPosition(const FRay& WorldRay);
-	UE_API bool ApplyStamp(const FSculptBrushStamp& InStamp);
-
-	UE::Geometry::FDynamicMeshOctree3 Octree;
-
-	UE_API void UpdateBrushType(EDataflowEditorToolEditOperation EditOperation);
-	UE_API bool UpdateBrushPosition(const FRay& WorldRay);
-
-	bool bPendingPickWeight = false;
-	bool bBrushIsHoverMesh = false;
-	bool bBrushMirroringEnabled = false;
-	bool bMirrorPlaneWidgetVisible = true;
-	bool bIsMirrorHiddenBecauseOfBrushStroke = false;
-
-	TArray<int32> ROITriangleBuffer;
-	TArray<double> ROIWeightValueBuffer;
-	UE_API bool SyncMeshWithWeightBuffer(FDynamicMesh3* Mesh);
-	UE_API bool SyncWeightBufferWithMesh(const FDynamicMesh3* Mesh);
-	UE_API void SetVerticesToWeightMap(const TSet<int32>& Vertices, double WeightValue);
-
-	friend class UE::Dataflow::Editor::FVertexAttributePaintToolDetailCustomization;
-	UE_API void BeginChange();
-	UE_API void EndChange();
-	UE_API void CancelChange();
-
-	TArray<FVector3d> TriNormals;
-	TArray<int32> UVSeamEdges;
-	TArray<int32> NormalSeamEdges;
-	UE_API void PrecomputeFilterData();
-
-	FDataflowEditorVertexAttributePaintToolData VertexData;
+	FDataflowEditorVertexAttributePaintToolDataAdapter DataAdapter;
 
 protected:
-	virtual bool ShowWorkPlane() const override { return false; }
-
 	friend class UDataflowEditorVertexAttributePaintToolBuilder;
 
-	bool bAnyChangeMade = false;
-
-	/** update the vertex color of the mesh from the attribute */
-	UE_API void UpdateVertexColorOverlay(const TSet<int>* TrianglesToUpdate = nullptr);
-
 	UDataflowEditorMode* Mode = nullptr;
 
-	static constexpr int32 PaintBrushId = 0;
-	static constexpr int32 SmoothBrushId = 1;
-
-private:
-	// this structure holds information for mirroring 
-	struct FMirrorData
-	{
-	public:
-		// lazily updates the mirror data tables for the current skeleton/mesh/mirror plane
-		void EnsureMirrorDataIsUpdated(const FDynamicMesh3& Mesh, EAxis::Type InMirrorAxis, EDataflowEditorToolMirrorDirection InMirrorDirection, bool bInObjectSpace);
-
-		void FindMirroredIndices(const FDynamicMesh3& Mesh, const TArray<int32>& SelectedVertices, TArray<int32>& OutVerticesToUpdate);
-
-	private:
-		// return true if the point lies on the TARGET side of the mirror plane
-		bool IsPointOnTargetMirrorSide(const FVector& InPoint) const;
-
-		EAxis::Type Axis = EAxis::X;
-		EDataflowEditorToolMirrorDirection Direction = EDataflowEditorToolMirrorDirection::PositiveToNegative;
-		bool bObjectSpace = true;
-		FVector PlaneOffsets = FVector{ 0, 0, 0 };
-		TMap<int32, int32> VertexMap; // <Target, Source>
-	};
-	FMirrorData MirrorData;
 };
 
 
diff --git a/Engine/Plugins/Experimental/MeshModelingToolsetExp/Source/MeshModelingToolsEditorOnlyExp/MeshModelingToolsEditorOnlyExp.Build.cs b/Engine/Plugins/Experimental/MeshModelingToolsetExp/Source/MeshModelingToolsEditorOnlyExp/MeshModelingToolsEditorOnlyExp.Build.cs
--- a/Engine/Plugins/Experimental/MeshModelingToolsetExp/Source/MeshModelingToolsEditorOnlyExp/MeshModelingToolsEditorOnlyExp.Build.cs
+++ b/Engine/Plugins/Experimental/MeshModelingToolsetExp/Source/MeshModelingToolsEditorOnlyExp/MeshModelingToolsEditorOnlyExp.Build.cs
@@ -40,6 +40,7 @@ public class MeshModelingToolsEditorOnlyExp : ModuleRules
 				"ModelingComponentsEditorOnly",
 				"ModelingOperators",
 				"ModelingOperatorsEditorOnly",
+				"MeshModelingToolsEditorOnly",
 				// ... add other public dependencies that you statically link with here ...
 			}
 		);
diff --git a/Engine/Plugins/Experimental/MeshModelingToolsetExp/Source/MeshModelingToolsEditorOnlyExp/Private/MeshVertexAttributePaintBrushOps.cpp b/Engine/Plugins/Experimental/MeshModelingToolsetExp/Source/MeshModelingToolsEditorOnlyExp/Private/MeshVertexAttributePaintBrushOps.cpp
new file mode 100644
--- /dev/null
+++ b/Engine/Plugins/Experimental/MeshModelingToolsetExp/Source/MeshModelingToolsEditorOnlyExp/Private/MeshVertexAttributePaintBrushOps.cpp
@@ -0,0 +1,267 @@
+// Copyright Epic Games, Inc. All Rights Reserved.
+
+#include "MeshVertexAttributePaintBrushOps.h"
+#include "Async/ParallelFor.h"
+#include "DynamicMesh/DynamicMesh3.h"
+#include UE_INLINE_GENERATED_CPP_BY_NAME(MeshVertexAttributePaintBrushOps)
+
+
+bool FMeshVertexAttributePaintToolBrushOpBase::IgnoreZeroMovements() const
+{
+	return true;
+}
+
+//
+// Erase Brush
+//
+
+
+void FMeshVertexAttributePaintToolEraseBrushOp::ApplyStampByVertices(
+	const FDynamicMesh3* Mesh,
+	const FSculptBrushStamp& Stamp,
+	const TArray<int32>& Vertices,
+	TArray<double>& NewAttributesOut)
+{
+	UMeshVertexAttributePaintToolEraseBrushOpProps* Props = GetPropertySetAs<UMeshVertexAttributePaintToolEraseBrushOpProps>();
+	const double EraseAttribute = (double)Props->GetAttribute();
+
+	// TODO: Add something here to get the old value so we can subtract (clamped) the AttributeValue from it.
+	// TODO: Handle the stamp's properties for fall off, etc..
+
+	check(NewAttributesOut.Num() == Vertices.Num());
+
+	for (int32 k = 0; k < Vertices.Num(); ++k)
+	{
+		if (bApplyRadiusLimit)
+		{
+			const FVector3d& StampPos = Stamp.LocalFrame.Origin;
+			const int32 VertIdx = Vertices[k];
+			const FVector3d VertexPos = Mesh->GetVertex(VertIdx);
+			const double DistanceSquared = (VertexPos - StampPos).SquaredLength();
+			if (DistanceSquared >= Stamp.Radius * Stamp.Radius)
+			{
+				continue;
+			}
+		}
+		NewAttributesOut[k] = EraseAttribute;
+	}
+}
+
+
+//
+// Smooth Brush
+//
+
+void FMeshVertexAttributePaintToolSmoothBrushOp::ApplyStampByVertices(
+	const FDynamicMesh3* Mesh,
+	const FSculptBrushStamp& Stamp,
+	const TArray<int32>& Vertices,
+	TArray<double>& VertexWeightValues)
+{
+	ApplyStampByVerticesStatic(Mesh, Stamp, Vertices, VertexWeightValues, bApplyRadiusLimit, Falloff);
+}
+
+
+/*static*/ void FMeshVertexAttributePaintToolSmoothBrushOp::ApplyStampByVerticesStatic(
+	const FDynamicMesh3* Mesh,
+	const FSculptBrushStamp& Stamp,
+	const TArray<int32>& Vertices,
+	TArray<double>& VertexWeightValues,
+	bool bApplyRadiusLimit,
+	TSharedPtr<FMeshSculptFallofFunc> FalloffFunc)
+{
+	auto FalloffFunction = [&Stamp, &FalloffFunc](const FVector& Position)
+		{
+			if (FalloffFunc)
+			{
+				return (float)FalloffFunc->Evaluate(Stamp, Position);
+			}
+			return 1.0f;
+		};
+	FMeshVertexAttributePaintToolSmoothBrushOp::ApplyToVerticesStatic(Mesh, Vertices, VertexWeightValues, Stamp.Power, bApplyRadiusLimit, Stamp.LocalFrame.Origin, Stamp.Radius, FalloffFunction);
+}
+
+/*static*/ void FMeshVertexAttributePaintToolSmoothBrushOp::ApplyToVerticesStatic(
+	const FDynamicMesh3* Mesh,
+	const TArray<int32>& Vertices,
+	TArray<double>& VertexWeightValues,
+	const double SmoothingFactor,
+	bool bApplyRadiusLimit,
+	const FVector Center,
+	const float Radius,
+	FalloffFunction FalloffFunc
+)
+{
+	// Converted from FClothPaintTool_Smooth::SmoothVertices
+	const int32 NumVertices = Vertices.Num();
+
+	TMap<int32, int32> VertexToBufferIndexMap;
+	for (int32 BufferIndex = 0; BufferIndex < NumVertices; ++BufferIndex)
+	{
+		const int32 VertexIndex = Vertices[BufferIndex];
+		VertexToBufferIndexMap.Add(VertexIndex, BufferIndex);
+	}
+
+	// Compute average values of one-rings for all vertices
+	TArray<double> OneRingAverages;
+	OneRingAverages.SetNumUninitialized(NumVertices);
+
+	const TMap<int32, int32>& VertexToBufferIndexMapConst = VertexToBufferIndexMap;
+	const TArray<double>& VertexWeightValuesConst = VertexWeightValues;
+	ParallelFor(NumVertices, 
+		[&Vertices, &Mesh, &VertexToBufferIndexMapConst, &VertexWeightValuesConst, &OneRingAverages]
+		(int32 BufferIndex)
+		{
+			const int32 VertexIndex = Vertices[BufferIndex];
+
+			// average include the vertex we are processing to avoid loosing information
+			double Accumulator = VertexWeightValuesConst[BufferIndex];
+			int32 NumNeighbors = 1;
+
+			for (const int32 NeighborIndex : Mesh->VtxVerticesItr(VertexIndex))
+			{
+				if (VertexToBufferIndexMapConst.Contains(NeighborIndex))
+				{
+					const int32 NeighborBufferIndex = VertexToBufferIndexMapConst[NeighborIndex];
+					Accumulator += VertexWeightValuesConst[NeighborBufferIndex];
+					++NumNeighbors;
+				}
+			}
+
+			OneRingAverages[BufferIndex] = Accumulator / NumNeighbors;
+		}, 
+		EParallelForFlags::None
+	);
+
+	// Blend vertex value with its average one-ring value
+	ParallelFor(NumVertices,
+		[&OneRingAverages, &VertexWeightValues, bApplyRadiusLimit, &Center, &Radius, &SmoothingFactor, &Vertices, &Mesh, &FalloffFunc](int32 BufferIndex)
+		{
+			float FalloffScalar = 1.0;
+			if (bApplyRadiusLimit)
+			{
+				const FVector3d& StampPos = Center;
+				const int32 VertexIndex = Vertices[BufferIndex];
+				const FVector3d VertexPos = Mesh->GetVertex(VertexIndex);
+				const double DistanceSquared = (VertexPos - StampPos).SquaredLength();
+				if (DistanceSquared >= Radius * Radius)
+				{
+					return;
+				}
+				if (FalloffFunc.IsSet())
+				{
+					FalloffScalar = FalloffFunc(VertexPos);
+				}
+			}
+			const float NewValue = FMath::Lerp(VertexWeightValues[BufferIndex], OneRingAverages[BufferIndex], SmoothingFactor * FalloffScalar);
+			VertexWeightValues[BufferIndex] = FMath::Clamp(NewValue, 0.f, 1.f);
+		}, 
+		EParallelForFlags::None
+	);
+}
+
+////////////////////////////////////////////////////////////////////////////////////////////////////////////////////
+
+void FMeshVertexAttributePaintToolPaintBrushOp::ApplyStampByVertices(
+	const FDynamicMesh3* Mesh,
+	const FSculptBrushStamp& Stamp,
+	const TArray<int32>& Vertices,
+	TArray<double>& OutValues)
+{
+	check(OutValues.Num() == Vertices.Num());
+
+	if (UMeshVertexAttributePaintToolPaintBrushOpProps* Props = GetPropertySetAs<UMeshVertexAttributePaintToolPaintBrushOpProps>())
+	{
+		ApplyStampByVerticesStatic(Mesh, Stamp, Vertices, OutValues, Props->EditOperation, Props->GetAttribute(), bApplyRadiusLimit, Falloff);
+	}
+}
+
+/*static*/ void FMeshVertexAttributePaintToolPaintBrushOp::ApplyStampByVerticesStatic(
+	const FDynamicMesh3* Mesh,
+	const FSculptBrushStamp& Stamp,
+	const TArray<int32>& Vertices,
+	TArray<double>& OutValues,
+	EMeshVertexAttributePaintToolEditOperation EditOperation,
+	const double ToolValue,
+	bool bApplyRadiusLimit,
+	TSharedPtr<FMeshSculptFallofFunc> FalloffFunc)
+{
+	auto FalloffFunction = [&Stamp, &FalloffFunc](const FVector& Position)
+		{
+			if (FalloffFunc)
+			{
+				return (float)FalloffFunc->Evaluate(Stamp, Position);
+			}
+			return 1.0f;
+		};
+	FMeshVertexAttributePaintToolPaintBrushOp::ApplyToVerticesStatic(Mesh, Vertices, OutValues, EditOperation, ToolValue, bApplyRadiusLimit, Stamp.LocalFrame.Origin, Stamp.Radius, FalloffFunction);
+}
+
+/*static*/ void FMeshVertexAttributePaintToolPaintBrushOp::ApplyToVerticesStatic(
+	const FDynamicMesh3* Mesh,
+	const TArray<int32>& Vertices,
+	TArray<double>& OutValues,
+	EMeshVertexAttributePaintToolEditOperation EditOperation,
+	const double ToolValue,
+	bool bApplyRadiusLimit,
+	const FVector Center,
+	const float Radius,
+	FalloffFunction FalloffFunc
+)
+{
+	// special case for Relax option 
+	if (EditOperation == EMeshVertexAttributePaintToolEditOperation::Relax)
+	{
+		FMeshVertexAttributePaintToolSmoothBrushOp::ApplyToVerticesStatic(Mesh, Vertices, OutValues, ToolValue, bApplyRadiusLimit, Center, Radius);
+		return;
+	}
+
+	const int32 NumVertices = Vertices.Num();
+	ParallelFor(NumVertices,
+		[bApplyRadiusLimit, ToolValue, EditOperation, &Center, &Radius, &Mesh, &Vertices, &OutValues, &FalloffFunc]
+		(int32 BufferIndex)
+		{
+			float FalloffScalar = 1.0;
+			if (bApplyRadiusLimit)
+			{
+				const FVector3d& StampPos = Center;
+				const int32 MeshVertIdx = Vertices[BufferIndex];
+				const FVector3d VertexPos = Mesh->GetVertex(MeshVertIdx);
+				const double DistanceSquared = (VertexPos - StampPos).SquaredLength();
+				if (DistanceSquared >= Radius * Radius)
+				{
+					return;
+				}
+				if (FalloffFunc.IsSet())
+				{
+					FalloffScalar = FalloffFunc(VertexPos);
+				}
+			}
+
+			const double CurrentValue = OutValues[BufferIndex];
+			double NewValue = CurrentValue;
+
+			switch (EditOperation)
+			{
+			case EMeshVertexAttributePaintToolEditOperation::Add:
+				NewValue = CurrentValue + ToolValue;
+				break;
+			case EMeshVertexAttributePaintToolEditOperation::Replace:
+				NewValue = ToolValue;
+				break;
+			case EMeshVertexAttributePaintToolEditOperation::Multiply:
+				NewValue = CurrentValue * ToolValue;
+				break;
+			case EMeshVertexAttributePaintToolEditOperation::Invert:
+				NewValue = 1.0 - CurrentValue; // no falloff
+				break;
+			case EMeshVertexAttributePaintToolEditOperation::Relax:
+				ensure(false); // already handled by SmoothBrush.ApplyStampByVertices
+			}
+
+			NewValue = FMath::Lerp(OutValues[BufferIndex], NewValue, FalloffScalar);
+			OutValues[BufferIndex] = FMath::Clamp(NewValue, 0.0, 1.0);
+		},
+		EParallelForFlags::None
+	);
+}
\ No newline at end of file
diff --git a/Engine/Plugins/Runtime/MeshModelingToolset/Source/MeshModelingTools/Public/MeshAttributePaintTool.h b/Engine/Plugins/Runtime/MeshModelingToolset/Source/MeshModelingTools/Public/MeshAttributePaintTool.h
--- a/Engine/Plugins/Runtime/MeshModelingToolset/Source/MeshModelingTools/Public/MeshAttributePaintTool.h
+++ b/Engine/Plugins/Runtime/MeshModelingToolset/Source/MeshModelingTools/Public/MeshAttributePaintTool.h
@@ -84,6 +84,7 @@ public:
 
 /**
  * Tool Builder for Attribute Paint Tool
+ * Note: a new experimental version of this tool is UMeshAttributePaintToolV2
  */
 UCLASS(MinimalAPI)
 class UMeshAttributePaintToolBuilder : public UMeshSurfacePointMeshEditingToolBuilder
@@ -244,7 +245,7 @@ public:
 
 /**
  * UMeshAttributePaintTool paints single-channel float attributes on a MeshDescription or DynamicMesh
- * 
+ * Note: a new experimental version of this tool is UMeshAttributePaintToolV2
  */
 UCLASS(MinimalAPI)
 class UMeshAttributePaintTool : public UDynamicMeshBrushTool, public IInteractiveToolManageGeometrySelectionAPI
