Dataflow : refactor the way we track errors, warning  and info in the editor and at the context level
- log message are now track using timestamp so that we re-evaluating a single node remove the corresponding messages only
- Context do not track individual nodes anymore , it only broadcast based on registered listeners
- Display node messages in their tooltip

#Synced-CL 51526819

diff --git a/Engine/Plugins/Experimental/Dataflow/Source/DataflowEditor/Private/Dataflow/DataflowEditorToolkit.cpp b/Engine/Plugins/Experimental/Dataflow/Source/DataflowEditor/Private/Dataflow/DataflowEditorToolkit.cpp
--- a/Engine/Plugins/Experimental/Dataflow/Source/DataflowEditor/Private/Dataflow/DataflowEditorToolkit.cpp
+++ b/Engine/Plugins/Experimental/Dataflow/Source/DataflowEditor/Private/Dataflow/DataflowEditorToolkit.cpp
@@ -242,12 +242,106 @@ namespace UE::Dataflow::Private
 		}
 		return nullptr;
 	}
+
+	//////////////////////////////////////////////////////////////
+
+	void FLogMessageStorage::GetMessage(int32 Index, UE::Dataflow::FContext::FLogMessage& OutMessage) const
+	{
+		if (Entries.IsValidIndex(Index))
+		{
+			OutMessage = Entries[Index].Message;
+		}
+		else
+		{
+			OutMessage = UE::Dataflow::FContext::FLogMessage();
+		}
+
+	}
+
+	void FLogMessageStorage::Add(const UE::Dataflow::FContext::FLogMessage& Message)
+	{
+		FTimestamp Timestamp = FTimestamp::Invalid;
+		if (TSharedPtr<const FDataflowNode> Node = Message.Node.Pin())
+		{
+			Timestamp = Node->GetTimestamp();
+		}
+		Entries.Emplace(FEntry
+			{
+				.Timestamp = Timestamp,
+				.Message = Message,
+			});
+	}
+
+	void FLogMessageStorage::GetMessageForNode(const FDataflowNode& Node, TArray<UE::Dataflow::FContext::FLogMessage>& OutMessages) const
+	{
+		for (const FEntry& Entry : Entries)
+		{
+			if (Entry.Message.Node == Node.AsWeak())
+			{
+				OutMessages.Add(Entry.Message);
+			}
+		}
+	}
+
+	EMessageSeverity::Type FLogMessageStorage::GetMaxSeverityForNode(const FDataflowNode& Node) const
+	{
+		EMessageSeverity::Type OutSeverity = EMessageSeverity::Info;
+		for (const FEntry& Entry : Entries)
+		{
+			if (Entry.Message.Node == Node.AsWeak())
+			{
+				if (Entry.Message.Severity < OutSeverity)
+				{
+					OutSeverity = Entry.Message.Severity;
+				}
+			}
+		}
+		return OutSeverity;
+	}
+
+	void FLogMessageStorage::RemoveForNode(const FDataflowNode& Node, const FTimestamp& Timestamp)
+	{
+		TArray<FEntry> NewEntries;
+		NewEntries.Reserve(Entries.Num());
+
+		for (const FEntry& Entry : Entries)
+		{
+			bool bKeep = true;
+			if (Entry.Message.Node == Node.AsWeak())
+			{
+				if (Timestamp == FTimestamp::Invalid || Entry.Timestamp < Timestamp)
+				{
+					bKeep = false;
+				}
+			}
+			if (bKeep)
+			{
+				NewEntries.Emplace(Entry);
+			}
+		}
+		Entries = MoveTemp(NewEntries);
+	}
+
+	EMessageSeverity::Type FLogMessageStorage::GetMaxSeverity()
+	{
+		EMessageSeverity::Type OutSeverity = EMessageSeverity::Info;
+		for (const FEntry& Entry : Entries)
+		{
+			if (Entry.Message.Severity < OutSeverity)
+			{
+				OutSeverity = Entry.Message.Severity;
+			}
+		}
+		return OutSeverity;
+	}
+
 }
 
 FDataflowEditorToolkit::FDataflowEditorToolkit(UAssetEditor* InOwningAssetEditor)
 	: FBaseCharacterFXEditorToolkit(InOwningAssetEditor, UE::Dataflow::Private::DataflowEditorToolkitName)
 	, DataflowEditor(Cast< UDataflowEditor>(InOwningAssetEditor))
 	, bDataflowEnableGraphEval(true)
+	, LogMessages(MakeShared<UE::Dataflow::Private::FLogMessageStorage>())
 {
 	EvaluationMode = EDataflowEditorEvaluationMode::Automatic;
 
@@ -885,17 +979,6 @@ FSlateIcon FDataflowEditorToolkit::GetEvaluationStatusImage() const
 	static const FName CompileStatusWarning("Blueprint.CompileStatus.Overlay.Warning");
 	static const FName CompileStatusError("Blueprint.CompileStatus.Overlay.Error");
 
-	bool bHasWarning = false;
-	bool bHasError = false;
-	if (const TObjectPtr<UDataflowBaseContent>& EditorContent = GetEditorContent())
-	{
-		if (TSharedPtr<UE::Dataflow::FEngineContext> DataflowContext = EditorContent->GetDataflowContext())
-		{
-			bHasWarning = (DataflowContext->GetNumWarnings() > 0);
-			bHasError = (DataflowContext->GetNumErrors() > 0);
-		}
-	}
-
 	FName OverlayIcon;
 	if (IsGraphDirty())
 	{
@@ -903,17 +986,20 @@ FSlateIcon FDataflowEditorToolkit::GetEvaluationStatusImage() const
 	}
 	else
 	{
-		if (bHasError)
+		const EMessageSeverity::Type Severity = LogMessages->GetMaxSeverity();
+		switch (Severity)
 		{
+		case EMessageSeverity::Error:
 			OverlayIcon = CompileStatusError;
-		}
-		else if (bHasWarning)
-		{
+			break;
+		case EMessageSeverity::PerformanceWarning:
+		case EMessageSeverity::Warning:
 			OverlayIcon = CompileStatusWarning;
-		}
-		else
-		{
+			break;
+		case EMessageSeverity::Info:
+		default:
 			OverlayIcon = CompileStatusGood;
+			break;
 		}
 	}
 	return FSlateIcon(FAppStyle::GetAppStyleSetName(), CompileStatusBackground, NAME_None, OverlayIcon);
@@ -1990,6 +2076,8 @@ void FDataflowEditorToolkit::OnNodeInvalidated(UDataflow& DataflowAsset, FDatafl
 			}
 		}
 	}
+
+	LogMessages->RemoveForNode(Node, Node.GetTimestamp());
 }
 
 void FDataflowEditorToolkit::OnNodeDeleted(const TSet<UObject*>& NewSelection)
@@ -2006,6 +2094,15 @@ void FDataflowEditorToolkit::OnNodeDeleted(const TSet<UObject*>& NewSelection)
 		DeletedNodesSet.Add(Node);
 	}
 
+	// remove from the logs
+	for (UDataflowEdNode* Node : DeletedNodes)
+	{
+		if (TSharedPtr<const FDataflowNode> DataflowNode = Node->GetDataflowNode())
+		{
+			LogMessages->RemoveForNode(*DataflowNode);
+		}
+	}
+
 	if (UDataflowEditorMode* DataflowEdMode = Cast<UDataflowEditorMode>(EditorModeManager->GetActiveScriptableMode(UDataflowEditorMode::EM_DataflowEditorModeId)))
 	{
 		DataflowEdMode->OnDataflowNodeDeleted(DeletedNodesSet);
@@ -2098,36 +2195,27 @@ void FDataflowEditorToolkit::LogMessage(const EMessageSeverity::Type Severity, c
 	}
 }
 
-void FDataflowEditorToolkit::OnContextLog(const UE::Dataflow::FContext::FLogMessage& InLogMessage)
+void FDataflowEditorToolkit::RefreshLogMessageView()
 {
-	if (DataflowOutputLog)
-	{
-		TSharedPtr<const FDataflowNode> Node = InLogMessage.Node.Pin();
-		const FDataflowOutput* Output = Node ? Node->FindOutput(InLogMessage.OutputGuid) : nullptr;
-		LogMessage(InLogMessage.Severity, Node.Get(), Output, InLogMessage.Text);
-	}
-}
-
-void FDataflowEditorToolkit::OnBeginEvaluate()
-{
-	GraphEvaluationBegin = FDateTime::Now();
-
 	if (DataflowOutputLog)
 	{
 		DataflowOutputLog->ClearMessageLog();
-	}
 
-	if (const TObjectPtr<UDataflowBaseContent>& EditorContent = GetEditorContent())
-	{
-		ensure(EditorContent);
-		if (TSharedPtr<UE::Dataflow::FEngineContext> DataflowContext = EditorContent->GetDataflowContext())
+		UE::Dataflow::FContext::FLogMessage Message;
+		const int32 NumMessages = LogMessages->Num();
+		for (int32 Index = 0; Index < NumMessages; ++Index)
 		{
-			DataflowContext->ClearNodesData();
+			LogMessages->GetMessage(Index, Message);
+			TSharedPtr<const FDataflowNode> Node = Message.Node.Pin();
+			const FDataflowOutput* Output = Node ? Node->FindOutput(Message.OutputGuid) : nullptr;
+			LogMessage(Message.Severity, Node.Get(), Output, Message.Text);
 		}
 	}
+}
 
-	// Refresh graph display to display nodes with warning/error
-	GraphEditor->NotifyGraphChanged();
+void FDataflowEditorToolkit::OnContextLog(const UE::Dataflow::FContext::FLogMessage& InLogMessage)
+{
+	LogMessages->Add(InLogMessage);
 }
 
 namespace UE::Dataflow::Private
@@ -2141,39 +2229,12 @@ namespace UE::Dataflow::Private
 	static FString GetElapsedTimeString(const FDateTime& InGraphEvaluationFinished, const FDateTime& InGraphEvaluationBegin)
 	{
 		const FTimespan ElapsedTimespan = InGraphEvaluationFinished - InGraphEvaluationBegin;
-		const FString ElapsedTimeString = FString::Printf(TEXT("%02dm%02ds%02dms"), ElapsedTimespan.GetMinutes(), ElapsedTimespan.GetSeconds(), GetMilliseconds(ElapsedTimespan));
+		const FString ElapsedTimeString = FString::Printf(TEXT("%02dm %02ds %02dms"), ElapsedTimespan.GetMinutes(), ElapsedTimespan.GetSeconds(), GetMilliseconds(ElapsedTimespan));
 
 		return ElapsedTimeString;
 	}
 }
 
-void FDataflowEditorToolkit::OnFinishEvaluate()
-{
-	// Display message stating that evaluation finished
-	const TObjectPtr<UDataflowBaseContent>& EditorContent = GetEditorContent();
-	ensure(EditorContent);
-
-	if (TSharedPtr<UE::Dataflow::FEngineContext> DataflowContext = EditorContent->GetDataflowContext())
-	{
-		const int32 NumWarnings = DataflowContext->GetNumWarnings();
-		const int32 NumErrors = DataflowContext->GetNumErrors();
-
-		if (NumWarnings > 0 || NumErrors > 0)
-		{
-			const FText MessageFormat = LOCTEXT("OnFinishEvaluate", "Finished graph evaluation\nWarning(s): {0} Error(s): {1}");
-			UE::Dataflow::Private::ShowNotificationMessage(FText::Format(MessageFormat, NumWarnings, NumErrors), SNotificationItem::CS_Fail);
-		}
-
-		GraphEvaluationFinished = FDateTime::Now();
-
-		const FString ElapsedTimeString = UE::Dataflow::Private::GetElapsedTimeString(GraphEvaluationFinished, GraphEvaluationBegin);
-		DataflowContext->Info(TEXT("Evaluation time: ") + ElapsedTimeString);
-	}
-
-	// Refresh graph display to update node output pin display (invalid or valid)
-	GraphEditor->NotifyGraphChanged();
-}
-
 void FDataflowEditorToolkit::OnOutputLogMessageTokenClicked(const FString TokenString)
 {
 	FDataflowPath Path;
@@ -2295,9 +2356,7 @@ void FDataflowEditorToolkit::EvaluateTerminalNode(const FDataflowTerminalNode& T
 	{
 		FTimestamp TerminalNodeTimeStamp = FTimestamp::Invalid;
 
-		OnBeginEvaluate();
 		EvaluateNode(&TerminalNode, nullptr, TerminalNodeTimeStamp);
-		OnFinishEvaluate();
 
 		EditorContent->SetLastModifiedTimestamp(TerminalNodeTimeStamp, /*bMakeDirty*/true);
 	}
@@ -2315,16 +2374,28 @@ void FDataflowEditorToolkit::EvaluateNode(const FDataflowNode* Node, const FData
 		{
 			if (const TObjectPtr<UDataflowBaseContent>& EditorContent = GetEditorContent())
 			{
-				if (ensure(Node))
+				if (Node)
 				{
+					const FDateTime StartEvaluationTime = FDateTime::Now();
+
 					// when evaluation completes, refresh the views
 					TWeakPtr<FDataflowEditorToolkit> WeakToolkit = SharedThis(this);
-					auto OnEvaluationCompleted = [WeakToolkit](UE::Dataflow::FContext& Context)
+					TWeakPtr<const FDataflowNode> WeakNode = Node->AsWeak();
+					auto OnEvaluationCompleted = [WeakToolkit, WeakNode, StartEvaluationTime](UE::Dataflow::FContext& Context)
 						{
+							const FDateTime EndEvaluationTime = FDateTime::Now();
 							// refresh views 
 							if (TSharedPtr<FDataflowEditorToolkit> ToolKit = WeakToolkit.Pin())
 							{
 								ToolKit->bViewsNeedRefresh = true;
+								ToolKit->GraphEditor->NotifyGraphChanged();
+							}
+							if (TSharedPtr<const FDataflowNode> Node = WeakNode.Pin())
+							{
+								const FText ElapsedTimeText = FText::FromString(UE::Dataflow::Private::GetElapsedTimeString(EndEvaluationTime, StartEvaluationTime));
+								const FText NodeNameText = FText::FromName(Node->GetName());
+								const FText MessageText = FText::Format(LOCTEXT("EvalElapsedLog", "Graph evaluation from node [{0}] time : {1}"), NodeNameText, ElapsedTimeText);
+								Context.Info(MessageText);
 							}
 						};
 					bool bUseAsyncEvaluation = false;
@@ -2363,6 +2434,7 @@ void FDataflowEditorToolkit::RefreshViewsIfNeeded(bool bForce)
 		}
 
 		UpdateDebugDraw();
+		RefreshLogMessageView();
 
 		bViewsNeedRefresh = false;
 	}
@@ -2440,6 +2512,8 @@ TSharedRef<SDataflowGraphEditor> FDataflowEditorToolkit::CreateGraphEditorWidget
 	OnSelectionChangedMulticastDelegateHandle = NewGraphEditor->OnSelectionChangedMulticast.AddSP(this, &FDataflowEditorToolkit::OnNodeSelectionChanged);
 	OnNodeDeletedMulticastDelegateHandle = NewGraphEditor->OnNodeDeletedMulticast.AddSP(this, &FDataflowEditorToolkit::OnNodeDeleted);
 	
+	NewGraphEditor->SetLogMessageStorage(LogMessages);
+
 	UnregisterContextHandlers();
 	RegisterContextHandlers();
 
diff --git a/Engine/Plugins/Experimental/Dataflow/Source/DataflowEditor/Private/Dataflow/DataflowGraphEditor.cpp b/Engine/Plugins/Experimental/Dataflow/Source/DataflowEditor/Private/Dataflow/DataflowGraphEditor.cpp
--- a/Engine/Plugins/Experimental/Dataflow/Source/DataflowEditor/Private/Dataflow/DataflowGraphEditor.cpp
+++ b/Engine/Plugins/Experimental/Dataflow/Source/DataflowEditor/Private/Dataflow/DataflowGraphEditor.cpp
@@ -360,6 +360,11 @@ SDataflowGraphEditor::~SDataflowGraphEditor()
 	}
 }
 
+void SDataflowGraphEditor::SetLogMessageStorage(TSharedPtr<UE::Dataflow::Private::FLogMessageStorage> InLogMessages)
+{
+	LogMessages = InLogMessages;
+}
+
 bool SDataflowGraphEditor::IsEditable() const
 {
 	return IsEditableCallback ? IsEditableCallback() : false;
@@ -478,6 +483,39 @@ void SDataflowGraphEditor::OnRenderToggleChanged(UDataflowEdNode* UpdatedEdNode)
 	GetGraphPanel()->SelectionManager.SetSelectionSet(SelectionSet);
 }
 
+FText SDataflowGraphEditor::GetMessageText(const UDataflowEdNode& Node) const
+{
+	FText MessageText;
+	if (LogMessages)
+	{
+		if (TSharedPtr<const FDataflowNode> DataflowNode = Node.GetDataflowNode())
+		{
+			using namespace UE::Dataflow;
+			TArray<FContext::FLogMessage> Messages;
+			LogMessages->GetMessageForNode(*DataflowNode, Messages);
+
+			for (const FContext::FLogMessage& Message : Messages)
+			{
+				MessageText = FText::Format(LOCTEXT("NodeLogMessageFormat", "{0}\n[{1}] - {2}"), MessageText, FTokenizedMessage::GetSeverityText(Message.Severity), Message.Text);
+			}
+		}
+	}
+	return MessageText;
+}
+
+EMessageSeverity::Type SDataflowGraphEditor::GetMessageMaxSeverity(const UDataflowEdNode& Node) const
+{
+	EMessageSeverity::Type OutSeverity = EMessageSeverity::Info;
+	if (LogMessages)
+	{
+		if (TSharedPtr<const FDataflowNode> DataflowNode = Node.GetDataflowNode())
+		{
+			OutSeverity = LogMessages->GetMaxSeverityForNode(*DataflowNode);
+		}
+	}
+	return OutSeverity;
+}
+
 void SDataflowGraphEditor::EvaluateNode()
 {
 	UE_LOG(LogChaosDataflow, VeryVerbose, TEXT("SDataflowGraphEditor::EvaluateNode(): Nodes [%s]"),
diff --git a/Engine/Plugins/Experimental/Dataflow/Source/DataflowEditor/Private/Dataflow/DataflowSNode.cpp b/Engine/Plugins/Experimental/Dataflow/Source/DataflowEditor/Private/Dataflow/DataflowSNode.cpp
--- a/Engine/Plugins/Experimental/Dataflow/Source/DataflowEditor/Private/Dataflow/DataflowSNode.cpp
+++ b/Engine/Plugins/Experimental/Dataflow/Source/DataflowEditor/Private/Dataflow/DataflowSNode.cpp
@@ -379,21 +379,21 @@ void SDataflowEdNode::UpdateErrorInfo()
 		{
 			if (DataflowInterface)
 			{
-				if (const TSharedPtr<UE::Dataflow::FContext> Context = DataflowInterface->GetDataflowContext())
+				const EMessageSeverity::Type Severity = DataflowInterface->GetMessageMaxSeverity(*DataflowGraphNode);
+				switch (Severity)
 				{
-					if (Context->NodeHasError(DataflowNode.Get()))
-					{
-						ErrorMsg = FString(TEXT("ERROR"));
-						const FLinearColor ErrorBackgroundColor = FLinearColor(1.f, 0.02f, 0.003f);
-						ErrorColor = ErrorBackgroundColor;
-						return;
-					}
-					else if (Context->NodeHasWarning(DataflowNode.Get()))
-					{
-						ErrorMsg = FString(TEXT("WARNING"));
-						ErrorColor = FAppStyle::GetColor("ErrorReporting.WarningBackgroundColor");
-						return;
-					}
+				case EMessageSeverity::Error:
+					ErrorMsg = FString(TEXT("ERROR"));
+					ErrorColor = FLinearColor(1.f, 0.02f, 0.003f);
+					return;
+				case EMessageSeverity::PerformanceWarning:
+				case EMessageSeverity::Warning:
+					ErrorMsg = FString(TEXT("WARNING"));
+					ErrorColor = FAppStyle::GetColor("ErrorReporting.WarningBackgroundColor");
+					return;
+				case EMessageSeverity::Info:
+				default:
+					break; // continue to other cases 
 				}
 			}
 
@@ -413,6 +413,23 @@ void SDataflowEdNode::UpdateErrorInfo()
 	}
 }
 
+FText SDataflowEdNode::GetNodeTooltip() const
+{
+	FText TooltipText = SGraphNode::GetNodeTooltip();
+	if (DataflowInterface)
+	{
+		if (UDataflowEdNode* DataflowNode = Cast<UDataflowEdNode>(GraphNode))
+		{
+			const FText MessageText = DataflowInterface->GetMessageText(*DataflowNode);
+			if (!MessageText.IsEmpty())
+			{
+				TooltipText = FText::Format(LOCTEXT("SDataflowNode_MergeText", "{0}\n{1}"), TooltipText, MessageText);
+			}
+		}
+	}
+	return TooltipText;
+}
+
 FReply SDataflowEdNode::OnMouseButtonDoubleClick(const FGeometry& InMyGeometry, const FPointerEvent& InMouseEvent)
 {
 	TSharedPtr<FDataflowNode> Node;
diff --git a/Engine/Plugins/Experimental/Dataflow/Source/DataflowEditor/Public/Dataflow/DataflowEditorToolkit.h b/Engine/Plugins/Experimental/Dataflow/Source/DataflowEditor/Public/Dataflow/DataflowEditorToolkit.h
--- a/Engine/Plugins/Experimental/Dataflow/Source/DataflowEditor/Public/Dataflow/DataflowEditorToolkit.h
+++ b/Engine/Plugins/Experimental/Dataflow/Source/DataflowEditor/Public/Dataflow/DataflowEditorToolkit.h
@@ -45,6 +45,30 @@ enum class EDataflowEditorEvaluationMode : uint8;
 namespace UE::Dataflow
 {
 	class FDataflowNodeDetailExtensionHandler;
+	namespace Private
+	{
+		// stores the log message per node and 
+		struct FLogMessageStorage
+		{
+		public:
+			void Add(const UE::Dataflow::FContext::FLogMessage& Message);
+			void RemoveForNode(const FDataflowNode& Node, const FTimestamp& Timestamp = FTimestamp::Invalid);
+			void GetMessageForNode(const FDataflowNode& Node, TArray<UE::Dataflow::FContext::FLogMessage>& OutMessages) const;
+			EMessageSeverity::Type GetMaxSeverityForNode(const FDataflowNode& Node) const;
+			int32 Num() const { return Entries.Num(); }
+			void GetMessage(int32 Index, UE::Dataflow::FContext::FLogMessage& OutMessage) const;
+			EMessageSeverity::Type GetMaxSeverity();
+
+		private:
+			struct FEntry
+			{
+				UE::Dataflow::FTimestamp Timestamp;
+				UE::Dataflow::FContext::FLogMessage Message;
+			};
+
+			TArray<FEntry> Entries;
+		};
+	}
 }
 namespace EMessageSeverity { enum Type : int; }
 
@@ -181,10 +205,9 @@ protected:
 	UE_API void RefreshViewsIfNeeded(bool bForce = false);
 	UE_API void OnNodeBeginEvaluate(const FDataflowNode* Node, const FDataflowOutput* Output);
 	UE_API void OnNodeFinishEvaluate(const FDataflowNode* Node, const FDataflowOutput* Output);
-	UE_API void OnBeginEvaluate();
-	UE_API void OnFinishEvaluate();
 	UE_API void OnOutputLogMessageTokenClicked(const FString TokenString);
 	UE_API void OnContextLog(const UE::Dataflow::FContext::FLogMessage& InLogMessage);
+	UE_API void RefreshLogMessageView();
 
 private:
 	
@@ -422,8 +445,9 @@ private:
 
 	TSet<FGuid> NodesToEvaluateOnTick;
 
-	FDateTime GraphEvaluationBegin;
-	FDateTime GraphEvaluationFinished;
+	/** relevant message to display */
+	TSharedRef<UE::Dataflow::Private::FLogMessageStorage> LogMessages;
+
 	bool bViewsNeedRefresh = false;
 
 	FString DebugDrawOverlayString;
diff --git a/Engine/Plugins/Experimental/Dataflow/Source/DataflowEditor/Public/Dataflow/DataflowGraphEditor.h b/Engine/Plugins/Experimental/Dataflow/Source/DataflowEditor/Public/Dataflow/DataflowGraphEditor.h
--- a/Engine/Plugins/Experimental/Dataflow/Source/DataflowEditor/Public/Dataflow/DataflowGraphEditor.h
+++ b/Engine/Plugins/Experimental/Dataflow/Source/DataflowEditor/Public/Dataflow/DataflowGraphEditor.h
@@ -31,6 +31,10 @@ class UDataflowEditor;
 struct FDataflowConnection;
 namespace UE::Dataflow {
 	class FContext;
+	namespace Private
+	{
+		struct FLogMessageStorage;
+	}
 }
 /**
  * The SDataflowGraphEditpr class is a specialization of SGraphEditor
@@ -168,9 +172,12 @@ public:
 	UE_API virtual void AddReferencedObjects(FReferenceCollector& Collector) override;
 	virtual FString GetReferencerName() const override { return TEXT("SDataflowGraphEditor"); }
 
-	/** FDataflowSNodeInterface */
+	//~ Begin FDataflowSEditorInterface 
 	UE_API virtual TSharedPtr<UE::Dataflow::FContext> GetDataflowContext() const override;
 	UE_API virtual void OnRenderToggleChanged(UDataflowEdNode* UpdatedEdNode) const override;
+	UE_API virtual FText GetMessageText(const UDataflowEdNode& Node) const override;
+	UE_API virtual EMessageSeverity::Type GetMessageMaxSeverity(const UDataflowEdNode& Node) const override;
+	//~ End FDataflowSEditorInterface 
 
 	const TSharedPtr<FUICommandList> GetCommands() const { return GraphEditorCommands; }
 
@@ -181,6 +188,8 @@ public:
 
 	bool GetFilterActionMenyByAssetType() const { return bFilterActionMenyByAssetType; }
 
+	void SetLogMessageStorage(TSharedPtr<UE::Dataflow::Private::FLogMessageStorage> InLogMessages);
+
 private:
 	/** Add an additional option pin to all selected Dataflow nodes for those that overrides the AddPin function. */
 	UE_API void OnAddOptionPin();
@@ -257,6 +266,8 @@ private:
 	TSharedPtr<SWidget> MessageBar;
 	TSharedPtr<SWidget> EvaluationProgressBar;
 	FText MessageBarText;
+
+	TSharedPtr<UE::Dataflow::Private::FLogMessageStorage> LogMessages;
 };
 
 #undef UE_API
diff --git a/Engine/Plugins/Experimental/Dataflow/Source/DataflowEditor/Public/Dataflow/DataflowSEditorInterface.h b/Engine/Plugins/Experimental/Dataflow/Source/DataflowEditor/Public/Dataflow/DataflowSEditorInterface.h
--- a/Engine/Plugins/Experimental/Dataflow/Source/DataflowEditor/Public/Dataflow/DataflowSEditorInterface.h
+++ b/Engine/Plugins/Experimental/Dataflow/Source/DataflowEditor/Public/Dataflow/DataflowSEditorInterface.h
@@ -3,6 +3,7 @@
 #pragma once
 
 #include "Templates/SharedPointerFwd.h"
+#include "Logging/TokenizedMessage.h"
 
 namespace UE::Dataflow
 {
@@ -35,6 +36,16 @@ public:
 	{
 	}
 
+	virtual FText GetMessageText(const UDataflowEdNode& Node) const
+	{
+		return FText();
+	}
+
+	virtual EMessageSeverity::Type GetMessageMaxSeverity(const UDataflowEdNode& Node) const
+	{
+		return EMessageSeverity::Info;
+	}
+
 protected:
 	FDataflowSEditorInterface() = default;
 };
diff --git a/Engine/Plugins/Experimental/Dataflow/Source/DataflowEditor/Public/Dataflow/DataflowSNode.h b/Engine/Plugins/Experimental/Dataflow/Source/DataflowEditor/Public/Dataflow/DataflowSNode.h
--- a/Engine/Plugins/Experimental/Dataflow/Source/DataflowEditor/Public/Dataflow/DataflowSNode.h
+++ b/Engine/Plugins/Experimental/Dataflow/Source/DataflowEditor/Public/Dataflow/DataflowSNode.h
@@ -38,6 +38,7 @@ public:
 	// SGraphNode interface
 	UE_API virtual FReply OnMouseButtonDoubleClick(const FGeometry& InMyGeometry, const FPointerEvent& InMouseEvent) override;
 	UE_API virtual TArray<FOverlayWidgetInfo> GetOverlayWidgets(bool bSelected, const FVector2f& WidgetSize) const override;
+	UE_API virtual FText GetNodeTooltip() const override;
 	UE_API virtual void UpdateErrorInfo() override;
 	virtual void RequestRenameOnSpawn() override { /* No auto rename on spawn, because it can interfere with Copy/Paste and cause a crash */ }
 
diff --git a/Engine/Source/Runtime/Experimental/Dataflow/Core/Private/Dataflow/DataflowInputOutput.cpp b/Engine/Source/Runtime/Experimental/Dataflow/Core/Private/Dataflow/DataflowInputOutput.cpp
--- a/Engine/Source/Runtime/Experimental/Dataflow/Core/Private/Dataflow/DataflowInputOutput.cpp
+++ b/Engine/Source/Runtime/Experimental/Dataflow/Core/Private/Dataflow/DataflowInputOutput.cpp
@@ -9,6 +9,8 @@
 
 #include UE_INLINE_GENERATED_CPP_BY_NAME(DataflowInputOutput)
 
+#define LOCTEXT_NAMESPACE "DataflowINputOutput"
+
 namespace UE::Dataflow::Private
 {
 	static bool bDisableFreezingState = true;
@@ -29,6 +31,15 @@ bool FDataflowInput::IsConnected() const
 	return (Connection != nullptr);
 }
 
+void FDataflowInput::CheckRequired(UE::Dataflow::FContext& Context) const
+{
+	if (IsRequired() && !IsConnected())
+	{
+		const FText WarningMessage = FText::Format(LOCTEXT("InputRequiredMessage", "Input {0} is required but not connected"), FText::FromName(GetName()));
+		Context.Warning(WarningMessage, this->OwningNode);
+	}
+}
+
 bool FDataflowInput::AddConnection(FDataflowConnection* InOutput)
 {
 	if (GetType() == InOutput->GetType())
@@ -297,18 +308,11 @@ bool FDataflowOutput::EvaluateImpl(UE::Dataflow::FContext& Context) const
 		UE::Dataflow::FContextScopedCallstack Callstack(Context, this);
 		OwningNode->Evaluate(Context, this);
 		
-		// Validation
-		if (!Context.NodeHasError(OwningNode))
-		{
-			if (!Context.HasData(CacheKey()))
-			{
-				ensureMsgf(false, TEXT("Failed to evaluate output (%s:%s)"), *OwningNode->GetName().ToString(), *GetName().ToString());
-				return false;
-			}
-		}
-		else
+		// if no output was set, place a null entry 
+		if (!Context.HasData(CacheKey(), OwningNode->GetTimestamp()))
 		{
 			this->SetNullValue(Context);
+			return false;
 		}
 	}
 	return true;
@@ -466,5 +470,7 @@ void* FDataflowArrayOutput::RealAddress() const
 
 bool FDataflowOutput::HasNodeFailedOrErrored(UE::Dataflow::FContext& Context) const
 {
-	return Context.NodeHasError(GetOwningNode());
+	return false; // deprecated 5.8
 }
+
+#undef LOCTEXT_NAMESPACE 
\ No newline at end of file
diff --git a/Engine/Source/Runtime/Experimental/Dataflow/Core/Private/Dataflow/DataflowNodeParameters.cpp b/Engine/Source/Runtime/Experimental/Dataflow/Core/Private/Dataflow/DataflowNodeParameters.cpp
--- a/Engine/Source/Runtime/Experimental/Dataflow/Core/Private/Dataflow/DataflowNodeParameters.cpp
+++ b/Engine/Source/Runtime/Experimental/Dataflow/Core/Private/Dataflow/DataflowNodeParameters.cpp
@@ -385,11 +385,6 @@ namespace UE::Dataflow
 
 	void FContext::Info(const FText& InInfo, const FDataflowNode* InNode, const FDataflowOutput* InOutput)
 	{
-		if (InNode)
-		{
-			FScopeLock Lock(&NodeTrackingLock);
-			NodesWithInfo.Add(InNode->GetGuid());
-		}
 		BroadCastLog(EMessageSeverity::Info, InInfo, InNode, InOutput);
 	}
 
@@ -399,19 +394,8 @@ namespace UE::Dataflow
 		Info(Message, InNode, InOutput);
 	}
 
-	int32 FContext::GetNumInfo() const
-	{
-		FScopeLock Lock(&NodeTrackingLock);
-		return NodesWithInfo.Num();
-	}
-
 	void FContext::Warning(const FText& InWarning, const FDataflowNode* InNode, const FDataflowOutput* InOutput)
 	{
-		if (InNode)
-		{
-			FScopeLock Lock(&NodeTrackingLock);
-			NodesWithWarning.Add(InNode->GetGuid());
-		}
 		BroadCastLog(EMessageSeverity::Warning, InWarning, InNode, InOutput);
 	}
 
@@ -421,29 +405,8 @@ namespace UE::Dataflow
 		Warning(Message, InNode, InOutput);
 	}
 
-	bool FContext::NodeHasWarning(const FDataflowNode* InNode) const
-	{
-		if (InNode)
-		{
-			FScopeLock Lock(&NodeTrackingLock);
-			return NodesWithWarning.Contains(InNode->GetGuid());
-		}
-		return false;
-	}
-
-	int32 FContext::GetNumWarnings() const
-	{
-		FScopeLock Lock(&NodeTrackingLock);
-		return NodesWithWarning.Num();
-	}
-
 	void FContext::Error(const FText& InError, const FDataflowNode* InNode, const FDataflowOutput* InOutput)
 	{
-		if (InNode)
-		{
-			FScopeLock Lock(&NodeTrackingLock);
-			NodesWithError.Add(InNode->GetGuid());
-		}
 		BroadCastLog(EMessageSeverity::Error, InError, InNode, InOutput);
 	}
 
@@ -453,22 +416,6 @@ namespace UE::Dataflow
 		Error(Message, InNode, InOutput);
 	}
 
-	bool FContext::NodeHasError(const FDataflowNode* InNode) const
-	{
-		if (InNode)
-		{
-			FScopeLock Lock(&NodeTrackingLock);
-			return NodesWithError.Contains(InNode->GetGuid());
-		}
-		return false;
-	}
-
-	int32 FContext::GetNumErrors() const
-	{
-		FScopeLock Lock(&NodeTrackingLock);
-		return NodesWithError.Num();
-	}
-
 	void FContext::BroadCastLog(EMessageSeverity::Type Severity, const FText& InMessage, const FDataflowNode* InNode, const FDataflowOutput* InOutput)
 	{
 		FLogMessage LogMessage
@@ -501,25 +448,6 @@ namespace UE::Dataflow
 			Async(EAsyncExecution::TaskGraphMainThread, MoveTemp(BroadCastLogOnGameThread));
 		}
 	}
-	void FContext::ClearNodesData()
-	{
-		FScopeLock Lock(&NodeTrackingLock);
-		NodesWithInfo.Reset();
-		NodesWithWarning.Reset();
-		NodesWithError.Reset();
-	}
-
-	void FContext::ClearNodeData(const FDataflowNode* InNode)
-	{
-		if (InNode)
-		{
-			FScopeLock Lock(&NodeTrackingLock);
-			const FGuid NodeGuid = InNode->GetGuid();
-			NodesWithInfo.Remove(NodeGuid);
-			NodesWithWarning.Remove(NodeGuid);
-			NodesWithError.Remove(NodeGuid);
-		}
-	}
 
 	FContextScopedCallstack::FContextScopedCallstack(FContext& InContext, const FDataflowConnection* InConnection)
 		: Context(InContext)
@@ -527,7 +455,6 @@ namespace UE::Dataflow
 	{
 		bLoopDetected = Context.IsInCallstack(Connection);
 		Context.ResetPerfDataForNode(*Connection->GetOwningNode());
-		Context.ClearNodeData(Connection->GetOwningNode());
 		Context.PushToCallstack(Connection);
 	}
 
@@ -541,27 +468,6 @@ namespace UE::Dataflow
 		, Delegates(MakeShared<FDelegates>())
 	{}
 
-	void FContext::CheckIntrinsicInputs(const FDataflowOutput& Connection)
-	{
-#if WITH_EDITOR
-		//
-		// Check if all DataflowIntrinsic input(s) connected
-		// 
-		if (const FDataflowNode* Node = Connection.GetOwningNode())
-		{
-			TArray<FDataflowInput*> Inputs = Node->GetInputs();
-			for (const FDataflowInput* Input : Inputs)
-			{
-				if (Input && Input->IsRequired() && !Input->IsConnected())
-				{
-					const FString WarningMsg = FString::Printf(TEXT("Input %s must be connected"), *Input->GetName().ToString());
-					Warning(WarningMsg, Node);
-				}
-			}
-		}
-#endif // EDITOR_ONLY
-	}
-
 	// DEPRECATED 5.7
 	int32 FContext::GetArraySizeFromData(const FContextCacheKey InKey) const
 	{
@@ -694,8 +600,6 @@ namespace UE::Dataflow
 
 	bool FContext::Evaluate(const FDataflowOutput& Connection)
 	{
-		CheckIntrinsicInputs(Connection);
-
 		if (IsThreaded())
 		{
 			if (!Connection.HasCachedValue(*this))
@@ -722,9 +626,10 @@ namespace UE::Dataflow
 				OnNodeFinishEvaluateMulticast.Broadcast(Connection.GetOwningNode(), &Connection);
 			}
 
-			const bool bHasError = NodeHasError(Connection.GetOwningNode());
+			/*const bool bHasError = NodeHasError(Connection.GetOwningNode());
 
-			return !bHasError && RetValue;
+			return !bHasError && RetValue;*/
+			return RetValue;
 		}
 	}
 
diff --git a/Engine/Source/Runtime/Experimental/Dataflow/Core/Public/Dataflow/DataflowInputOutput.h b/Engine/Source/Runtime/Experimental/Dataflow/Core/Public/Dataflow/DataflowInputOutput.h
--- a/Engine/Source/Runtime/Experimental/Dataflow/Core/Public/Dataflow/DataflowInputOutput.h
+++ b/Engine/Source/Runtime/Experimental/Dataflow/Core/Public/Dataflow/DataflowInputOutput.h
@@ -55,6 +55,7 @@ protected:
 private:
 	// if the Input has DATAFLOW_INTRINSIC metatag specified
 	bool bIsRequired = false;
+	DATAFLOWCORE_API void CheckRequired(UE::Dataflow::FContext& Context) const;
 
 public:
 	DATAFLOWCORE_API FDataflowInput(const UE::Dataflow::FInputParameters& Param = {});
@@ -345,7 +346,7 @@ public:
 
 	DATAFLOWCORE_API virtual void Invalidate(const UE::Dataflow::FTimestamp& ModifiedTimestamp = UE::Dataflow::FTimestamp::Current()) override;
 
-	// Check if OwningNode has an error or it has failed, and if it does set the timestamp on the output(s) to FTimestamp::Invalid
+	UE_DEPRECATED(5.8, "HasNodeFailedOrErrored is deprecated and will only return false, failure or error of node should be done by listening to the context log messages")
 	DATAFLOWCORE_API bool HasNodeFailedOrErrored(UE::Dataflow::FContext& Context) const;
 
 private:
@@ -391,6 +392,7 @@ struct TStructOpsTypeTraits<FDataflowArrayOutput> : public  TStructOpsTypeTraits
 template<typename T>
 const T& FDataflowInput::GetValue(UE::Dataflow::FContext& Context, const T& Default) const
 {
+	CheckRequired(Context);
 	if (const FDataflowOutput* ConnectionOut = GetConnection())
 	{
 		if (ConnectionOut->HasFrozenValue())
@@ -413,6 +415,7 @@ const T& FDataflowInput::GetValue(UE::Dataflow::FContext& Context, const T& Defa
 template<typename TAnyType UE_REQUIRES(!std::is_same_v<typename TAnyType::FStorageType, void>)>
 typename TAnyType::FStorageType FDataflowInput::GetValueFromAnyType(UE::Dataflow::FContext& Context, const typename TAnyType::FStorageType& Default) const
 {
+	CheckRequired(Context);
 	typename TAnyType::FStorageType ReturnValue = Default;
 	if (const FDataflowOutput* ConnectionOut = GetConnection())
 	{
@@ -444,6 +447,7 @@ typename TAnyType::FStorageType FDataflowInput::GetValueFromAnyType(UE::Dataflow
 template<typename TAnyType UE_REQUIRES(std::is_same_v<typename TAnyType::FStorageType, void>)>
 typename UE::Dataflow::FContextValue FDataflowInput::GetValueFromAnyType(UE::Dataflow::FContext& Context) const
 {
+	CheckRequired(Context);
 	UE::Dataflow::FContextCacheKey CacheKey = 0;
 	if (const FDataflowOutput* ConnectionOut = GetConnection())
 	{
diff --git a/Engine/Source/Runtime/Experimental/Dataflow/Core/Public/Dataflow/DataflowNodeParameters.h b/Engine/Source/Runtime/Experimental/Dataflow/Core/Public/Dataflow/DataflowNodeParameters.h
--- a/Engine/Source/Runtime/Experimental/Dataflow/Core/Public/Dataflow/DataflowNodeParameters.h
+++ b/Engine/Source/Runtime/Experimental/Dataflow/Core/Public/Dataflow/DataflowNodeParameters.h
@@ -269,12 +269,6 @@ namespace UE::Dataflow
 		// in that case posting an invalid reference, will allow the evaluatino to go through and the node reading it will get a default value instead
 		DATAFLOWCORE_API void SetNullData(FContextCacheKey InKey, const FProperty* InProperty, const FGuid& InNodeGuid, uint32 InNodeHash, const FTimestamp& InTimestamp);
 
-		UE_DEPRECATED(5.7, "Use UE::Dataflow::FContextValue from calling GetValue on the node to access the same functionality")
-		DATAFLOWCORE_API int32 GetArraySizeFromData(const FContextCacheKey InKey) const;
-
-		UE_DEPRECATED(5.7, "Use UE::Dataflow::FContextValue from calling GetValue on the node to access the same functionality")
-		DATAFLOWCORE_API void SetArrayElementFromData(const FContextCacheKey InArrayKey, int32 Index, const FContextCacheKey InElementKey, const FProperty* InProperty, const FGuid& InNodeGuid, uint32 InNodeHash, const FTimestamp& InTimestamp);
-
 		DATAFLOWCORE_API bool CopyDataToAnotherContext(const FContextCacheKey InSourceKey, FContext& TargetContext, const FContextCacheKey InTargetKey, const FProperty* InProperty, const FGuid& InNodeGuid, uint32 InNodeHash, const FTimestamp& InTimestamp) const;
 
 		template<class T>
@@ -320,24 +314,12 @@ namespace UE::Dataflow
 
 		DATAFLOWCORE_API void Info(const FText& InInfo, const FDataflowNode* InNode = nullptr, const FDataflowOutput* InOutput = nullptr);
 		DATAFLOWCORE_API void Info(const FString& InInfo, const FDataflowNode* InNode = nullptr, const FDataflowOutput* InOutput = nullptr);
-		DATAFLOWCORE_API int32 GetNumInfo() const;
 		
 		DATAFLOWCORE_API void Warning(const FText& InWarning, const FDataflowNode* InNode = nullptr, const FDataflowOutput* InOutput = nullptr);
 		DATAFLOWCORE_API void Warning(const FString& InWarning, const FDataflowNode* InNode = nullptr, const FDataflowOutput* InOutput = nullptr);
-		DATAFLOWCORE_API int32 GetNumWarnings() const;
-
+		
 		DATAFLOWCORE_API void Error(const FText& InError, const FDataflowNode* InNode = nullptr, const FDataflowOutput* InOutput = nullptr);
 		DATAFLOWCORE_API void Error(const FString& InError, const FDataflowNode* InNode = nullptr, const FDataflowOutput* InOutput = nullptr);
-		DATAFLOWCORE_API int32 GetNumErrors() const;
-
-		DATAFLOWCORE_API bool NodeHasWarning(const FDataflowNode* InNode) const;
-		DATAFLOWCORE_API bool NodeHasError(const FDataflowNode* InNode) const;
-		
-		UE_DEPRECATED(5.8, "Deprecated state of a node, only Errors and Warnings are now reported")
-		bool NodeFailed(const FDataflowNode* InNode) const { return false; }
-
-		DATAFLOWCORE_API void ClearNodesData();
-		DATAFLOWCORE_API void ClearNodeData(const FDataflowNode* InNode);
 
 		DECLARE_MULTICAST_DELEGATE_TwoParams(FOnNodeBeginEvaluateMulticast, const FDataflowNode* Node, const FDataflowOutput* Output);
 		FOnNodeBeginEvaluateMulticast OnNodeBeginEvaluateMulticast;
@@ -345,6 +327,50 @@ namespace UE::Dataflow
 		DECLARE_MULTICAST_DELEGATE_TwoParams(FOnNodeFinishEvaluateMulticast, const FDataflowNode* Node, const FDataflowOutput* Output);
 		FOnNodeFinishEvaluateMulticast OnNodeFinishEvaluateMulticast;
 
+		struct FLogMessage
+		{
+			EMessageSeverity::Type Severity = EMessageSeverity::Info;
+			TWeakPtr<const FDataflowNode> Node;
+			FGuid OutputGuid;
+			FText Text;
+		};
+
+		DECLARE_MULTICAST_DELEGATE_OneParam(FOnContextLogMulticast, const FLogMessage& Message);
+		FOnContextLogMulticast& GetOnContextLogMulticast() { return Delegates->OnContextLog; }
+
+		//---------------------------------------------------------------
+		// Deprecated methods section
+		//---------------------------------------------------------------
+		UE_DEPRECATED(5.7, "Use UE::Dataflow::FContextValue from calling GetValue on the node to access the same functionality")
+		DATAFLOWCORE_API int32 GetArraySizeFromData(const FContextCacheKey InKey) const;
+
+		UE_DEPRECATED(5.7, "Use UE::Dataflow::FContextValue from calling GetValue on the node to access the same functionality")
+		DATAFLOWCORE_API void SetArrayElementFromData(const FContextCacheKey InArrayKey, int32 Index, const FContextCacheKey InElementKey, const FProperty* InProperty, const FGuid& InNodeGuid, uint32 InNodeHash, const FTimestamp& InTimestamp);
+
+		UE_DEPRECATED(5.8, "Deprecated method, context no longer count the Info calls")
+		int32 GetNumInfo() const { return 0; }
+
+		UE_DEPRECATED(5.8, "Deprecated method, context no longer count the Warning calls")
+		int32 GetNumWarnings() const { return 0; }
+
+		UE_DEPRECATED(5.8, "Deprecated method, context no longer count the Error calls")
+		int32 GetNumErrors() const { return 0; }
+
+		UE_DEPRECATED(5.8, "Deprecated method, context does not store the warning state of a node, this is the responsability of the code using the context")
+		bool NodeHasWarning(const FDataflowNode* InNode) const { return false; }
+
+		UE_DEPRECATED(5.8, "Deprecated method, context does not store the error state of a node, this is the responsability of the code using the context")
+		bool NodeHasError(const FDataflowNode* InNode) const { return false; }
+
+		UE_DEPRECATED(5.8, "Deprecated state of a node, failure state is not longer reported")
+		bool NodeFailed(const FDataflowNode* InNode) const { return false; }
+
+		UE_DEPRECATED(5.8, "ClearNodesData is deprecated as it is no longer needed")
+		void ClearNodesData() {}
+
+		UE_DEPRECATED(5.8, "ClearNodeData is deprecated as it is no longer needed")
+		void ClearNodeData(const FDataflowNode* InNode) {}
+
 		DECLARE_MULTICAST_DELEGATE_ThreeParams(FOnContextHasInfoMulticast, const FDataflowNode* Node, const FDataflowOutput* Output, const FString& Info);
 		UE_DEPRECATED(5.8, "Use OnContextLogMulticast instead.")
 		FOnContextHasInfoMulticast OnContextHasInfo;
@@ -357,17 +383,6 @@ namespace UE::Dataflow
 		UE_DEPRECATED(5.8, "Use OnContextLogMulticast instead.")
 		FOnContextHasErrorMulticast OnContextHasError;
 
-		struct FLogMessage
-		{
-			EMessageSeverity::Type Severity = EMessageSeverity::Info;
-			TWeakPtr<const FDataflowNode> Node;
-			FGuid OutputGuid;
-			FText Text;
-		};
-
-		DECLARE_MULTICAST_DELEGATE_OneParam(FOnContextLogMulticast, const FLogMessage& Message);
-		FOnContextLogMulticast& GetOnContextLogMulticast() { return Delegates->OnContextLog; }
-
 	protected:
 		DATAFLOWCORE_API void BeginContextEvaluation(const FDataflowNode* Node, const FDataflowOutput* Output);
 
@@ -383,17 +398,10 @@ namespace UE::Dataflow
 		};
 		TSharedRef<FDelegates> Delegates;
 
-		void CheckIntrinsicInputs(const FDataflowOutput& Connection);
-
 #if DATAFLOW_EDITOR_EVALUATION
 		FContextCallstack Callstack;
 		FContextPerfData PerfData;
 #endif
-		mutable FCriticalSection NodeTrackingLock;
-		TSet<FGuid> NodesWithInfo;
-		TSet<FGuid> NodesWithWarning;
-		TSet<FGuid> NodesWithError;
-
 		/** used to store the dependent asset created during the evaluation of the graph */
 		FContextAssetStore AssetStore;
 	};
diff --git a/Engine/Source/Runtime/Experimental/Dataflow/Engine/Private/Dataflow/DataflowBlueprintLibrary.cpp b/Engine/Source/Runtime/Experimental/Dataflow/Engine/Private/Dataflow/DataflowBlueprintLibrary.cpp
--- a/Engine/Source/Runtime/Experimental/Dataflow/Engine/Private/Dataflow/DataflowBlueprintLibrary.cpp
+++ b/Engine/Source/Runtime/Experimental/Dataflow/Engine/Private/Dataflow/DataflowBlueprintLibrary.cpp
@@ -41,13 +41,17 @@ bool UDataflowBlueprintLibrary::EvaluateDataflow(UDataflow* Dataflow, UObject* A
 {
 	if (Dataflow && Dataflow->Dataflow)
 	{
+		TAtomic<bool> bHasError = false;
+
 		using namespace UE::Dataflow;
 		FEngineContext Context(AssetToUpdate);
 		Context.GetOnContextLogMulticast().AddLambda(
-			[](const FContext::FLogMessage& Message)
+			[&bHasError](const FContext::FLogMessage& Message)
 			{
 				if (Message.Severity == EMessageSeverity::Error)
 				{
+					bHasError = true;
+
 					FString NodeString(TEXT("None"));
 					if (TSharedPtr<const FDataflowNode> Node = Message.Node.Pin())
 					{
@@ -57,7 +61,7 @@ bool UDataflowBlueprintLibrary::EvaluateDataflow(UDataflow* Dataflow, UObject* A
 				}
 			});
 		Context.EvaluateGraph<FDataflowDirectEvaluator>(Dataflow, {});
-		return (Context.GetNumErrors() == 0);
+		return bHasError;
 	}
 	UE_LOG(LogChaosDataflow, Error, TEXT("EvaluateDataflow : Invalid Dataflow Asset : %s"), *(Dataflow? Dataflow->GetFName().ToString(): FString(TEXT("null"))));
 	return false;
diff --git a/Engine/Source/Runtime/Experimental/Dataflow/Engine/Private/Dataflow/DataflowSubGraphNodes.cpp b/Engine/Source/Runtime/Experimental/Dataflow/Engine/Private/Dataflow/DataflowSubGraphNodes.cpp
--- a/Engine/Source/Runtime/Experimental/Dataflow/Engine/Private/Dataflow/DataflowSubGraphNodes.cpp
+++ b/Engine/Source/Runtime/Experimental/Dataflow/Engine/Private/Dataflow/DataflowSubGraphNodes.cpp
@@ -107,6 +107,21 @@ namespace UE::Dataflow::SubGraphNodes::Private
 					// we need to make sure the graph output and caller output are of compatible type
 					if (UE::Dataflow::AreTypesCompatible(OutputToEvaluate->GetType(), CallerOutput.GetType()))
 					{
+						TAtomic<bool> bHasWarning = false;
+						TAtomic<bool> bHasError = false;
+						FDelegateHandle OnLogHandle = ParentContext.GetOnContextLogMulticast().AddLambda(
+							[&bHasWarning, &bHasError](const FLogMessage& InMessage)
+							{
+								if (InMessage.Severity == EMessageSeverity::Warning)
+								{
+									bHasWarning = true;
+								}
+								else if (InMessage.Severity == EMessageSeverity::Error)
+								{
+									bHasError = true;
+								}
+							});
+
 						if (Evaluate(*OutputToEvaluate))
 						{
 							// copy the cache value from this cachestore to the parent one 
@@ -115,12 +130,14 @@ namespace UE::Dataflow::SubGraphNodes::Private
 							);
 						}
 
-						if (GetNumErrors() > 0)
+						ParentContext.GetOnContextLogMulticast().Remove(OnLogHandle);
+
+						if (bHasError)
 						{
 							const FText ErrorMessage(LOCTEXT("SubgraphContext_EvaluationError", "Evaluation of SubGraph returned errors : see the details above"));
 							ParentContext.Error(ErrorMessage, &CallerNode, &CallerOutput);
 						}
-						if (GetNumWarnings() > 0)
+						if (bHasWarning)
 						{
 							const FText WarningMessage(LOCTEXT("SubgraphContext_EvaluationWarning", "Evaluation of SubGraph returned warnings: see the details above"));
 							ParentContext.Warning(WarningMessage, &CallerNode, &CallerOutput);
