Dataflow tool : update tool snapshot  and implement UI customization
- add lock and description to snapshots
- add a snapshot set class to track all snapshots and the active one
- add ui customization to display them in a in reverse order and allow user to change the active snapshot

#Synced-CL 51688678

diff --git a/Engine/Plugins/Experimental/Dataflow/Source/DataflowNodes/DataflowNodes.Build.cs b/Engine/Plugins/Experimental/Dataflow/Source/DataflowNodes/DataflowNodes.Build.cs
--- a/Engine/Plugins/Experimental/Dataflow/Source/DataflowNodes/DataflowNodes.Build.cs
+++ b/Engine/Plugins/Experimental/Dataflow/Source/DataflowNodes/DataflowNodes.Build.cs
@@ -30,7 +30,10 @@ namespace UnrealBuildTool.Rules
 				{
 					"GeometryFramework",
 					"ImageCore",
+					"InputCore",
 					"ModelingComponents",
+					"Slate",
+					"SlateCore",
 					"StaticMeshDescription",
 				}
 			);
@@ -40,7 +43,8 @@ namespace UnrealBuildTool.Rules
 				PrivateDependencyModuleNames.AddRange(
 					new string[]
 					{
-						"SkeletalMeshUtilitiesCommon"
+						"SkeletalMeshUtilitiesCommon",
+						"PropertyEditor",
 					}
 				);
 			}
diff --git a/Engine/Plugins/Experimental/Dataflow/Source/DataflowNodes/Private/Dataflow/Customization/DataflowToolNodeSnapshotCustomization.cpp b/Engine/Plugins/Experimental/Dataflow/Source/DataflowNodes/Private/Dataflow/Customization/DataflowToolNodeSnapshotCustomization.cpp
new file mode 100644
--- /dev/null
+++ b/Engine/Plugins/Experimental/Dataflow/Source/DataflowNodes/Private/Dataflow/Customization/DataflowToolNodeSnapshotCustomization.cpp
@@ -0,0 +1,432 @@
+// Copyright Epic Games, Inc. All Rights Reserved.#include "DataflowToolNodeSnapshotCustomization.h"
+
+#include "Dataflow/Customization/DataflowToolNodeSnapshotCustomization.h"
+
+#include "Dataflow/DataflowToolNode.h"
+
+#if WITH_EDITOR
+
+#include "DetailLayoutBuilder.h"
+#include "DetailWidgetRow.h"
+#include "IDetailChildrenBuilder.h"
+
+#include "Framework/MultiBox/MultiBoxBuilder.h"
+#include "Widgets/Text/STextBlock.h"
+#include "Widgets/Input/SButton.h"
+#include "Widgets/Input/SEditableTextBox.h"
+#include "Widgets/Images/SImage.h"
+#include "Styling/AppStyle.h"
+
+#define LOCTEXT_NAMESPACE "DataflowToolNodeSnapshotCustomization"
+
+namespace FDataflowToolNodeSnapshotProperties::Private
+{
+	static const FName NamePropertyName = "Name";
+	static const FName DescPropertyName = "Description";
+	static const FName DatePropertyName = "Date";	
+	static const FName LockedPropertyName = "bLocked";
+}
+
+////////////////////////////////////////////////////////////////////////////////////////////////////////////////
+
+TSharedRef<IPropertyTypeCustomization> FDataflowToolNodeSnapshotCustomization::MakeInstance()
+{
+	return MakeShareable(new FDataflowToolNodeSnapshotCustomization());
+}
+
+// CustomizeHeader just delegates to BuildWidget
+void FDataflowToolNodeSnapshotCustomization::CustomizeHeader(
+	TSharedRef<IPropertyHandle> PropertyHandle,
+	FDetailWidgetRow& HeaderRow,
+	IPropertyTypeCustomizationUtils& CustomizationUtils)
+{
+	HeaderRow
+		.WholeRowContent()
+		[
+			BuildWidget(PropertyHandle)
+		];
+}
+
+TSharedRef<SWidget> FDataflowToolNodeSnapshotCustomization::BuildWidget(
+	TSharedRef<IPropertyHandle> SnapshotHandle)
+{
+	return SNew(SHorizontalBox)
+		.ToolTipText_Lambda([SnapshotHandle]()
+			{
+				FString Desc;
+				if (TSharedPtr<IPropertyHandle> DescHandle = SnapshotHandle->GetChildHandle(FDataflowToolNodeSnapshotProperties::Private::DescPropertyName))
+				{
+					DescHandle->GetValue(Desc);
+				}
+				return FText::FromString(Desc);
+			})
+
+		// Lock toggle
+		+ SHorizontalBox::Slot()
+		.AutoWidth()
+		.VAlign(VAlign_Center)
+		.Padding(2.f, 0.f, 6.f, 0.f)
+		[
+			SNew(SButton)
+			.ButtonStyle(FAppStyle::Get(), "NoBorder")
+			.ContentPadding(2.f)
+			.ToolTipText_Lambda([SnapshotHandle]()
+				{
+					bool bLocked = false;
+					if (TSharedPtr<IPropertyHandle> LockedHandle = SnapshotHandle->GetChildHandle(FDataflowToolNodeSnapshotProperties::Private::LockedPropertyName))
+					{
+						LockedHandle->GetValue(bLocked);
+					}
+					return bLocked
+						? LOCTEXT("UnlockAction", "Unlock — may be discarded during cleanup")
+						: LOCTEXT("LockAction", "Lock — preserved during cleanup");
+				})
+			.OnClicked_Lambda([SnapshotHandle]() -> FReply
+				{
+					if (TSharedPtr<IPropertyHandle> LockedHandle = SnapshotHandle->GetChildHandle(FDataflowToolNodeSnapshotProperties::Private::LockedPropertyName))
+					{
+						bool bLocked = false;
+						LockedHandle->GetValue(bLocked);
+						LockedHandle->SetValue(!bLocked);
+					}
+					return FReply::Handled();
+				})
+			[
+				SNew(SImage)
+				.DesiredSizeOverride(FVector2D(16.f, 16.f))
+				.Image_Lambda([SnapshotHandle]() -> const FSlateBrush*
+					{
+						bool bLocked = false;
+						if (TSharedPtr<IPropertyHandle> LockedHandle = SnapshotHandle->GetChildHandle(FDataflowToolNodeSnapshotProperties::Private::LockedPropertyName))
+						{
+							LockedHandle->GetValue(bLocked);
+						}
+						return bLocked
+							? FAppStyle::GetBrush("PropertyWindow.Locked")
+							: FAppStyle::GetBrush("PropertyWindow.Unlocked");
+					})
+			]
+		]
+
+	// Editable name 
+	+ SHorizontalBox::Slot()
+	.FillWidth(0.50f)
+	.VAlign(VAlign_Center)
+	.Padding(0.f, 0.f, 8.f, 0.f)
+	[
+		SNew(SEditableTextBox)
+		.Font(FAppStyle::GetFontStyle("PropertyWindow.NormalFont"))
+		.Text_Lambda([SnapshotHandle]() -> FText
+			{
+				FString Name;
+				if (TSharedPtr<IPropertyHandle> NameHandle = SnapshotHandle->GetChildHandle(FDataflowToolNodeSnapshotProperties::Private::NamePropertyName))
+				{
+					NameHandle->GetValue(Name);
+				}
+				return FText::FromString(Name);
+			})
+		.OnTextCommitted_Lambda([SnapshotHandle](const FText& NewText, ETextCommit::Type)
+			{
+				if (TSharedPtr<IPropertyHandle> NameHandle = SnapshotHandle->GetChildHandle(FDataflowToolNodeSnapshotProperties::Private::NamePropertyName))
+				{
+					NameHandle->SetValue(NewText.ToString());
+				}
+			})
+	]
+
+	// ead-only description
+	+ SHorizontalBox::Slot()
+		.FillWidth(0.50f)
+		.VAlign(VAlign_Center)
+		.Padding(0.f, 0.f, 8.f, 0.f)
+		[
+			SNew(STextBlock)
+				.Font(FAppStyle::GetFontStyle("PropertyWindow.NormalFont"))
+				.Text_Lambda([SnapshotHandle]() -> FText
+					{
+						FString Desc;
+						if (TSharedPtr<IPropertyHandle> DescHandle = SnapshotHandle->GetChildHandle(FDataflowToolNodeSnapshotProperties::Private::DescPropertyName))
+						{
+							DescHandle->GetValue(Desc);
+						}
+						return FText::FromString(Desc);
+					})
+		]
+
+	// Date (read-only)
+	+ SHorizontalBox::Slot()
+	.AutoWidth()
+	.VAlign(VAlign_Center)
+	.Padding(0.f, 0.f, 4.f, 0.f)
+	[
+		SNew(STextBlock)
+		.Font(FAppStyle::GetFontStyle("PropertyWindow.NormalFont"))
+		.ColorAndOpacity(FSlateColor::UseSubduedForeground())
+		.Text_Lambda([SnapshotHandle]() -> FText
+			{
+				if (TSharedPtr<IPropertyHandle> DateHandle = SnapshotHandle->GetChildHandle(FDataflowToolNodeSnapshotProperties::Private::DatePropertyName))
+				{
+					TArray<void*> RawData;
+					DateHandle->AccessRawData(RawData);
+					if (!RawData.IsEmpty() && RawData[0])
+					{
+						return FText::FromString(static_cast<FDateTime*>(RawData[0])->ToFormattedString(TEXT("%a, %b %e, %Y at %H:%M")));
+					}
+				}
+				return FText::GetEmpty();
+			})
+	];
+}
+
+////////////////////////////////////////////////////////////////////////////////////////////////////////////////////////////////////////
+
+
+TSharedRef<IPropertyTypeCustomization> FDataflowToolNodeSnapshotSetCustomization::MakeInstance()
+{
+	return MakeShareable(new FDataflowToolNodeSnapshotSetCustomization());
+}
+
+void FDataflowToolNodeSnapshotSetCustomization::CustomizeHeader(
+	TSharedRef<IPropertyHandle> PropertyHandle,
+	FDetailWidgetRow& HeaderRow,
+	IPropertyTypeCustomizationUtils& CustomizationUtils)
+{
+	// PropertyHandle is the owning struct — find the snapshots array inside it
+	SnapshotSetHandle = PropertyHandle;
+	ArrayHandle = PropertyHandle->GetChildHandle(FName("Snapshots"));
+	ActiveIndexHandle = PropertyHandle->GetChildHandle(FName("ActiveSnapshot"));
+
+	if (ArrayHandle.IsValid())
+	{
+		// Rebuild list whenever the array changes
+		ArrayHandle->SetOnPropertyValueChanged(FSimpleDelegate::CreateLambda([this]()
+			{
+				RefreshView();
+			}));
+
+		ArrayHandle->AsArray()->SetOnNumElementsChanged(
+			FSimpleDelegate::CreateLambda([this]()
+				{
+					RefreshView();
+				}));
+	}
+
+	RefreshListItems();
+
+	HeaderRow
+		.NameContent()
+		[
+			SNew(STextBlock)
+				.Text(PropertyHandle->GetPropertyDisplayName())
+				.Font(IDetailLayoutBuilder::GetDetailFontBold())
+		]
+		.ValueContent()
+		[
+			SNew(STextBlock)
+				.Text(this, &FDataflowToolNodeSnapshotSetCustomization::GetHeaderText)
+				.Font(IDetailLayoutBuilder::GetDetailFont())
+				.ColorAndOpacity(FSlateColor::UseSubduedForeground())
+		];
+}
+
+void FDataflowToolNodeSnapshotSetCustomization::CustomizeChildren(
+	TSharedRef<IPropertyHandle> PropertyHandle,
+	IDetailChildrenBuilder& ChildBuilder,
+	IPropertyTypeCustomizationUtils& CustomizationUtils)
+{
+	if (ArrayHandle.IsValid())
+	{
+		ChildBuilder.AddCustomRow(FText::GetEmpty())
+		.WholeRowContent()
+		[
+			SNew(SBox)
+			.MaxDesiredHeight(300.f)
+			[
+				SAssignNew(ListView, SListView<FSnapshotItemPtr>)
+				.ListItemsSource(&ListItems)
+				.SelectionMode(ESelectionMode::Single)
+				.OnGenerateRow(this, &FDataflowToolNodeSnapshotSetCustomization::OnGenerateRow)
+				.OnSelectionChanged(this, &FDataflowToolNodeSnapshotSetCustomization::OnSelectionChanged)
+				.OnContextMenuOpening(this, &FDataflowToolNodeSnapshotSetCustomization::OnContextMenuOpening)
+				.HeaderRow
+				(
+					SNew(SHeaderRow)
+					+ SHeaderRow::Column("Snapshot")
+					.DefaultLabel(LOCTEXT("SnapshotHeader", "Snapshot"))
+					.FillWidth(1.f)
+				)
+			]
+		];
+	}
+}
+
+int32 FDataflowToolNodeSnapshotSetCustomization::GetSnapshotIndex(const FSnapshotItemPtr& Item) const
+{
+	const int32 SelectedItemIndex = ListItems.Find(Item);
+	if (SelectedItemIndex != INDEX_NONE)
+	{
+		return (ListItems.Num() - 1 - SelectedItemIndex); // list view is in inverse order
+	}
+	return INDEX_NONE;
+}
+
+bool FDataflowToolNodeSnapshotSetCustomization::IsActiveSnapshot(const FSnapshotItemPtr& Item) const
+{
+	const int32 SnapshotIndex = GetSnapshotIndex(Item);
+
+	int32 ActiveSnapshotIndex = INDEX_NONE;
+	if (ActiveIndexHandle)
+	{
+		ActiveIndexHandle->GetValue(ActiveSnapshotIndex);
+	}
+	return (ActiveSnapshotIndex == SnapshotIndex);
+}
+
+TSharedPtr<SWidget> FDataflowToolNodeSnapshotSetCustomization::OnContextMenuOpening()
+{
+	TArray<FSnapshotItemPtr> SelectedItems = ListView->GetSelectedItems();
+	if (SelectedItems.Num() == 0)
+	{
+		return nullptr;
+	}
+
+	FSnapshotItemPtr Item = SelectedItems[0];
+	const bool bIsActiveSnapshot = IsActiveSnapshot(Item);
+
+	FMenuBuilder MenuBuilder(true, nullptr);
+
+	MenuBuilder.AddMenuEntry(
+		LOCTEXT("ActivateSnapshot", "Activate snapshot"),
+		LOCTEXT("ActivateSnapshotTooltip", "Make this snaphot the active one"),
+		FSlateIcon(),
+		FUIAction(
+			FExecuteAction::CreateLambda([this, Item]()
+				{
+					const int32 SnapshotIndex = GetSnapshotIndex(Item);
+					if (SnapshotIndex != INDEX_NONE)
+					{
+						void* StructData = nullptr;
+						if (SnapshotSetHandle->GetValueData(StructData))
+						{
+							if (FDataflowToolNodeSnapshotSet* SnapshotSet = reinterpret_cast<FDataflowToolNodeSnapshotSet*>(StructData))
+							{
+								SnapshotSet->SetActiveSnapshot(SnapshotIndex);
+								RefreshView();
+							}
+						}
+					}
+				}),
+			FCanExecuteAction::CreateLambda([bIsActiveSnapshot]()
+				{
+					return !bIsActiveSnapshot;
+				}))
+	);
+
+	MenuBuilder.AddMenuSeparator();
+
+	MenuBuilder.AddMenuEntry(
+		LOCTEXT("DeleteSnapshot", "Delete snapshot"),
+		LOCTEXT("DeleteSnapshotTooltip", "Delete this snapshot"),
+		FSlateIcon(),
+		FUIAction(
+			FExecuteAction::CreateLambda([this, Item]()
+				{
+					const int32 SnapshotIndex = GetSnapshotIndex(Item);
+					if (SnapshotIndex != INDEX_NONE)
+					{
+						void* StructData = nullptr;
+						if (SnapshotSetHandle->GetValueData(StructData))
+						{
+							if (FDataflowToolNodeSnapshotSet* SnapshotSet = reinterpret_cast<FDataflowToolNodeSnapshotSet*>(StructData))
+							{
+								SnapshotSet->RemoveSnapshot(SnapshotIndex);
+								RefreshView();
+							}
+						}
+					}
+				}))
+	);
+
+	return MenuBuilder.MakeWidget();
+}
+
+void FDataflowToolNodeSnapshotSetCustomization::RefreshView()
+{
+	RefreshListItems();
+	if (ListView.IsValid())
+	{
+		ListView->RequestListRefresh();
+	}
+}
+
+void FDataflowToolNodeSnapshotSetCustomization::RefreshListItems()
+{
+	ListItems.Reset();
+	if (ArrayHandle.IsValid())
+	{
+		uint32 NumElements = 0;
+		ArrayHandle->AsArray()->GetNumElements(NumElements);
+
+		// Add in reverse order as we want the more recent on top
+		if (NumElements > 0)
+		{
+			for (int32 Index = (NumElements - 1); Index >= 0; --Index)
+			{
+				ListItems.Add(ArrayHandle->AsArray()->GetElement(Index));
+			}
+		}
+	}
+}
+
+TSharedRef<ITableRow> FDataflowToolNodeSnapshotSetCustomization::OnGenerateRow(
+	FSnapshotItemPtr Item,
+	const TSharedRef<STableViewBase>& OwnerTable)
+{
+	return SNew(STableRow<FSnapshotItemPtr>, OwnerTable)
+	[
+		SNew(SHorizontalBox)
+
+		+ SHorizontalBox::Slot()
+		.AutoWidth()
+		.VAlign(VAlign_Center)
+		.MinWidth(16.f)
+		.Padding(0.f, 0.f, 4.f, 0.f)
+		[
+			SNew(SImage)
+			.DesiredSizeOverride(FVector2D(16.f, 16.f))
+			.Image_Lambda([this, Item]() -> const FSlateBrush*
+				{
+					const bool bIsActiveSnapshot = IsActiveSnapshot(Item);
+					return bIsActiveSnapshot ? FAppStyle::GetBrush("Symbols.Check") : nullptr;
+				})
+			.ColorAndOpacity(FColor::Green)
+		]
+
+		+ SHorizontalBox::Slot()
+		.AutoWidth()
+		.VAlign(VAlign_Center)
+		.FillWidth(1.0)
+		.Padding(0.f, 0.f, 4.f, 0.f)
+		[
+			FDataflowToolNodeSnapshotCustomization::BuildWidget(Item.ToSharedRef())
+		]
+	];
+}
+
+void FDataflowToolNodeSnapshotSetCustomization::OnSelectionChanged(
+	FSnapshotItemPtr Item,
+	ESelectInfo::Type SelectInfo)
+{
+	SelectedItem = Item;
+}
+
+FText FDataflowToolNodeSnapshotSetCustomization::GetHeaderText() const
+{
+	return FText::Format(
+		LOCTEXT("SnapshotCount", "{0} snapshot(s)"),
+		FText::AsNumber(ListItems.Num()));
+}
+
+#endif // WITH_EDITOR
+
+#undef LOCTEXT_NAMESPACE
\ No newline at end of file
diff --git a/Engine/Plugins/Experimental/Dataflow/Source/DataflowNodes/Private/Dataflow/Customization/DataflowToolNodeSnapshotCustomization.h b/Engine/Plugins/Experimental/Dataflow/Source/DataflowNodes/Private/Dataflow/Customization/DataflowToolNodeSnapshotCustomization.h
new file mode 100644
--- /dev/null
+++ b/Engine/Plugins/Experimental/Dataflow/Source/DataflowNodes/Private/Dataflow/Customization/DataflowToolNodeSnapshotCustomization.h
@@ -0,0 +1,74 @@
+// Copyright Epic Games, Inc. All Rights Reserved.
+
+#pragma once
+
+#if WITH_EDITOR
+
+#include "IPropertyTypeCustomization.h"
+#include "PropertyHandle.h"
+#include "Widgets/Views/SListView.h"
+
+class FDataflowToolNodeSnapshotCustomization : public IPropertyTypeCustomization
+{
+public:
+	static TSharedRef<IPropertyTypeCustomization> MakeInstance();
+
+	// Reusable — call this anywhere you have a handle to a snapshot element
+	static TSharedRef<SWidget> BuildWidget(TSharedRef<IPropertyHandle> SnapshotHandle);
+
+	virtual void CustomizeHeader(
+		TSharedRef<IPropertyHandle> PropertyHandle,
+		FDetailWidgetRow& HeaderRow,
+		IPropertyTypeCustomizationUtils& CustomizationUtils) override;
+
+	virtual void CustomizeChildren(
+		TSharedRef<IPropertyHandle> PropertyHandle,
+		IDetailChildrenBuilder& ChildBuilder,
+		IPropertyTypeCustomizationUtils& CustomizationUtils) override
+	{
+	}
+};
+
+class FDataflowToolNodeSnapshotSetCustomization : public IPropertyTypeCustomization
+{
+public:
+	static TSharedRef<IPropertyTypeCustomization> MakeInstance();
+
+	virtual void CustomizeHeader(
+		TSharedRef<IPropertyHandle> PropertyHandle,
+		FDetailWidgetRow& HeaderRow,
+		IPropertyTypeCustomizationUtils& CustomizationUtils) override;
+
+	virtual void CustomizeChildren(
+		TSharedRef<IPropertyHandle> PropertyHandle,
+		IDetailChildrenBuilder& ChildBuilder,
+		IPropertyTypeCustomizationUtils& CustomizationUtils) override;
+
+private:
+	using FSnapshotItemPtr = TSharedPtr<IPropertyHandle>;
+
+	TSharedPtr<IPropertyHandle>				SnapshotSetHandle;
+	TSharedPtr<IPropertyHandle>				ArrayHandle;
+	TSharedPtr<IPropertyHandle>				ActiveIndexHandle;
+	TArray<FSnapshotItemPtr>				ListItems;
+	TSharedPtr<SListView<FSnapshotItemPtr>> ListView;
+	FSnapshotItemPtr						SelectedItem;
+
+	TSharedPtr<SWidget> OnContextMenuOpening();
+	void  RefreshView();
+	void RefreshListItems();
+	int32 GetSnapshotIndex(const FSnapshotItemPtr& Item) const;
+	bool IsActiveSnapshot(const FSnapshotItemPtr& Item) const;
+
+	TSharedRef<ITableRow> OnGenerateRow(
+		FSnapshotItemPtr Item,
+		const TSharedRef<STableViewBase>& OwnerTable);
+
+	void OnSelectionChanged(
+		FSnapshotItemPtr Item,
+		ESelectInfo::Type SelectInfo);
+
+	FText GetHeaderText() const;
+};
+
+#endif // WITH_EDITOR
\ No newline at end of file
diff --git a/Engine/Plugins/Experimental/Dataflow/Source/DataflowNodes/Private/Dataflow/DataflowToolNode.cpp b/Engine/Plugins/Experimental/Dataflow/Source/DataflowNodes/Private/Dataflow/DataflowToolNode.cpp
--- a/Engine/Plugins/Experimental/Dataflow/Source/DataflowNodes/Private/Dataflow/DataflowToolNode.cpp
+++ b/Engine/Plugins/Experimental/Dataflow/Source/DataflowNodes/Private/Dataflow/DataflowToolNode.cpp
@@ -11,45 +11,113 @@ uint64 FDataflowToolNodeSnapshot::GetAllocatedSize() const
 
 //////////////////////////////////////////////////////////////////////////////////////////////////////////////////////////
 
-FDataflowToolNode::FDataflowToolNode(const UE::Dataflow::FNodeParameters& InParam, FGuid InGuid)
-	: FDataflowNode(InParam, InGuid)
-{}
+TConstArrayView<FDataflowToolNodeSnapshot> FDataflowToolNodeSnapshotSet::GetSnapshots() const
+{
+	return Snapshots;
+}
 
-FDataflowToolNodeSnapshot& FDataflowToolNode::AddSnapshot()
+const FDataflowToolNodeSnapshot* FDataflowToolNodeSnapshotSet::GetActiveSnapshot() const
+{
+	if (Snapshots.IsValidIndex(ActiveSnapshot))
+	{
+		return &Snapshots[ActiveSnapshot];
+	}
+	return nullptr;
+	
+}
+
+void FDataflowToolNodeSnapshotSet::SetActiveSnapshot(int32 ActiveIndex)
+{
+	if (Snapshots.IsValidIndex(ActiveIndex))
+	{
+		ActiveSnapshot = ActiveIndex;
+	}
+}
+
+FDataflowToolNodeSnapshot& FDataflowToolNodeSnapshotSet::AddSnapshot()
 {
-	constexpr int32 MaxNumberOfSnapshots = 5;
-	if (Snapshots.Num() >= MaxNumberOfSnapshots)
+	// if all snapshots are lock we will exceed this number 
+	constexpr int32 MaxDesiredNumberOfSnapshots = 5;
+	if (Snapshots.Num() >= MaxDesiredNumberOfSnapshots)
 	{
-		Snapshots.RemoveAt(0);
+		RemoveSnapshot(FindOldestUnlockedSnapShot());
 	}
 	FDataflowToolNodeSnapshot& NewSnapshot = Snapshots.AddDefaulted_GetRef();
 	NewSnapshot.Name = TEXT("Tool Snapshot");
+	ActiveSnapshot = (Snapshots.Num() - 1);
 	return NewSnapshot;
 }
 
-bool FDataflowToolNode::RemoveSnapshot(int32 Index)
+bool FDataflowToolNodeSnapshotSet::RemoveSnapshot(int32 Index)
 {
 	if (Snapshots.IsValidIndex(Index))
 	{
 		Snapshots.RemoveAt(Index);
+		if (ActiveSnapshot == Index)
+		{
+			ActiveSnapshot = Snapshots.IsEmpty() ? INDEX_NONE : (Snapshots.Num() - 1);
+		}
+		else if (ActiveSnapshot > Index)
+		{
+			ActiveSnapshot = ActiveSnapshot - 1;
+		}
+
 		return true;
 	}
 	return false;
 }
 
-void FDataflowToolNode::SwapSnapshots(TArray<FDataflowToolNodeSnapshot>& InOutSnapshots)
+int32 FDataflowToolNodeSnapshotSet::FindOldestUnlockedSnapShot()
+{
+	int32 OldestSnapshotIndex = INDEX_NONE;
+	FDateTime OldestSnapshotDate = FDateTime::MaxValue();
+	const int32 NumSnapShots = Snapshots.Num();
+	for (int32 Index = 0; Index < NumSnapShots; ++Index)
+	{
+		const FDataflowToolNodeSnapshot& Snapshot = Snapshots[Index];
+		if (Snapshot.GetDate() < OldestSnapshotDate && !Snapshot.IsLocked())
+		{
+			OldestSnapshotIndex = Index;
+			OldestSnapshotDate = Snapshot.GetDate();
+		}
+	}
+	return OldestSnapshotIndex;
+}
+
+void FDataflowToolNodeSnapshotSet::SwapSnapshots(TArray<FDataflowToolNodeSnapshot>& InOutSnapshots)
 {
 	Swap(InOutSnapshots, Snapshots);
 }
 
+//////////////////////////////////////////////////////////////////////////////////////////////////////////////////////////
+
+FDataflowToolNode::FDataflowToolNode(const UE::Dataflow::FNodeParameters& InParam, FGuid InGuid)
+	: FDataflowNode(InParam, InGuid)
+{}
+
+FDataflowToolNodeSnapshot& FDataflowToolNode::AddSnapshot()
+{
+	return Snapshots.AddSnapshot();
+}
+
+bool FDataflowToolNode::RemoveSnapshot(int32 Index)
+{
+	return Snapshots.RemoveSnapshot(Index);
+}
+
+void FDataflowToolNode::SwapSnapshots(TArray<FDataflowToolNodeSnapshot>& InOutSnapshots)
+{
+	Snapshots.SwapSnapshots(InOutSnapshots);
+}
+
 TConstArrayView<FDataflowToolNodeSnapshot> FDataflowToolNode::GetSnapshots() const
 {
-	return Snapshots;
+	return Snapshots.GetSnapshots();
 }
 
-const FDataflowToolNodeSnapshot* FDataflowToolNode::GetLastSnapshot() const
+const FDataflowToolNodeSnapshot* FDataflowToolNode::GetActiveSnapshot() const
 {
-	return Snapshots.IsEmpty() ? nullptr : &Snapshots.Last();
+	return Snapshots.GetActiveSnapshot();
 }
 
 //////////////////////////////////////////////////////////////////////////////////////////////////////////////////////////
diff --git a/Engine/Plugins/Experimental/Dataflow/Source/DataflowNodes/Public/Dataflow/DataflowToolNode.h b/Engine/Plugins/Experimental/Dataflow/Source/DataflowNodes/Public/Dataflow/DataflowToolNode.h
--- a/Engine/Plugins/Experimental/Dataflow/Source/DataflowNodes/Public/Dataflow/DataflowToolNode.h
+++ b/Engine/Plugins/Experimental/Dataflow/Source/DataflowNodes/Public/Dataflow/DataflowToolNode.h
@@ -25,17 +25,55 @@ public:
 	/** Returns the approximate amount of memory taken by the snapshot in bytes */
 	uint64 GetAllocatedSize() const;
 
-public:
+	const FDateTime& GetDate() const { return Date; }
+	bool IsLocked() const { return bLocked; }
+
 	/** Display name of the snapshot */
 	UPROPERTY(EditAnywhere, Category = Snapshots)
 	FString Name;
 
+	/** description of the snapshot , not user editable */
 	UPROPERTY(VisibleAnywhere, Category = Snapshots)
-	FDateTime Date = FDateTime::Now();
+	FString Description;
 
 	/** collection that stores the data of the snapshot */
 	UPROPERTY()
 	FManagedArrayCollection Data;
+
+private:
+	/** Date time at when the snapshot was created - Read-only */
+	UPROPERTY(VisibleAnywhere, Category = Snapshots)
+	FDateTime Date = FDateTime::Now();
+
+	/** When true, this snapshot cannot be dropped if the maximum number of snapshot is reached and room needs to be made for a new one */
+	UPROPERTY(EditAnywhere, Category = Snapshots)
+	bool bLocked = false;
+};
+
+/**
+* This represent a set of snapshots
+*/
+USTRUCT()
+struct FDataflowToolNodeSnapshotSet
+{
+	GENERATED_BODY()
+public:
+	UE_API TConstArrayView<FDataflowToolNodeSnapshot> GetSnapshots() const;
+	UE_API const FDataflowToolNodeSnapshot* GetActiveSnapshot() const;
+
+	UE_API FDataflowToolNodeSnapshot& AddSnapshot();
+	UE_API void SetActiveSnapshot(int32 ActiveIndex);
+	UE_API bool RemoveSnapshot(int32 Index);
+	UE_API int32 FindOldestUnlockedSnapShot();
+
+	UE_API void SwapSnapshots(TArray<FDataflowToolNodeSnapshot>& InOutSnapshots);
+
+private:
+	UPROPERTY(EditAnywhere, Category = Snapshots)
+	int32 ActiveSnapshot = INDEX_NONE;
+
+	UPROPERTY(EditAnywhere, Category = Snapshots, meta = (EditFixedSize, TitleProperty = "{Name} {Date}"))
+	TArray<FDataflowToolNodeSnapshot> Snapshots;
 };
 
 /**
@@ -52,11 +90,12 @@ struct FDataflowToolNode : public FDataflowNode
 
 public:
 	UE_API TConstArrayView<FDataflowToolNodeSnapshot> GetSnapshots() const;
-	UE_API const FDataflowToolNodeSnapshot* GetLastSnapshot() const;
+	UE_API const FDataflowToolNodeSnapshot* GetActiveSnapshot() const;
 
 protected:
 	UE_API FDataflowToolNodeSnapshot& AddSnapshot();
 	UE_API bool RemoveSnapshot(int32 Index);
+	UE_API int32 FindOldestUnlockedSnapShot();
 
 	class FSnapshotToolChange : public FToolCommandChange
 	{
@@ -78,8 +117,8 @@ protected:
 private:
 	UE_API void SwapSnapshots(TArray<FDataflowToolNodeSnapshot>& InOutSnapshots);
 
-	UPROPERTY(EditAnywhere, Category = Snapshots, meta = (EditFixedSize, TitleProperty = "{Name} {Date}"))
-	TArray<FDataflowToolNodeSnapshot> Snapshots;
+	UPROPERTY(EditAnywhere, Category = Snapshots)
+	FDataflowToolNodeSnapshotSet Snapshots;
 };
 
 #undef UE_API
