Dataflow , allow managed array collection to be added as variables

#Synced-CL 51713790

diff --git a/Engine/Plugins/Experimental/Dataflow/Source/DataflowEditor/Private/Dataflow/DataflowVariablePaletteItem.cpp b/Engine/Plugins/Experimental/Dataflow/Source/DataflowEditor/Private/Dataflow/DataflowVariablePaletteItem.cpp
--- a/Engine/Plugins/Experimental/Dataflow/Source/DataflowEditor/Private/Dataflow/DataflowVariablePaletteItem.cpp
+++ b/Engine/Plugins/Experimental/Dataflow/Source/DataflowEditor/Private/Dataflow/DataflowVariablePaletteItem.cpp
@@ -3,11 +3,14 @@
 #include "Dataflow/DataflowVariablePaletteItem.h"
 
 #include "Dataflow/DataflowGraphSchemaAction.h"
+#include "GeometryCollection/ManagedArrayCollection.h"
 #include "DetailLayoutBuilder.h"
 #include "PropertyBagDetails.h"
 #include "SPinTypeSelector.h"
 #include "Widgets/Text/SInlineEditableTextBlock.h"
 
+#define LOCTEXT_NAMESPACE "DataflowVariablePaletteItem"
+
 void SDataflowVariablePaletteItem::Construct(const FArguments& InArgs, FCreateWidgetForActionData* const InCreateData, TSharedPtr<SDataflowGraphEditor> InEditor)
 {
 	if (InCreateData)
@@ -72,6 +75,29 @@ namespace UE::Dataflow::Private
 		TArray<TSharedPtr<UEdGraphSchema_K2::FPinTypeTreeInfo>> TempTypeTree;
 		GetDefault<UPropertyBagSchema>()->GetVariableTypeTree(TempTypeTree, TypeTreeFilter);
 
+		// make sure we add managed array collection type ( before the struct category )
+		int32 StructIndex = INDEX_NONE;
+		for (int32 Index = 0; Index < TempTypeTree.Num(); ++Index)
+		{
+			if (TempTypeTree[Index])
+			{
+				if (TempTypeTree[Index]->GetPinTypeNoResolve().PinCategory == UEdGraphSchema_K2::PC_Struct && !TempTypeTree[Index]->Children.IsEmpty())
+				{
+					StructIndex = Index;
+					break;
+				}
+			}
+		}
+		TSharedPtr<UEdGraphSchema_K2::FPinTypeTreeInfo> CollectionType = MakeShareable(new UEdGraphSchema_K2::FPinTypeTreeInfo(UEdGraphSchema_K2::PC_Struct, TBaseStructure<FManagedArrayCollection>::Get(), LOCTEXT("ManagedArrayCollection", "A collection fo groups and attributes")));
+		if (StructIndex != INDEX_NONE)
+		{
+			TempTypeTree.Insert(CollectionType, StructIndex);
+		}
+		else
+		{
+			TempTypeTree.Add(CollectionType);
+		}
+
 		// Filter the results
 		for (TSharedPtr<UEdGraphSchema_K2::FPinTypeTreeInfo>& PinType : TempTypeTree)
 		{
@@ -173,3 +199,5 @@ void SDataflowVariablePaletteItem::OnPinTypeChanged(const FEdGraphPinType& PinTy
 		VariableAction->SetVariableType(PinType);
 	}
 }
+
+#undef LOCTEXT_NAMESPACE 
\ No newline at end of file
