Dataflow : Add collection selection tool node
- Allow for multiple type of selection output
- Support new snapshot system

#Synced-CL 51713790

diff --git a/Engine/Plugins/Experimental/Dataflow/Source/DataflowEditor/Private/Dataflow/DataflowEditorModule.cpp b/Engine/Plugins/Experimental/Dataflow/Source/DataflowEditor/Private/Dataflow/DataflowEditorModule.cpp
--- a/Engine/Plugins/Experimental/Dataflow/Source/DataflowEditor/Private/Dataflow/DataflowEditorModule.cpp
+++ b/Engine/Plugins/Experimental/Dataflow/Source/DataflowEditor/Private/Dataflow/DataflowEditorModule.cpp
@@ -21,6 +21,8 @@
 #include "DataflowEditorTools/DataflowEditorWeightMapPaintTool.h"
 #include "DataflowEditorTools/DataflowEditorVertexAttributePaintTool.h"
 #include "Dataflow/DataflowCollectionAddScalarVertexPropertyNode.h"
+#include "Dataflow/DataflowSelectionToolNode.h"
+#include "DataflowEditorTools/DataflowMeshSelectionTool.h"
 
 #include "Dataflow/DataflowInstanceDetails.h"
 
@@ -332,6 +334,10 @@ void FDataflowEditorModule::StartupModule()
 	ToolRegistry.AddNodeToToolMapping(FDataflowCollectionEditSkeletonBonesNode::StaticType(), 
 		NewObject<UDataflowEditorEditSkeletonBonesToolBuilder>(), Actions,
 		FSlateIcon(FName("DataflowEditorStyle"), FName("Dataflow.EditSkeletonBones")), LOCTEXT("AddSkeletonEditNodeButtonText", "Edit Skeleton Bones"));
+
+	ToolRegistry.AddNodeToToolMapping(FDataflowSelectionToolNode::StaticType(),
+		NewObject<UDataflowMeshSelectionToolBuilder>(), Actions,
+		FSlateIcon(FAppStyle::GetAppStyleSetName(), "LevelEditor.StrictBoxSelect"), LOCTEXT("AddSelectionNodeButtonText", "Selection"));
 		
 	DATAFLOW_NODE_REGISTER_CREATION_FACTORY(FDataflowCorrectSkinWeightsNode);
 	DATAFLOW_NODE_REGISTER_CREATION_FACTORY(FDataflowGetSkinningSelectionNode);
@@ -400,6 +406,7 @@ void FDataflowEditorModule::ShutdownModule()
 	ToolRegistry.RemoveNodeToToolMapping(FDataflowCollectionAddScalarVertexPropertyNode::StaticType());
 	ToolRegistry.RemoveNodeToToolMapping(FDataflowCollectionEditSkinWeightsNode::StaticType());
 	ToolRegistry.RemoveNodeToToolMapping(FDataflowCollectionEditSkeletonBonesNode::StaticType());
+	ToolRegistry.RemoveNodeToToolMapping(FDataflowSelectionToolNode::StaticType());
 
 	UE::Dataflow::FDataflowConstructionVisualizationRegistry& ConstructionVisualizationRegistry = UE::Dataflow::FDataflowConstructionVisualizationRegistry::GetInstance();
 	ConstructionVisualizationRegistry.DeregisterVisualization(UE::Dataflow::FMeshStatsConstructionVisualization::Name);
diff --git a/Engine/Plugins/Experimental/Dataflow/Source/DataflowEditor/Private/DataflowEditorTools/DataflowDynamicMeshMapping.cpp b/Engine/Plugins/Experimental/Dataflow/Source/DataflowEditor/Private/DataflowEditorTools/DataflowDynamicMeshMapping.cpp
new file mode 100644
--- /dev/null
+++ b/Engine/Plugins/Experimental/Dataflow/Source/DataflowEditor/Private/DataflowEditorTools/DataflowDynamicMeshMapping.cpp
@@ -0,0 +1,111 @@
+// Copyright Epic Games, Inc. All Rights Reserved.
+
+#include "DataflowEditorTools/DataflowDynamicMeshMapping.h"
+
+#include "DynamicMesh/DynamicMesh3.h"
+#include "DynamicMesh/NonManifoldMappingSupport.h"
+
+namespace UE::Dataflow
+{
+	void FDynamicMeshMapping::Init(const UE::Geometry::FDynamicMesh3& Mesh)
+	{
+		Reset();
+
+		const int32 MeshMaxVertexID = Mesh.MaxVertexID();
+		if (MeshMaxVertexID <= 0)
+		{
+			return;
+		}
+
+		using namespace UE::Geometry;
+		const FNonManifoldMappingSupport NonManifoldMapping(Mesh);
+
+		bHasNonManifoldMapping = NonManifoldMapping.IsNonManifoldVertexInSource();
+		if (bHasNonManifoldMapping)
+		{
+			int32 MinSourceVertID = TNumericLimits<int32>::Max();
+			int32 MaxSourceVertID = 0;
+			DynamicMeshToSource.SetNumUninitialized(MeshMaxVertexID);
+			for (int32 DynamicMeshVertID = 0; DynamicMeshVertID < MeshMaxVertexID; ++DynamicMeshVertID)
+			{
+				if (Mesh.IsVertex(DynamicMeshVertID))
+				{
+					const int32 SourceVertID = NonManifoldMapping.GetOriginalNonManifoldVertexID(DynamicMeshVertID);
+					MinSourceVertID = FMath::Min(MinSourceVertID, SourceVertID);
+					MaxSourceVertID = FMath::Max(MaxSourceVertID, SourceVertID);
+					DynamicMeshToSource[DynamicMeshVertID] = SourceVertID;
+				}
+				else
+				{
+					DynamicMeshToSource[DynamicMeshVertID] = INDEX_NONE;
+				}
+			}
+
+			if (MaxSourceVertID >= MinSourceVertID)
+			{
+				// This number is not the total number of verts in the original mesh but the number of source verts mapped to the dynamic mesh
+				const int32 NumSourceVerts = MaxSourceVertID - MinSourceVertID + 1;
+
+				SourceVertexIDOffset = MinSourceVertID;
+				SourceToDynamicMesh.SetNum(NumSourceVerts);
+				for (int32 DynamicMeshVertID = 0; DynamicMeshVertID < MeshMaxVertexID; ++DynamicMeshVertID)
+				{
+					const int32 SourceVertID = DynamicMeshToSource[DynamicMeshVertID];
+					if (SourceVertID != INDEX_NONE)
+					{
+						const int32 LocalSourceVertID = (SourceVertID - SourceVertexIDOffset);
+						if (SourceToDynamicMesh.IsValidIndex(LocalSourceVertID))
+						{
+							SourceToDynamicMesh[SourceVertID - SourceVertexIDOffset].Add(DynamicMeshVertID);
+						}
+					}
+				}
+			}
+		}
+	}
+
+	void FDynamicMeshMapping::Reset()
+	{
+		bHasNonManifoldMapping = false;
+		DynamicMeshToSource.Reset();
+		SourceToDynamicMesh.Reset();
+		SourceVertexIDOffset = 0;
+	}
+
+	void FDynamicMeshMapping::RemapVertices(const TSet<int32>& InVertices, TSet<int32>& OutVertices, ERemapDirection Direction) const
+	{
+		OutVertices.Reset();
+		if (!bHasNonManifoldMapping)
+		{
+			// Direct mapping, regardless of the direction
+			OutVertices.Append(InVertices);
+			return;
+		}
+		if (Direction == ERemapDirection::DynamicMeshToSource)
+		{
+			OutVertices.Reserve(InVertices.Num());
+			for (const int32 DynamicMeshVertID : InVertices)
+			{
+				if (DynamicMeshToSource.IsValidIndex(DynamicMeshVertID))
+				{
+					if (ensure(DynamicMeshToSource[DynamicMeshVertID] != INDEX_NONE))
+					{
+						OutVertices.Add(DynamicMeshToSource[DynamicMeshVertID]);
+					}
+				}
+			}
+		}
+		else if (Direction == ERemapDirection::SourceToDynamicMesh)
+		{
+			OutVertices.Reserve(InVertices.Num());
+			for (const int32 SourceVertID : InVertices)
+			{
+				const int32 Index = (SourceVertID - SourceVertexIDOffset);
+				if (SourceToDynamicMesh.IsValidIndex(Index))
+				{
+					OutVertices.Append(SourceToDynamicMesh[Index]);
+				}
+			}
+		}
+	}
+}
diff --git a/Engine/Plugins/Experimental/Dataflow/Source/DataflowEditor/Private/DataflowEditorTools/DataflowDynamicMeshMapping.h b/Engine/Plugins/Experimental/Dataflow/Source/DataflowEditor/Private/DataflowEditorTools/DataflowDynamicMeshMapping.h
new file mode 100644
--- /dev/null
+++ b/Engine/Plugins/Experimental/Dataflow/Source/DataflowEditor/Private/DataflowEditorTools/DataflowDynamicMeshMapping.h
@@ -0,0 +1,58 @@
+// Copyright Epic Games, Inc. All Rights Reserved.
+
+#pragma once
+
+#include "HAL/Platform.h"
+#include "Containers/Array.h"
+
+#define UE_API DATAFLOWEDITOR_API
+
+namespace UE { namespace Geometry { class FDynamicMesh3; } }
+
+namespace UE::Dataflow
+{
+	/**
+	* Represent data to map from the dynamic mesh to the source mesh
+	* in some case the dynamic mesh may have more vertices created during its construction than the source
+	*/
+	struct FDynamicMeshMapping
+	{
+	public:
+		enum class ERemapDirection : uint8
+		{
+			SourceToDynamicMesh,
+			DynamicMeshToSource
+		};
+
+		void Init(const UE::Geometry::FDynamicMesh3& Mesh);
+		void Reset();
+
+		const TArray<int32>& GetMappedSourceVertices() const { return DynamicMeshToSource; }
+
+		void RemapVertices(const TSet<int32>& InVertices, TSet<int32>& OutVertices, ERemapDirection Direction) const;
+
+	private:
+		/** true if the source mesh was non manifold and the dynamic mesh stores a mapping */
+		bool bHasNonManifoldMapping = false;
+
+		/**
+		* mapping from the dynamic mesh vertices to the source mesh vertices (1 to 1 vertices)
+		* empty if bHasNonManifoldMapping = false
+		*/
+		TArray<int32> DynamicMeshToSource;
+
+		/**
+		* source vertex ID that correspond to the first entry of SourceToDynamicMesh
+		* this is used to save memory when only a subset of the orignal mesh is respresented by the dynamic mesh
+		*/
+		int32 SourceVertexIDOffset = 0;
+
+		/**
+		* mapping from the source mesh vertices to the dynamic mesh ones (1 to N relationship)
+		* empty if bHasNonManifoldMapping = false
+		*/
+		TArray<TArray<int32>> SourceToDynamicMesh;
+	};
+}
+
+#undef UE_API
diff --git a/Engine/Plugins/Experimental/Dataflow/Source/DataflowEditor/Private/DataflowEditorTools/DataflowMeshSelectionTool.cpp b/Engine/Plugins/Experimental/Dataflow/Source/DataflowEditor/Private/DataflowEditorTools/DataflowMeshSelectionTool.cpp
new file mode 100644
--- /dev/null
+++ b/Engine/Plugins/Experimental/Dataflow/Source/DataflowEditor/Private/DataflowEditorTools/DataflowMeshSelectionTool.cpp
@@ -0,0 +1,533 @@
+// Copyright Epic Games, Inc. All Rights Reserved.
+
+#include "DataflowEditorTools/DataflowMeshSelectionTool.h"
+
+//#include "ChaosClothAsset/ClothEditorToolBuilders.h"
+//#include "ChaosClothAsset/ClothCollectionGroup.h"
+//#include "ChaosClothAsset/ClothGeometryTools.h"
+//#include "ChaosClothAsset/ClothPatternVertexType.h"
+//#include "ChaosClothAsset/SelectionNode.h"
+//#include "ChaosClothAsset/CollectionClothFacade.h"
+//#include "ChaosClothAsset/CollectionClothSelectionFacade.h"
+//#include "ChaosClothAsset/WeightedValue.h"
+//#include "InteractiveToolManager.h"
+//#include "PreviewMesh.h"
+//#include "Selection/PolygonSelectionMechanic.h"
+//#include "ModelingToolTargetUtil.h"
+//#include "GroupTopology.h"
+//#include "ToolSetupUtil.h"
+//#include "Selections/GeometrySelection.h"
+//#include "ContextObjectStore.h"
+//#include "Dataflow/DataflowContextObject.h"
+//#include "Dataflow/DataflowGraphEditor.h"
+//#include "Dataflow/DataflowEdNode.h"
+//#include "Dataflow/DataflowObject.h"
+//#include "DynamicMesh/NonManifoldMappingSupport.h"
+//#include "Selections/GeometrySelectionUtil.h"
+//#include "Materials/Material.h"
+//#include "Selections/MeshConnectedComponents.h"
+
+#include "InteractiveToolManager.h"
+#include "Materials/MaterialInterface.h"
+#include "ModelingToolTargetUtil.h"
+#include "PreviewMesh.h"
+#include "Selection/PolygonSelectionMechanic.h"
+#include "ToolMeshSelector.h"
+#include "ToolSetupUtil.h"
+
+
+#include UE_INLINE_GENERATED_CPP_BY_NAME(DataflowMeshSelectionTool)
+
+#define LOCTEXT_NAMESPACE "DataflowMeshSelectionTool"
+
+
+///////////////////////////////////////////////////////////////////////////////////////////////////////////////////////////////////////////////////////////////////
+
+#include "Dataflow/DataflowSelectionToolNode.h"
+
+void UDataflowSelectionToolNodeDataProvider::BindToNode(FDataflowSelectionToolNode& InNode, UE::Dataflow::FContext& Context)
+{
+	WeakNode = StaticCastWeakPtr<FDataflowSelectionToolNode>(InNode.AsWeak());
+	InNode.LoadData(Context, NodeData);
+}
+
+void UDataflowSelectionToolNodeDataProvider::Init(const UE::Geometry::FDynamicMesh3& Mesh)
+{
+	MeshMapping.Init(Mesh);
+}
+
+void UDataflowSelectionToolNodeDataProvider::Shutdown()
+{
+	MeshMapping.Reset();
+}
+
+void UDataflowSelectionToolNodeDataProvider::GetVertexSelection(TSet<int32>& OutSelection) const
+{
+	using namespace UE::Dataflow;
+	TSet<int32> SelectedVertices;
+	NodeData.GetSelectedVertices(SelectedVertices);
+	MeshMapping.RemapVertices(SelectedVertices, OutSelection, FDynamicMeshMapping::ERemapDirection::SourceToDynamicMesh);
+}
+
+void UDataflowSelectionToolNodeDataProvider::SetVertexSelection(const TSet<int32>& InSelection)
+{
+	using namespace UE::Dataflow;
+	if (TSharedPtr<FDataflowSelectionToolNode> Node = WeakNode.Pin())
+	{
+		TSet<int32> SourceVerticesToSelect;
+		MeshMapping.RemapVertices(InSelection, SourceVerticesToSelect, FDynamicMeshMapping::ERemapDirection::DynamicMeshToSource);
+		const TArray<int32>& MappedSourceVertices = MeshMapping.GetMappedSourceVertices();
+		if (MappedSourceVertices.IsEmpty())
+		{
+			NodeData.DeselectAllVertices();
+		}
+		else
+		{
+			NodeData.DeselectVertices(MeshMapping.GetMappedSourceVertices());
+		}
+		NodeData.SelectVertices(SourceVerticesToSelect);
+
+		Node->SaveData(NodeData);
+	}
+}
+
+///////////////////////////////////////////////////////////////////////////////////////////////////////////////////////////////////////////////////////////////////
+
+#include "Dataflow/DataflowRenderingViewMode.h"
+#include "TargetInterfaces/PrimitiveComponentBackedTarget.h"
+#include "ToolTargetManager.h"
+#include "Dataflow/DataflowContent.h"
+#include "ToolContextInterfaces.h"
+#include "ContextObjectStore.h"
+
+void UDataflowMeshSelectionToolBuilder::GetSupportedConstructionViewModes(const UDataflowContextObject& ContextObject, TArray<const UE::Dataflow::IDataflowConstructionViewMode*>& Modes) const
+{
+	// matches UE::Dataflow::FDataflowConstruction3DViewMode::Name - but we cannot include the module to get it directly 
+	static const FName DataflowConstruction3DViewModeName = TEXT("3DView");
+	static const FName Cloth3DSimView = TEXT("Cloth3DSimView");
+	static const FName ClothRenderView = TEXT("ClothRenderView");
+
+	TArray<FName> ViewModeNames;
+	// TODO : get the views from the node 
+	ViewModeNames.Add(DataflowConstruction3DViewModeName);
+	ViewModeNames.Add(Cloth3DSimView);
+	ViewModeNames.Add(ClothRenderView);
+
+	const UE::Dataflow::FRenderingViewModeFactory& Factory = UE::Dataflow::FRenderingViewModeFactory::GetInstance();
+	for (FName ViewModeName : ViewModeNames)
+	{
+		if (const UE::Dataflow::IDataflowConstructionViewMode* ViewMode = Factory.GetViewMode(ViewModeName))
+		{
+			Modes.Add(ViewMode);
+		}
+	}
+}
+
+bool UDataflowMeshSelectionToolBuilder::CanSceneStateChange(const UInteractiveTool* ActiveTool, const FToolBuilderState& SceneState) const
+{
+	return ActiveTool->IsA<UDataflowMeshSelectionTool>();
+}
+
+void UDataflowMeshSelectionToolBuilder::SceneStateChanged(UInteractiveTool* ActiveTool, const FToolBuilderState& SceneState)
+{
+	check(CanSceneStateChange(ActiveTool, SceneState));
+
+	if (UDataflowMeshSelectionTool* const SelectionTool = Cast<UDataflowMeshSelectionTool>(ActiveTool))
+	{
+		UToolTarget* const Target = SceneState.TargetManager->BuildFirstSelectedTargetable(SceneState, GetTargetRequirements());
+		check(Target);
+		check(Target->IsValid());
+		SelectionTool->SetTarget(Target);
+		check(SelectionTool->GetTargetWorld() == SceneState.World);
+		if (UDataflowSelectionToolNodeDataProvider* DataProvider = NewObject<UDataflowSelectionToolNodeDataProvider>(SelectionTool))
+		{
+			if (UDataflowBaseContent* ContextObject = SceneState.ToolManager->GetContextObjectStore()->FindContext<UDataflowBaseContent>())
+			{
+				if (FDataflowSelectionToolNode* Node = ContextObject->GetSelectedNodeOfType<FDataflowSelectionToolNode>())
+				{
+					if (TSharedPtr<UE::Dataflow::FEngineContext> DataflowContext = ContextObject->GetDataflowContext())
+					{
+						DataProvider->BindToNode(*Node, *DataflowContext);
+						SelectionTool->DataProvider = DataProvider;
+					}
+				}
+			}
+		}
+		SelectionTool->NotifyTargetChanged();
+
+		// These are likely to be empty functions but are called here for completeness (see UInteractiveToolManager::ActivateToolInternal())
+		PostBuildTool(ActiveTool, SceneState);
+		PostSetupTool(ActiveTool, SceneState);
+	}
+
+}
+
+const FToolTargetTypeRequirements& UDataflowMeshSelectionToolBuilder::GetTargetRequirements() const
+{
+	static FToolTargetTypeRequirements TypeRequirements(UPrimitiveComponentBackedTarget::StaticClass());
+	return TypeRequirements;
+}
+
+bool UDataflowMeshSelectionToolBuilder::CanBuildTool(const FToolBuilderState& SceneState) const
+{
+	if (SceneState.SelectedComponents.Num() == 1)
+	{
+		const bool bIsADynamicMesh = SceneState.SelectedComponents[0]->IsA<UDynamicMeshComponent>();
+		return bIsADynamicMesh;
+	}
+	return false;
+	// TODO do a check on the node type
+	//if (UDataflowContextObject* const DataflowContextObject = SceneState.ToolManager->GetContextObjectStore()->FindContext<UDataflowContextObject>())
+	//{
+	//	return DataflowContextObject->GetSelectedNodeOfType<FChaosClothAssetSelectionNode_v2>() != nullptr && (SceneState.TargetManager->CountSelectedAndTargetable(SceneState, GetTargetRequirements()) == 1);
+	//}
+	//return false;
+}
+
+UInteractiveTool* UDataflowMeshSelectionToolBuilder::BuildTool(const FToolBuilderState& SceneState) const
+{
+	UDataflowMeshSelectionTool* const SelectionTool = NewObject<UDataflowMeshSelectionTool>(SceneState.ToolManager);
+
+	UToolTarget* const Target = SceneState.TargetManager->BuildFirstSelectedTargetable(SceneState, GetTargetRequirements());
+	SelectionTool->SetTarget(Target);
+	SelectionTool->SetWorld(SceneState.World);
+
+	if (UDataflowSelectionToolNodeDataProvider* DataProvider = NewObject<UDataflowSelectionToolNodeDataProvider>(SelectionTool))
+	{
+		if (UDataflowBaseContent* ContextObject = SceneState.ToolManager->GetContextObjectStore()->FindContext<UDataflowBaseContent>())
+		{
+			if (FDataflowSelectionToolNode* Node = ContextObject->GetSelectedNodeOfType<FDataflowSelectionToolNode>())
+			{
+				if (TSharedPtr<UE::Dataflow::FEngineContext> DataflowContext = ContextObject->GetDataflowContext())
+				{
+					DataProvider->BindToNode(*Node, *DataflowContext);
+					SelectionTool->DataProvider = DataProvider;
+				}
+			}
+		}
+	}
+
+	//if (UDataflowContextObject* const DataflowContextObject = SceneState.ToolManager->GetContextObjectStore()->FindContext<UDataflowContextObject>())
+	//{
+	//	NewTool->SetDataflowContextObject(DataflowContextObject);
+	//}
+
+	// TODO : need to set instanciate and set the data provider 
+	// the data provider holder a weak ref to the node and contains various cached data 
+	// it can also liste to node changes if needed 
+
+	return SelectionTool;
+}
+
+///////////////////////////////////////////////////////////////////////////////////////////////////////////////////////////////////
+
+void UDataflowMeshSelectionTool::InitializeSculptMeshFromTarget()
+{
+	using namespace UE::Geometry;
+
+	// (Re-)Create the preview mesh
+	if (PreviewMesh)
+	{
+		PreviewMesh->Disconnect();
+	}
+	
+	PreviewMesh = NewObject<UPreviewMesh>(this, TEXT("PreviewMesh"));
+	if (!PreviewMesh)
+	{
+		return;
+	}
+
+	PreviewMesh->CreateInWorld(GetTargetWorld(), FTransform::Identity);
+	ToolSetupUtil::ApplyRenderingConfigurationToPreview(PreviewMesh, Target);
+
+	// We will use the preview mesh's spatial data structure
+	PreviewMesh->bBuildSpatialDataStructure = true;
+
+	// set materials
+	FComponentMaterialSet MaterialSet = UE::ToolTarget::GetMaterialSet(Target);
+	PreviewMesh->SetMaterials(MaterialSet.Materials);
+
+	// configure secondary render material for selected triangles
+	// NOTE: the material returned by ToolSetupUtil::GetSelectionMaterial has a checkerboard pattern on back faces which makes it hard to use
+	UMaterialInterface* Material = LoadObject<UMaterialInterface>(nullptr, TEXT("/MeshModelingToolsetExp/Materials/SculptMaterial"));
+	if (Material != nullptr)
+	{
+		if (UMaterialInstanceDynamic* MatInstance = UMaterialInstanceDynamic::Create(Material, GetToolManager()))
+		{
+			MatInstance->SetVectorParameterValue(TEXT("Color"), FLinearColor::Yellow);
+			PreviewMesh->SetSecondaryRenderMaterial(MatInstance);
+		}
+	}
+
+	PreviewMesh->SetTangentsMode(EDynamicMeshComponentTangentsMode::AutoCalculated);
+	PreviewMesh->UpdatePreview(UE::ToolTarget::GetDynamicMeshCopy(Target));
+	PreviewMesh->SetVisible(true);
+
+	// Hide input target mesh
+	UE::ToolTarget::HideSourceObject(Target);
+
+	if (DataProvider && PreviewMesh->GetMesh())
+	{
+		DataProvider->Init(*PreviewMesh->GetMesh());
+	}
+}
+
+void UDataflowMeshSelectionTool::InitializeMeshSelector()
+{
+	// Setup mesh selector
+	MeshSelector = NewObject<UToolMeshSelector>(this);
+	if (MeshSelector)
+	{
+		auto OnSelectionChangedLambda = [this]() { OnSelectionModified(); };
+		MeshSelector->InitialSetup(TargetWorld.Get(), this, OnSelectionChangedLambda);
+		MeshSelector->SetMesh(PreviewMesh, FTransform::Identity);
+		MeshSelector->SetComponentSelectionMode(EComponentSelectionMode::Vertices);
+		MeshSelector->SetSelectionTool(EMeshSelectorTool::Marquee);
+		MeshSelector->SetIsEnabled(true);
+
+		if (UPolygonSelectionMechanic* SelectionMechanic = MeshSelector->GetSelectionMechanic())
+		{
+			SelectionMechanic->SetSelectionDragToolUpdateType(ESelectionDragToolUpdateType::OnTickAndRelease);
+			if (SelectionMechanic->Properties)
+			{
+				SelectionMechanic->Properties->RestoreProperties(this);
+				SelectionMechanic->Properties->bCanSelectEdges = false;
+				SelectionMechanic->Properties->bSelectVertices = true;
+				SelectionMechanic->Properties->bSelectFaces = false;
+				SelectionMechanic->Properties->bDisplaySelectionDragToolsControls = true;
+				SelectionMechanic->Properties->bDisplayPolygroupReliantControls = false;
+				AddToolPropertySource(SelectionMechanic->Properties);
+
+				// make sure we can toggle between them 
+				SelectionMechanic->Properties->WatchProperty(SelectionMechanic->Properties->bSelectVertices, [this, SelectionMechanic = SelectionMechanic](bool bSelectVertices) {
+					if (bSelectVertices && SelectionMechanic->Properties)
+					{
+						if (SelectionMechanic->Properties->bSelectFaces == true)
+						{
+							SelectionMechanic->Properties->bSelectFaces = false;
+							SelectionMechanic->SetShowSelectableCorners(true);
+							SelectionMechanic->SetShowEdges(false);
+							ConvertFaceToVertexSelection();
+						}
+					}
+					});
+				SelectionMechanic->Properties->WatchProperty(SelectionMechanic->Properties->bSelectFaces, [this, SelectionMechanic = SelectionMechanic](bool bSelectFaces) {
+					if (bSelectFaces && SelectionMechanic->Properties)
+					{
+						if (SelectionMechanic->Properties->bSelectVertices == true)
+						{
+							SelectionMechanic->Properties->bSelectVertices = false;
+							SelectionMechanic->SetShowSelectableCorners(false);
+							SelectionMechanic->SetShowEdges(true);
+							ConvertVertexToFaceSelection();
+						}
+					}
+					});
+			}
+		}
+	}
+}
+
+void UDataflowMeshSelectionTool::ConvertFaceToVertexSelection()
+{
+	if (PreviewMesh && MeshSelector)
+	{
+		if (const UE::Geometry::FDynamicMesh3* Mesh = PreviewMesh->GetMesh())
+		{
+			if (UPolygonSelectionMechanic* SelectionMechanic = MeshSelector->GetSelectionMechanic())
+			{
+				FGroupTopologySelection NewSelection;
+				const FGroupTopologySelection& Selection = SelectionMechanic->GetActiveSelection();
+				for (const int32 GroupId : Selection.SelectedGroupIDs)
+				{
+					if (Mesh->IsTriangle(GroupId))
+					{
+						UE::Geometry::FIndex3i Indices = Mesh->GetTriangle(GroupId);
+						NewSelection.SelectedCornerIDs.Add(Indices[0]);
+						NewSelection.SelectedCornerIDs.Add(Indices[1]);
+						NewSelection.SelectedCornerIDs.Add(Indices[2]);
+					}
+				}
+				SelectionMechanic->SetSelection(NewSelection);
+			}
+		}
+	}
+}
+
+void UDataflowMeshSelectionTool::ConvertVertexToFaceSelection()
+{
+	if (PreviewMesh && MeshSelector)
+	{
+		if (const UE::Geometry::FDynamicMesh3* Mesh = PreviewMesh->GetMesh())
+		{
+			if (UPolygonSelectionMechanic* SelectionMechanic = MeshSelector->GetSelectionMechanic())
+			{
+				FGroupTopologySelection NewSelection;
+				const FGroupTopologySelection& Selection = SelectionMechanic->GetActiveSelection();
+				for (const int32 TriangleId : Mesh->TriangleIndicesItr())
+				{
+					if (Mesh->IsTriangle(TriangleId))
+					{
+						UE::Geometry::FIndex3i Indices = Mesh->GetTriangle(TriangleId);
+						if (Selection.SelectedCornerIDs.Contains(Indices[0]) &&
+							Selection.SelectedCornerIDs.Contains(Indices[1]) &&
+							Selection.SelectedCornerIDs.Contains(Indices[2]))
+						{
+							NewSelection.SelectedGroupIDs.Add(TriangleId);
+						}
+					}
+				}
+				SelectionMechanic->SetSelection(NewSelection);
+			}
+		}
+	}
+}
+
+void UDataflowMeshSelectionTool::Setup()
+{
+	using namespace UE::Geometry;
+
+	InitializeSculptMeshFromTarget();
+
+	InitializeMeshSelector();
+
+	LoadFromDataProvider();
+
+	bAnyChangeMade = false;
+}
+
+void UDataflowMeshSelectionTool::OnSelectionModified()
+{
+	bAnyChangeMade = true;
+}
+
+void UDataflowMeshSelectionTool::OnShutdown(EToolShutdownType ShutdownType)
+{
+	if (ShutdownType == EToolShutdownType::Accept && CanAccept())
+	{
+		SaveToDataProvider();
+	}
+
+	if (DataProvider)
+	{
+		DataProvider->Shutdown();
+	}
+
+	if (MeshSelector)
+	{
+		if (UPolygonSelectionMechanic* SelectionMechanic = MeshSelector->GetSelectionMechanic())
+		{
+			if (SelectionMechanic->Properties)
+			{
+				SelectionMechanic->Properties->SaveProperties(this);
+			}
+		}
+
+		MeshSelector->Shutdown();
+		MeshSelector = nullptr;
+	}
+
+	if (PreviewMesh != nullptr)
+	{
+		UE::ToolTarget::ShowSourceObject(Target);
+		PreviewMesh->Disconnect();
+		PreviewMesh = nullptr;
+	}
+}
+
+void UDataflowMeshSelectionTool::LoadFromDataProvider()
+{
+	if (MeshSelector && DataProvider)
+	{
+		if (UPolygonSelectionMechanic* SelectionMechanic = MeshSelector->GetSelectionMechanic())
+		{
+			FGroupTopologySelection Selection;
+			DataProvider->GetVertexSelection(Selection.SelectedCornerIDs);
+
+			constexpr bool bBroadcastChange = false;
+			SelectionMechanic->SetSelection(Selection, bBroadcastChange);
+		}
+	}
+}
+
+void UDataflowMeshSelectionTool::SaveToDataProvider()
+{
+	if (MeshSelector && DataProvider)
+	{
+		if (UPolygonSelectionMechanic* SelectionMechanic = MeshSelector->GetSelectionMechanic())
+		{
+			GetToolManager()->BeginUndoTransaction(LOCTEXT("SelectionToolTransactionName", "Mesh Selection"));
+
+			const FGroupTopologySelection& Selection = SelectionMechanic->GetActiveSelection();
+			DataProvider->SetVertexSelection(Selection.SelectedCornerIDs);
+			GetToolManager()->EndUndoTransaction();
+		}
+	}
+}
+
+void UDataflowMeshSelectionTool::Render(IToolsContextRenderAPI* RenderAPI)
+{
+	USingleSelectionMeshEditingTool::Render(RenderAPI);
+	if (MeshSelector)
+	{
+		MeshSelector->Render(RenderAPI);
+	}
+}
+
+void UDataflowMeshSelectionTool::DrawHUD(FCanvas* Canvas, IToolsContextRenderAPI* RenderAPI)
+{
+	USingleSelectionMeshEditingTool::DrawHUD(Canvas, RenderAPI);
+	if (MeshSelector)
+	{
+		MeshSelector->DrawHUD(Canvas, RenderAPI);
+	}
+}
+
+void UDataflowMeshSelectionTool::OnTick(float DeltaTime)
+{
+	USingleSelectionMeshEditingTool::OnTick(DeltaTime);
+
+	if (MeshSelector)
+	{
+		MeshSelector->Tick(DeltaTime);
+	}
+
+	//if (bHavePendingAction)
+	//{
+	//	ApplyAction(PendingAction);
+	//	bHavePendingAction = false;
+	//	PendingAction = EClothMeshSelectionToolActions::NoAction;
+	//}
+}
+
+bool UDataflowMeshSelectionTool::CanAccept() const
+{
+	return bAnyChangeMade;
+}
+
+FBox UDataflowMeshSelectionTool::GetWorldSpaceFocusBox()
+{
+	if (MeshSelector)
+	{
+		if (UPolygonSelectionMechanic* SelectionMechanic = MeshSelector->GetSelectionMechanic())
+		{
+			static constexpr bool bWorld = true;
+			return FBox(SelectionMechanic->GetSelectionBounds(bWorld));
+		}
+	}
+	return FBox(EForceInit::ForceInitToZero);
+}
+
+
+void UDataflowMeshSelectionTool::NotifyTargetChanged()
+{
+	InitializeSculptMeshFromTarget();
+
+	if (MeshSelector)
+	{
+		MeshSelector->SetMesh(PreviewMesh, FTransform::Identity);
+	}
+	GetToolManager()->PostInvalidation();
+}
+//
+#undef LOCTEXT_NAMESPACE
diff --git a/Engine/Plugins/Experimental/Dataflow/Source/DataflowEditor/Private/DataflowEditorTools/DataflowMeshSelectionTool.h b/Engine/Plugins/Experimental/Dataflow/Source/DataflowEditor/Private/DataflowEditorTools/DataflowMeshSelectionTool.h
new file mode 100644
--- /dev/null
+++ b/Engine/Plugins/Experimental/Dataflow/Source/DataflowEditor/Private/DataflowEditorTools/DataflowMeshSelectionTool.h
@@ -0,0 +1,166 @@
+// Copyright Epic Games, Inc. All Rights Reserved.
+
+#pragma once
+
+#include "BaseTools/SingleSelectionMeshEditingTool.h"
+#include "DataflowEditorTools/DataflowDynamicMeshMapping.h"
+#include "Dataflow/DataflowSelectionToolNode.h"
+#include "UObject/ObjectPtr.h"
+
+// builder includes
+#include "DataflowEditorTools/DataflowEditorToolBuilder.h"
+#include "InteractiveToolBuilder.h"
+
+#include "DataflowMeshSelectionTool.generated.h"
+
+#define UE_API DATAFLOWEDITOR_API
+
+class UPreviewMesh;
+class UToolMeshSelector;
+
+struct FDataflowSelectionToolNode;
+
+namespace UE::Dataflow { class FContext; }
+
+UINTERFACE(MinimalAPI)
+class UDataflowMeshSelectionToolDataProviderInterface : public UInterface
+{
+	GENERATED_BODY()
+};
+
+class IDataflowMeshSelectionToolDataProviderInterface
+{
+	GENERATED_BODY()
+
+public:
+	virtual void Init(const UE::Geometry::FDynamicMesh3& Mesh) = 0;
+	virtual void Shutdown() = 0;
+
+	// Face and edge selection are infered from the vertex selection
+	virtual void GetVertexSelection(TSet<int32>& OutSelection) const = 0;
+	virtual void SetVertexSelection(const TSet<int32>& InSelection) = 0;
+};
+
+////////////////////////////////////////////////////////////////////////////////////////////////////
+
+UCLASS(MinimalAPI)
+class UDataflowSelectionToolNodeDataProvider
+	: public UObject
+	, public IDataflowMeshSelectionToolDataProviderInterface
+{
+	GENERATED_BODY()
+
+public:
+	void BindToNode(FDataflowSelectionToolNode& InNode, UE::Dataflow::FContext& Context);
+
+	//~ Begin IDataflowMeshSelectionToolDataProviderInterface 
+	virtual void Init(const UE::Geometry::FDynamicMesh3& Mesh) override;
+	virtual void Shutdown() override;
+	virtual void GetVertexSelection(TSet<int32>& OutSelection) const override;
+	virtual void SetVertexSelection(const TSet<int32>& InSelection) override;
+	//~ End IDataflowMeshSelectionToolDataProviderInterface 
+
+private:
+	TWeakPtr<FDataflowSelectionToolNode> WeakNode;
+	UE::Dataflow::FDynamicMeshMapping MeshMapping;
+	FDataflowSelectionToolNodeData NodeData;
+};
+
+////////////////////////////////////////////////////////////////////////////////////////////////////////////
+
+UCLASS(MinimalAPI)
+class UDataflowToolNodeSnapshotProvider	: public UObject
+{
+	GENERATED_BODY()
+
+public:
+	void BindToNode(FDataflowToolNode& InNode);
+
+	TConstArrayView<FDataflowToolNodeSnapshot> GetSnapshots() const;
+	void RemoveSnapshot();
+	void MoveSnapshot(int32 FromPosition, int32 ToPosition);
+
+private:
+	TWeakPtr<FDataflowToolNode> WeakNode;
+};
+
+////////////////////////////////////////////////////////////////////////////////////////////////////////////
+
+class UDataflowContextObject;
+
+namespace UE::Dataflow
+{
+	class IDataflowConstructionViewMode;
+}
+
+UCLASS(MinimalAPI)
+class UDataflowMeshSelectionToolBuilder : public UInteractiveToolWithToolTargetsBuilder, public IDataflowEditorToolBuilder
+{
+	GENERATED_BODY()
+
+private:
+	// IDataflowEditorToolBuilder
+	UE_API virtual void GetSupportedConstructionViewModes(const UDataflowContextObject& ContextObject, TArray<const UE::Dataflow::IDataflowConstructionViewMode*>& Modes) const override;
+	UE_API virtual bool CanSceneStateChange(const UInteractiveTool* ActiveTool, const FToolBuilderState& SceneState) const override;
+	UE_API virtual void SceneStateChanged(UInteractiveTool* ActiveTool, const FToolBuilderState& SceneState) override;
+
+	// UInteractiveToolWithToolTargetsBuilder
+	UE_API virtual bool CanBuildTool(const FToolBuilderState& SceneState) const override;
+	UE_API virtual UInteractiveTool* BuildTool(const FToolBuilderState& SceneState) const override;
+	UE_API virtual const FToolTargetTypeRequirements& GetTargetRequirements() const override;
+};
+
+UCLASS(MinimalAPI)
+class UDataflowMeshSelectionTool : public USingleSelectionMeshEditingTool
+{
+	GENERATED_BODY()
+
+private:
+
+	friend class UDataflowMeshSelectionToolBuilder;
+
+	//~ Begin UInteractiveTool interface
+	virtual void Setup() override;
+	virtual void OnShutdown(EToolShutdownType ShutdownType) override;
+	
+	virtual void Render(IToolsContextRenderAPI* RenderAPI) override;
+	virtual void DrawHUD(FCanvas* Canvas, IToolsContextRenderAPI* RenderAPI) override;
+	virtual void OnTick(float DeltaTime) override;
+
+	virtual bool HasCancel() const override { return true; }
+	virtual bool HasAccept() const override { return true; }
+	virtual bool CanAccept() const override;
+	//~ End UInteractiveTool interface
+
+	//~ Begin IInteractiveToolCameraFocusAPI interface
+	virtual FBox GetWorldSpaceFocusBox() override;
+	//~ End IInteractiveToolCameraFocusAPI interface
+
+	void OnSelectionModified();
+
+public:
+	UE_API void NotifyTargetChanged();
+
+	UPROPERTY()
+	TObjectPtr<UPreviewMesh> PreviewMesh = nullptr;
+
+	UPROPERTY()
+	TObjectPtr<UToolMeshSelector> MeshSelector = nullptr;
+
+	UPROPERTY(Transient)
+	TScriptInterface<IDataflowMeshSelectionToolDataProviderInterface> DataProvider = nullptr;
+
+private:
+	void InitializeSculptMeshFromTarget();
+	void InitializeMeshSelector();
+
+	void LoadFromDataProvider();
+	void SaveToDataProvider();
+
+	void ConvertFaceToVertexSelection();
+	void ConvertVertexToFaceSelection();
+
+	bool bAnyChangeMade = false;
+};
+
+#undef UE_API
diff --git a/Engine/Plugins/Experimental/Dataflow/Source/DataflowNodes/Private/Dataflow/DataflowSelectionToolNode.cpp b/Engine/Plugins/Experimental/Dataflow/Source/DataflowNodes/Private/Dataflow/DataflowSelectionToolNode.cpp
new file mode 100644
--- /dev/null
+++ b/Engine/Plugins/Experimental/Dataflow/Source/DataflowNodes/Private/Dataflow/DataflowSelectionToolNode.cpp
@@ -0,0 +1,252 @@
+// Copyright Epic Games, Inc. All Rights Reserved.
+
+#include "Dataflow/DataflowSelectionToolNode.h"
+
+#include "Dataflow/DataflowNodeFactory.h"
+#include "GeometryCollection/Facades/CollectionTransformSelectionFacade.h"
+
+#include UE_INLINE_GENERATED_CPP_BY_NAME(DataflowSelectionToolNode)
+
+#define LOCTEXT_NAMESPACE "DataflowSelectionToolNode"
+
+namespace UE::Dataflow
+{
+	void RegisterSelectionToolNode()
+	{
+		DATAFLOW_NODE_REGISTER_CREATION_FACTORY(FDataflowSelectionToolNode);
+	}
+}
+
+///////////////////////////////////////////////////////////////////////////////////////////////////////////////////
+
+void FDataflowSelectionToolNodeData::DeselectVertices(TConstArrayView<int32> Vertices)
+{
+	for (const int32 Index : Vertices)
+	{
+		if (Selection.IsValidIndex(Index))
+		{
+			Selection.SetNotSelected(Index);
+		}
+	}
+}
+
+void FDataflowSelectionToolNodeData::DeselectVertices(const TSet<int32>& Vertices)
+{
+	for (const int32 Index : Vertices)
+	{
+		if (Selection.IsValidIndex(Index))
+		{
+			Selection.SetNotSelected(Index);
+		}
+	}
+}
+
+void FDataflowSelectionToolNodeData::DeselectAllVertices()
+{
+	Selection.Clear();
+}
+
+void FDataflowSelectionToolNodeData::SelectVertices(TConstArrayView<int32> Vertices)
+{
+	for (const int32 Index : Vertices)
+	{
+		if (Selection.IsValidIndex(Index))
+		{
+			Selection.SetSelected(Index);
+		}
+	}
+}
+
+void FDataflowSelectionToolNodeData::SelectVertices(const TSet<int32>& Vertices)
+{
+	for (const int32 Index : Vertices)
+	{
+		if (Selection.IsValidIndex(Index))
+		{
+			Selection.SetSelected(Index);
+		}
+	}
+}
+
+void FDataflowSelectionToolNodeData::GetSelectedVertices(TSet<int32>& OutSelection) const
+{
+	OutSelection.Reserve(Selection.NumSelected());
+
+	const int32 NumVertices = Selection.Num();
+	for (int32 Index = 0; Index < NumVertices; ++Index)
+	{
+		if (Selection.IsSelected(Index))
+		{
+			OutSelection.Add(Index);
+		}
+	}
+}
+
+const FDataflowVertexSelection& FDataflowSelectionToolNodeData::GetVertexSelection() const
+{
+	return Selection;
+}
+
+namespace DataflowSelectionToolNodeData::Private
+{
+	static FName SelectionGroupName = TEXT("Selection");
+	static FName SelectionAttributeName = TEXT("Selection");
+}
+
+void FDataflowSelectionToolNodeData::Init(const FManagedArrayCollection& InCollection)
+{
+	Selection.InitializeFromCollection(InCollection, false);
+}
+
+void FDataflowSelectionToolNodeData::LoadFromSnapshot(const FDataflowToolNodeSnapshot& InSnapshot)
+{
+	using namespace DataflowSelectionToolNodeData::Private;
+
+	//const int32 NumVerticesInSnaphot = InSnapshot.Data.NumElements(FDataflowVertexSelection::VerticesGroupName);
+	//if (NumVerticesInSnaphot != Selection.Num())
+	//{
+	//	// TODO : transfer selection if the snapshot does not match the incoming collection
+	//}
+	//else
+	{
+		if (const TManagedArray<int32>* SelectionAttribute = InSnapshot.Data.FindAttribute<int32>(SelectionGroupName, SelectionAttributeName))
+		{
+			for (const int32 Index : *SelectionAttribute)
+			{
+				if (Selection.IsValidIndex(Index))
+				{
+					Selection.SetSelected(Index);
+				}
+			}
+		}
+	}
+}
+
+void FDataflowSelectionToolNodeData::SaveToSnapshot(FDataflowToolNodeSnapshot& OutSnapshot) const
+{
+	using namespace DataflowSelectionToolNodeData::Private;
+
+	const int32 NumVertices = Selection.Num();
+	const int32 NumSelected = Selection.NumSelected();
+
+	OutSnapshot.Name = TEXT("Selection");
+	OutSnapshot.Description = FString::Format(TEXT("{0} out of {1} vertices selected"), { FString::FormatAsNumber(NumSelected), FString::FormatAsNumber(NumVertices) });
+	OutSnapshot.Data.Reset();
+	OutSnapshot.Data.AddGroup(SelectionGroupName);
+	OutSnapshot.Data.AddElements(NumSelected, SelectionGroupName);
+	TManagedArray<int32>& SelectionAttribute = OutSnapshot.Data.AddAttribute<int32>(SelectionGroupName, SelectionAttributeName);
+	int32 SelectedIndex = 0;
+	for (int32 Index = 0; Index < NumVertices; ++Index)
+	{
+		if (Selection.IsSelected(Index) && SelectionAttribute.IsValidIndex(SelectedIndex))
+		{
+			SelectionAttribute[SelectedIndex] = Index;
+			SelectedIndex++;
+		}
+	}
+}
+
+///////////////////////////////////////////////////////////////////////////////////////////////////////////////////
+
+FDataflowSelectionToolNode::FDataflowSelectionToolNode(const UE::Dataflow::FNodeParameters& InParam, FGuid InGuid)
+	: FDataflowToolNode(InParam, InGuid)
+{
+	RegisterInputConnection(&Collection);
+	RegisterOutputConnection(&Collection, &Collection);
+	RegisterOutputConnection(&Selection);
+}
+
+void FDataflowSelectionToolNode::LoadData(UE::Dataflow::FContext& Context, FDataflowSelectionToolNodeData& OutData) const
+{
+	const FManagedArrayCollection& InCollection = GetValue(Context, &Collection);
+	OutData.Init(InCollection);
+
+	if (const FDataflowToolNodeSnapshot* LastSnaphot = GetActiveSnapshot())
+	{
+		OutData.LoadFromSnapshot(*LastSnaphot);
+	}
+}
+
+void FDataflowSelectionToolNode::SaveData(const FDataflowSelectionToolNodeData& InData)
+{
+	FDataflowToolNodeSnapshot& NewSnapshot = AddSnapshot();
+	InData.SaveToSnapshot(NewSnapshot);
+	Invalidate();
+}
+
+void FDataflowSelectionToolNode::Evaluate(UE::Dataflow::FContext& Context, const FDataflowOutput* Out) const
+{
+	if (Out->IsA(&Collection))
+	{
+		SafeForwardInput(Context, &Collection, &Collection);
+	}
+	else if (Out->IsA(&Selection))
+	{
+		FDataflowSelectionToolNodeData NodeData;
+		LoadData(Context, NodeData);
+
+		const FManagedArrayCollection& InCollection = GetValue(Context, &Collection);
+		GeometryCollection::Facades::FCollectionTransformSelectionFacade SelectionFacade(InCollection);
+
+		constexpr bool bAllElementsMustBeSelected = false;
+
+		if (const FDataflowOutput* SelectionOutput = FindOutput(&Selection))
+		{
+			if (SelectionOutput->IsType<FDataflowTransformSelection>())
+			{
+				FDataflowTransformSelection OutTransformSelection;
+				TArray<int32> TransformIndices = SelectionFacade.ConvertVertexSelectionToTransformSelection(NodeData.GetVertexSelection().AsArray(), bAllElementsMustBeSelected);
+				OutTransformSelection.InitFromArray(InCollection, TransformIndices);
+
+				SetValue(Context, OutTransformSelection, &Selection);
+			}
+			else if (SelectionOutput->IsType<FDataflowVertexSelection>())
+			{
+				SetValue(Context, NodeData.GetVertexSelection(), &Selection);
+			}
+			else if (SelectionOutput->IsType<FDataflowFaceSelection>())
+			{
+				FDataflowFaceSelection OutFaceSelection;
+				TArray<int32> FaceIndices = SelectionFacade.ConvertVertexSelectionToFaceSelection(NodeData.GetVertexSelection().AsArray(), bAllElementsMustBeSelected);
+				OutFaceSelection.InitFromArray(InCollection, FaceIndices);
+
+				SetValue(Context, OutFaceSelection, &Selection);
+			}
+			else if (SelectionOutput->IsType<FDataflowGeometrySelection>())
+			{
+				FDataflowGeometrySelection OutGeoSelection;
+				TArray<int32> GeoIndices = SelectionFacade.ConvertVertexSelectionToGeometrySelection(NodeData.GetVertexSelection().AsArray(), bAllElementsMustBeSelected);
+				OutGeoSelection.InitFromArray(InCollection, GeoIndices);
+
+				SetValue(Context, OutGeoSelection, &Selection);
+			}
+			else if (SelectionOutput->IsType<FDataflowMaterialSelection>())
+			{
+				Context.Error(LOCTEXT("MaterialSelectionNotSupported", "Material Selection output is not supported"));
+				SetValue(Context, FDataflowSelection(), &Selection);
+			}
+			else if (SelectionOutput->IsType<FDataflowCurveSelection>())
+			{
+				FDataflowCurveSelection OutCurveSelection;
+				TBitArray<> CurveSelectionArray;
+				SelectionFacade.ConvertVertexSelectionToCurveSelection(NodeData.GetVertexSelection().GetBitArray(), bAllElementsMustBeSelected, CurveSelectionArray);
+				OutCurveSelection.InitializeFromCollection(InCollection, false);
+				OutCurveSelection.SetFromBitArray(MoveTemp(CurveSelectionArray));
+
+				SetValue(Context, OutCurveSelection, &Selection);
+			}
+			else
+			{
+				Context.Error(LOCTEXT("SelectionTypeNotSupported", "The selection output type is not supported"));
+				SetValue(Context, FDataflowSelection(), &Selection);
+			}
+		}
+		else
+		{
+			Context.Error(LOCTEXT("OutputNotFound", "The selection output could not be found"));
+			SetValue(Context, FDataflowSelection(), &Selection);
+		}
+	}
+}
+
+#undef LOCTEXT_NAMESPACE
\ No newline at end of file
diff --git a/Engine/Plugins/Experimental/Dataflow/Source/DataflowNodes/Public/Dataflow/DataflowSelectionToolNode.h b/Engine/Plugins/Experimental/Dataflow/Source/DataflowNodes/Public/Dataflow/DataflowSelectionToolNode.h
new file mode 100644
--- /dev/null
+++ b/Engine/Plugins/Experimental/Dataflow/Source/DataflowNodes/Public/Dataflow/DataflowSelectionToolNode.h
@@ -0,0 +1,72 @@
+// Copyright Epic Games, Inc. All Rights Reserved.
+
+#pragma once
+
+#include "Dataflow/DataflowCollectionAddScalarVertexPropertyNode.h"
+#include "Dataflow/DataflowNode.h"
+#include "Dataflow/DataflowToolNode.h"
+#include "Dataflow/DataflowSelection.h"
+#include "GeometryCollection/ManagedArrayCollection.h"
+#include "HAL/Platform.h"
+
+#include "DataflowSelectionToolNode.generated.h"
+
+#define UE_API DATAFLOWNODES_API
+
+struct FDataflowSelectionToolNodeData
+{
+public:
+	UE_API void DeselectVertices(TConstArrayView<int32> Vertices);
+	UE_API void DeselectVertices(const TSet<int32>& Vertices);
+	UE_API void DeselectAllVertices();
+	UE_API void SelectVertices(TConstArrayView<int32> Vertices);
+	UE_API void SelectVertices(const TSet<int32>& Vertices);
+
+	UE_API void GetSelectedVertices(TSet<int32>& OutSelection) const;
+
+	UE_API const FDataflowVertexSelection& GetVertexSelection() const;
+
+private:
+	void Init(const FManagedArrayCollection& InCollection);
+	void LoadFromSnapshot(const FDataflowToolNodeSnapshot& InSnapshot);
+	void SaveToSnapshot(FDataflowToolNodeSnapshot& OutSnapshot) const;
+
+	friend struct FDataflowSelectionToolNode;
+	FDataflowVertexSelection Selection;
+};
+
+/**
+* Dataflow selection tool node
+* When selected this node opens a selection tool to select vertex or face selection 
+*/
+USTRUCT(meta=(icon="LevelEditor.StrictBoxSelect"))
+struct FDataflowSelectionToolNode : public FDataflowToolNode
+{
+	GENERATED_USTRUCT_BODY()
+	DATAFLOW_NODE_DEFINE_INTERNAL(FDataflowSelectionToolNode, "DataflowSelectionToolNode", "Tools", "Vertex Face Collection")
+	DATAFLOW_NODE_RENDER_TYPE("SurfaceRender", FGeometryCollection::StaticType(), "Collection")
+
+public:
+	FDataflowSelectionToolNode(const UE::Dataflow::FNodeParameters& InParam, FGuid InGuid = FGuid::NewGuid());
+
+	UE_API void LoadData(UE::Dataflow::FContext& Context, FDataflowSelectionToolNodeData& OutData) const;
+	UE_API void SaveData(const FDataflowSelectionToolNodeData& InData);
+
+private:
+	/** Collection to use for selecting on */
+	UPROPERTY(meta = (DataflowInput, DataflowOutput, DataflowPassthrough = Collection))
+	FManagedArrayCollection Collection;
+
+	/** Output selection - this supports many types and will convert the internal vertex selection to the concrete type of the output */
+	UPROPERTY(meta = (DataflowOutput))
+	FDataflowSelectionTypes Selection;
+
+	virtual void Evaluate(UE::Dataflow::FContext& Context, const FDataflowOutput* Out) const override;
+};
+
+namespace UE::Dataflow
+{
+	void RegisterSelectionToolNode();
+}
+
+#undef UE_API
