Dataflow editor : Node now display they underlying types in the title

#Synced-CL 51713790

diff --git a/Engine/Plugins/Experimental/Dataflow/Source/DataflowEditor/Private/Dataflow/DataflowSNode.cpp b/Engine/Plugins/Experimental/Dataflow/Source/DataflowEditor/Private/Dataflow/DataflowSNode.cpp
--- a/Engine/Plugins/Experimental/Dataflow/Source/DataflowEditor/Private/Dataflow/DataflowSNode.cpp
+++ b/Engine/Plugins/Experimental/Dataflow/Source/DataflowEditor/Private/Dataflow/DataflowSNode.cpp
@@ -21,6 +21,7 @@
 #include "Widgets/SBoxPanel.h"
 #include "Widgets/Images/SImage.h"
 #include "Widgets/Input/SButton.h"
+#include "Widgets/Text/SInlineEditableTextBlock.h"
 #include "Editor/Transactor.h"
 #include "GraphEditorSettings.h"
 #include "Dataflow/DataflowSchema.h"
@@ -28,6 +29,7 @@
 #include "Dataflow/DataflowPoints.h"
 #include "GeometryCollection/Facades/PointsFacade.h"
 
+
 #define LOCTEXT_NAMESPACE "SDataflowEdNode"
 
 namespace UE::Dataflow::Private
@@ -204,6 +206,36 @@ TSharedPtr<SGraphPin> SDataflowEdNode::CreatePinWidget(UEdGraphPin* Pin) const
 	return SGraphNode::CreatePinWidget(Pin);
 }
 
+TSharedRef<SWidget> SDataflowEdNode::CreateTitleWidget(TSharedPtr<SNodeTitle> NodeTitle)
+{
+	TSharedRef<SWidget> DefaultTitle = SGraphNode::CreateTitleWidget(NodeTitle);
+	
+	static const FName UnknownTypeName = TEXT("Unknown");
+	FName NodeType = UnknownTypeName;
+	if (TSharedPtr<FDataflowNode> DataflowNode = DataflowGraphNode->GetDataflowNode())
+	{
+		NodeType = DataflowNode->GetDisplayName();
+	}
+	return SNew(SVerticalBox)
+		+ SVerticalBox::Slot()
+		[
+			DefaultTitle
+		]
+		+ SVerticalBox::Slot()
+		[
+			SNew(STextBlock)
+			.Text(FText::FromName(NodeType))
+			.Font(FCoreStyle::Get().GetFontStyle("SmallFont"))
+			.Visibility_Lambda([InlineEditableText = InlineEditableText]()
+				{
+					return InlineEditableText->IsInEditMode()
+						? EVisibility::Collapsed
+						: EVisibility::Visible;
+				})
+		]
+		;
+}
+
 TSharedRef<SWidget> SDataflowEdNode::CreateTitleRightWidget()
 {
 	RenderCheckBoxWidget = SNew(SCheckBox)
diff --git a/Engine/Plugins/Experimental/Dataflow/Source/DataflowEditor/Public/Dataflow/DataflowSNode.h b/Engine/Plugins/Experimental/Dataflow/Source/DataflowEditor/Public/Dataflow/DataflowSNode.h
--- a/Engine/Plugins/Experimental/Dataflow/Source/DataflowEditor/Public/Dataflow/DataflowSNode.h
+++ b/Engine/Plugins/Experimental/Dataflow/Source/DataflowEditor/Public/Dataflow/DataflowSNode.h
@@ -70,6 +70,7 @@ private:
 	UE_API const FSlateBrush* GetPinButtonImage() const;
 
 	UE_API virtual TSharedPtr<SGraphPin> CreatePinWidget(UEdGraphPin* Pin) const override;
+	UE_API virtual TSharedRef<SWidget> CreateTitleWidget(TSharedPtr<SNodeTitle> NodeTitle);
 	UE_API virtual TSharedRef<SWidget> CreateTitleRightWidget() override;
 
 	TObjectPtr<UDataflowEdNode> DataflowGraphNode = nullptr;	
