Dataflow Editor : Imporved the tool user experience
* Construction viewport now can display a message when the tool fail to load or if it requires user action
* Made the editor mode more aware of the actually selection
* Fixed placement of newly created tool nodes when created from the tool palette buttons

#Synced-CL 51755560

diff --git a/Engine/Plugins/Experimental/Dataflow/Source/DataflowEditor/Private/Dataflow/DataflowConstructionViewport.cpp b/Engine/Plugins/Experimental/Dataflow/Source/DataflowEditor/Private/Dataflow/DataflowConstructionViewport.cpp
--- a/Engine/Plugins/Experimental/Dataflow/Source/DataflowEditor/Private/Dataflow/DataflowConstructionViewport.cpp
+++ b/Engine/Plugins/Experimental/Dataflow/Source/DataflowEditor/Private/Dataflow/DataflowConstructionViewport.cpp
@@ -334,6 +334,24 @@ void SDataflowConstructionViewport::PopulateViewportOverlays(TSharedRef<SOverlay
 				]
 		];
 
+	Overlay->AddSlot()
+		.VAlign(VAlign_Top)
+		.HAlign(HAlign_Fill)
+		[
+			SNew(SBorder)
+			.BorderBackgroundColor(FColor::Orange)
+			.BorderImage(FCoreStyle::Get().GetBrush("ErrorReporting.Box"))
+			.HAlign(HAlign_Center)
+			.VAlign(VAlign_Center)
+			.Padding(FMargin(8, 8))
+			.Visibility(this, &SDataflowConstructionViewport::GetEditorModeMessageVisibility)
+			[
+				SNew(STextBlock)
+				.ColorAndOpacity(FColor::Black)
+				.Text(this, &SDataflowConstructionViewport::GetEditorModeMessage)
+			]
+		];
+
 	// this widget will display the current viewed feature level
 	Overlay->AddSlot()
 		.VAlign(VAlign_Bottom)
@@ -344,6 +362,31 @@ void SDataflowConstructionViewport::PopulateViewportOverlays(TSharedRef<SOverlay
 		];
 }
 
+bool SDataflowConstructionViewport::HasEditorModeMessage() const
+{
+	if (const UDataflowEditorMode* const DataflowEdMode = GetEdMode())
+	{
+		return DataflowEdMode->HasUserMessage();
+	}
+	return false;
+}
+
+EVisibility SDataflowConstructionViewport::GetEditorModeMessageVisibility() const
+{
+	return HasEditorModeMessage()
+		? EVisibility::Visible
+		: EVisibility::Collapsed;
+}
+
+FText SDataflowConstructionViewport::GetEditorModeMessage() const
+{
+	if (const UDataflowEditorMode* const DataflowEdMode = GetEdMode())
+	{
+		return DataflowEdMode->GetUserMessage();
+	}
+	return {};
+}
+
 FText SDataflowConstructionViewport::GetOverlayText() const
 {
 	if (DataflowClient)
diff --git a/Engine/Plugins/Experimental/Dataflow/Source/DataflowEditor/Private/Dataflow/DataflowEditorMode.cpp b/Engine/Plugins/Experimental/Dataflow/Source/DataflowEditor/Private/Dataflow/DataflowEditorMode.cpp
--- a/Engine/Plugins/Experimental/Dataflow/Source/DataflowEditor/Private/Dataflow/DataflowEditorMode.cpp
+++ b/Engine/Plugins/Experimental/Dataflow/Source/DataflowEditor/Private/Dataflow/DataflowEditorMode.cpp
@@ -247,6 +247,34 @@ bool UDataflowEditorMode::CanUseTools() const
 	return IsDataflowEditable();
 }
 
+TSharedPtr<const FUICommandInfo> UDataflowEditorMode::GetSelectedNodeToolCommand() const
+{
+	if (CanUseTools())
+	{
+		const FName NodeTypeName = GetSelectedNodeTypeName();
+		if (!NodeTypeName.IsNone())
+		{
+			const TSharedRef<FUICommandList> CommandList = Toolkit->GetToolkitCommands();
+			const FDataflowEditorCommandsImpl& CommandInfos = FDataflowEditorCommandsImpl::Get();
+			return NodeTypeToToolCommandMap.FindRef(NodeTypeName, {});
+		}
+	}
+	return {};
+}
+
+bool UDataflowEditorMode::CanSelectedNodeLaunchTool() const
+{
+	if (CanUseTools())
+	{
+		const bool bIsToolActive = GetToolManager()->HasActiveTool(EToolSide::Left);
+		if (!bIsToolActive && !bPendingNodeSelectionChanged)
+		{
+			return (UDataflowEditorMode::GetSelectedNodeToolCommand() != nullptr);
+		}
+	}
+	return false;
+}
+
 bool UDataflowEditorMode::IsDataflowEditable() const
 {
 	if (TSharedPtr<SDataflowGraphEditor> GraphEditor = DataflowGraphEditor.Pin())
@@ -687,6 +715,15 @@ void UDataflowEditorMode::InitializeTargets(const TArray<TObjectPtr<UObject>>& O
 	}
 }
 
+TObjectPtr<UDataflowBaseContent> UDataflowEditorMode::GetDataflowEditorContent() const
+{
+	if (const TSharedPtr<FDataflowConstructionScene> ConstructionScenePtr = ConstructionScene.Pin())
+	{
+		return ConstructionScenePtr ? ConstructionScenePtr->GetEditorContent().Get() : nullptr;
+	}
+	return nullptr;
+}
+
 void UDataflowEditorMode::ModeTick(float DeltaTime)
 {
 	Super::ModeTick(DeltaTime);
@@ -713,21 +750,36 @@ void UDataflowEditorMode::ModeTick(float DeltaTime)
 		FirstTimeFocusConstructionViewport();
 	}
 
-	if (CanUseTools())
+	TObjectPtr<UDataflowEdNode> SelectedNode = GetSelectedNode();
+
+	if (CanSelectedNodeLaunchTool())
 	{
-		if (!NodeTypeForPendingToolStart.IsNone() && !GetToolManager()->HasActiveTool(EToolSide::Left) && !bPendingNodeSelectionChanged)
+		if (TSharedPtr<const FUICommandInfo> Command = UDataflowEditorMode::GetSelectedNodeToolCommand())
 		{
-			const TSharedRef<FUICommandList> CommandList = Toolkit->GetToolkitCommands();
-			const FDataflowEditorCommandsImpl& CommandInfos = FDataflowEditorCommandsImpl::Get();
-
-			if (const TSharedPtr<const FUICommandInfo>* const Command = NodeTypeToToolCommandMap.Find(NodeTypeForPendingToolStart))
+			const bool bLaunchedTool = Toolkit->GetToolkitCommands()->TryExecuteAction(Command.ToSharedRef());
+			if (!bLaunchedTool)
+			{
+				const USelection* Selection = ConstructionScenePtr ? ConstructionScenePtr->GetSelectedComponents(): nullptr;
+				if (Selection && Selection->Num() == 0)
+				{
+					// TODO : find a more accurate message ?
+					UserMessage = FText::Format(LOCTEXT("ToolNeedoneComponent", "You must select a component to launch the {0} tool"), Command->GetLabel());
+				}
+				else
+				{
+					UserMessage = FText::Format(LOCTEXT("ToolFailedToLaunch", "Failed to launch {0} tool - Check Editor logs"), Command->GetLabel());
+				}
+			}
+			else
 			{
-				CommandList->TryExecuteAction(Command->ToSharedRef());
+				UserMessage = FText::GetEmpty();
 			}
-
-			NodeTypeForPendingToolStart = FName();
 		}
 	}
+	else
+	{
+		UserMessage = FText::GetEmpty();
+	}
 
 	if (bShouldRestartToolNextTick)
 	{
@@ -748,13 +800,9 @@ void UDataflowEditorMode::ModeTick(float DeltaTime)
 			}
 
 			// Now start the tool
-			if (const TSharedPtr<const SDataflowGraphEditor> PinnedGraphEditor = DataflowGraphEditor.Pin())
+			if (SelectedNode)
 			{
-				const FGraphPanelSelectionSet& SelectedNodes = PinnedGraphEditor->GetSelectedNodes();
-				if (SelectedNodes.Num() == 1)
-				{
-					StartToolForSelectedNode(*SelectedNodes.CreateConstIterator());
-				}
+				StartToolForSelectedNode(SelectedNode);
 			}
 
 			bShouldRestartToolNextTick = false;
@@ -947,28 +995,21 @@ void UDataflowEditorMode::SetConstructionViewMode(const FName& NewViewModeName)
 		}
 	}
 
-
-	if (const TSharedPtr<const SDataflowGraphEditor> PinnedDataflowGraphEditor = DataflowGraphEditor.Pin())
+	if (const TObjectPtr<const UDataflowEdNode> SelectedDataflowEdNode = GetSelectedNode())
 	{
-		if (const UEdGraphNode* const SelectedNode = PinnedDataflowGraphEditor->GetSingleSelectedNode())
+		if (const TObjectPtr<UDataflowBaseContent>& EditorContent = GetDataflowEditorContent())
 		{
-			if (const UDataflowEdNode* const SelectedDataflowEdNode = Cast<UDataflowEdNode>(SelectedNode))
+			TArray<FName> ViewModesForNode;
+			UE::Dataflow::GetViewModesForNode(*SelectedDataflowEdNode, *EditorContent, ViewModesForNode);
+
+			if (ViewModesForNode.Contains(NewViewModeName))
 			{
-				if (const TObjectPtr<UDataflowBaseContent>& EditorContent = ConstructionScenePtr ? ConstructionScenePtr->GetEditorContent() : nullptr)
+				ViewModesForNode.Sort([](const FName& A, const FName& B)
 				{
-					TArray<FName> ViewModesForNode;
-					UE::Dataflow::GetViewModesForNode(*SelectedDataflowEdNode, *EditorContent, ViewModesForNode);
+					return A.FastLess(B);
+				});
 
-					if (ViewModesForNode.Contains(NewViewModeName))
-					{
-						ViewModesForNode.Sort([](const FName& A, const FName& B)
-						{
-							return A.FastLess(B);
-						});
-
-						NodeViewModeHistory.FindOrAdd(ViewModesForNode) = NewViewModeName;
-					}
-				}
+				NodeViewModeHistory.FindOrAdd(ViewModesForNode) = NewViewModeName;
 			}
 		}
 	}
@@ -1046,21 +1087,15 @@ bool UDataflowEditorMode::CanChangeConstructionViewModeTo(const FName& NewViewMo
 	{
 		if (const TSharedPtr<FDataflowConstructionScene> ConstructionScenePtr = ConstructionScene.Pin())
 		{
-			if (const TObjectPtr<UDataflowBaseContent>& EditorContent = ConstructionScenePtr->GetEditorContent())
+			if (const TObjectPtr<UDataflowBaseContent> EditorContent = GetDataflowEditorContent())
 			{
-				if (const TSharedPtr<const SDataflowGraphEditor> PinnedDataflowGraphEditor = DataflowGraphEditor.Pin())
+				if (const TObjectPtr<UDataflowEdNode> SelectedDataflowEdNode = EditorContent->GetSelectedNode())
 				{
-					if (const UEdGraphNode* const SelectedNode = PinnedDataflowGraphEditor->GetSingleSelectedNode())
+					if (const UE::Dataflow::IDataflowConstructionViewMode* const ViewMode = UE::Dataflow::FRenderingViewModeFactory::GetInstance().GetViewMode(NewViewModeName))
 					{
-						if (const UDataflowEdNode* const SelectedDataflowEdNode = Cast<UDataflowEdNode>(SelectedNode))
+						if (UE::Dataflow::IsViewModeSupported(*SelectedDataflowEdNode, *EditorContent, *ViewMode))
 						{
-							if (const UE::Dataflow::IDataflowConstructionViewMode* const ViewMode = UE::Dataflow::FRenderingViewModeFactory::GetInstance().GetViewMode(NewViewModeName))
-							{
-								if (UE::Dataflow::IsViewModeSupported(*SelectedDataflowEdNode, *EditorContent, *ViewMode))
-								{
-									return true;
-								}
-							}
+							return true;
 						}
 					}
 				}
@@ -1178,6 +1213,27 @@ void UDataflowEditorMode::SetDataflowGraphEditor(TSharedPtr<SDataflowGraphEditor
 	}
 }
 
+TObjectPtr<UDataflowEdNode> UDataflowEditorMode::GetSelectedNode() const
+{
+	if (TObjectPtr<UDataflowBaseContent> EditorContent = GetDataflowEditorContent())
+	{
+		return EditorContent->GetSelectedNode();
+	}
+	return nullptr;
+}
+
+FName UDataflowEditorMode::GetSelectedNodeTypeName() const
+{
+	if (TObjectPtr<UDataflowEdNode> SelectedNode = GetSelectedNode())
+	{
+		if (const TSharedPtr<const FDataflowNode> DataflowNode = SelectedNode->GetDataflowNode())
+		{
+			return DataflowNode->GetType();
+		}
+	}
+	return NAME_None;
+}
+
 void UDataflowEditorMode::StartToolForSelectedNode(const UObject* SelectedNode)
 {
 	if (const UDataflowEdNode* const EdNode = Cast<UDataflowEdNode>(SelectedNode))
@@ -1219,64 +1275,44 @@ void UDataflowEditorMode::OnDataflowNodeDeleted(const TSet<UObject*>& DeletedNod
 
 bool UDataflowEditorMode::IsAnyNodeSelected() const
 {
-	const TSharedPtr<const SDataflowGraphEditor> PinnedDataflowGraphEditor = DataflowGraphEditor.Pin();
-	if (!PinnedDataflowGraphEditor)
-	{
-		return false;
-	}
-
-	UEdGraphNode* const SelectedNode = PinnedDataflowGraphEditor->GetSingleSelectedNode();
-	return (SelectedNode != nullptr);
+	return (GetSelectedNode() != nullptr);
 }
 
 UEdGraphNode* UDataflowEditorMode::GetSingleSelectedNodeWithOutputType(const FName& SelectedNodeOutputTypeName) const
 {
-	const TSharedPtr<const SDataflowGraphEditor> PinnedDataflowGraphEditor = DataflowGraphEditor.Pin();
-	if (!PinnedDataflowGraphEditor)
-	{
-		return nullptr;
-	}
-
-	UEdGraphNode* const SelectedNode = PinnedDataflowGraphEditor->GetSingleSelectedNode();
-	if (!SelectedNode)
+	if (TObjectPtr<UDataflowEdNode> SelectedDataflowEdNode = GetSelectedNode())
 	{
-		return nullptr;
+		if (TSharedPtr<const FDataflowNode> SelectedDataflowNode = SelectedDataflowEdNode->GetDataflowNode())
+		{
+			for (const FDataflowOutput* const Output : SelectedDataflowNode->GetOutputs())
+			{
+				if (Output->GetType() == SelectedNodeOutputTypeName)
+				{
+					return SelectedDataflowEdNode;
+				}
+			}
+		}
 	}
-
-	 if (const UDataflowEdNode* const SelectedDataflowEdNode = Cast<UDataflowEdNode>(SelectedNode))
-	 {
-		 if (TSharedPtr<const FDataflowNode> SelectedDataflowNode = SelectedDataflowEdNode->GetDataflowNode())
-		 {
-			 for (const FDataflowOutput* const Output : SelectedDataflowNode->GetOutputs())
-			 {
-				 if (Output->GetType() == SelectedNodeOutputTypeName)
-				 {
-					 return SelectedNode;
-				 }
-			 }
-		 }
-	 }
-
 	return nullptr;
 }
 
 UEdGraphNode* UDataflowEditorMode::CreateNewNode(const FName& NewNodeTypeName)
 {
-	const TSharedPtr<const SDataflowGraphEditor> PinnedDataflowGraphEditor = DataflowGraphEditor.Pin();
-	if (!PinnedDataflowGraphEditor)
-	{
-		return nullptr;
-	}
-	const TSharedPtr<FDataflowConstructionScene> ConstructionScenePtr = ConstructionScene.Pin();
-	if (const TObjectPtr<UDataflowBaseContent>& EditorContent = ConstructionScenePtr ? ConstructionScenePtr->GetEditorContent() : nullptr)
+	if (const TObjectPtr<UDataflowBaseContent> EditorContent = GetDataflowEditorContent())
 	{
-		if (const TObjectPtr<UDataflow>& DataflowGraph = EditorContent->GetDataflowAsset())
+		if (TObjectPtr<const UDataflowEdNode> SelectedNode = EditorContent->GetSelectedNode())
 		{
-			const TSharedPtr<FAssetSchemaAction_Dataflow_CreateNode_DataflowEdNode> NodeAction =
-				FAssetSchemaAction_Dataflow_CreateNode_DataflowEdNode::CreateAction(DataflowGraph, NewNodeTypeName);
-			constexpr UEdGraphPin* FromPin = nullptr;
-			constexpr bool bSelectNewNode = true;
-			return NodeAction->PerformAction(DataflowGraph, FromPin, PinnedDataflowGraphEditor->GetPasteLocation2f(), bSelectNewNode);
+			if (const TObjectPtr<UDataflow>& DataflowGraph = EditorContent->GetDataflowAsset())
+			{
+				const TSharedPtr<FAssetSchemaAction_Dataflow_CreateNode_DataflowEdNode> NodeAction =
+					FAssetSchemaAction_Dataflow_CreateNode_DataflowEdNode::CreateAction(DataflowGraph, NewNodeTypeName);
+				constexpr UEdGraphPin* FromPin = nullptr;
+				constexpr bool bSelectNewNode = true;
+
+				const FVector2f Location = SelectedNode->GetPosition() + FVector2f(250, 0);
+
+				return NodeAction->PerformAction(DataflowGraph, FromPin, Location, bSelectNewNode);
+			}
 		}
 	}
 	return nullptr;
diff --git a/Engine/Plugins/Experimental/Dataflow/Source/DataflowEditor/Public/Dataflow/DataflowConstructionViewport.h b/Engine/Plugins/Experimental/Dataflow/Source/DataflowEditor/Public/Dataflow/DataflowConstructionViewport.h
--- a/Engine/Plugins/Experimental/Dataflow/Source/DataflowEditor/Public/Dataflow/DataflowConstructionViewport.h
+++ b/Engine/Plugins/Experimental/Dataflow/Source/DataflowEditor/Public/Dataflow/DataflowConstructionViewport.h
@@ -3,11 +3,13 @@
 
 #include "SBaseCharacterFXEditorViewport.h"
 #include "SCommonEditorViewportToolbarBase.h"
+#include "Widgets/Notifications/SErrorText.h"
 
 class UDataflowEditorMode;
 class FDataflowConstructionScene;
 class FDataflowConstructionViewportClient;
 class SRichTextBlock;
+class SErrorText;
 
 struct FToolMenuSection;
 
@@ -38,6 +40,9 @@ public:
 	virtual void PopulateViewportOverlays(TSharedRef<SOverlay> Overlay) override;
 
 	FText GetOverlayText() const;
+	bool HasEditorModeMessage() const;
+	FText GetEditorModeMessage() const;
+	EVisibility GetEditorModeMessageVisibility() const;
 	void FocusViewportToSelection(bool bInstant);
 
 	TSharedPtr<FDataflowConstructionViewportClient> GetDataflowClient() const { return DataflowClient; }
@@ -52,6 +57,7 @@ private:
 
 	/** Pointer to the box into which the overlay text items are added */
 	TSharedPtr<SRichTextBlock> OverlayText;
+	TSharedPtr<SErrorText> MessageWidget;
 
 	UDataflowEditorMode* GetEdMode() const;
 };
diff --git a/Engine/Plugins/Experimental/Dataflow/Source/DataflowEditor/Public/Dataflow/DataflowEditorMode.h b/Engine/Plugins/Experimental/Dataflow/Source/DataflowEditor/Public/Dataflow/DataflowEditorMode.h
--- a/Engine/Plugins/Experimental/Dataflow/Source/DataflowEditor/Public/Dataflow/DataflowEditorMode.h
+++ b/Engine/Plugins/Experimental/Dataflow/Source/DataflowEditor/Public/Dataflow/DataflowEditorMode.h
@@ -12,6 +12,7 @@ class SDataflowGraphEditor;
 class FDataflowConstructionViewportClient;
 class FDataflowSimulationScene;
 class FDataflowSimulationViewportClient;
+class UDataflowBaseContent;
 class UDataflowEdNode;
 class UEdGraphNode;
 class FViewport;
@@ -88,7 +89,6 @@ public:
 		return ActiveToolsContext;
 	}
 
-
 	UE_API void StartToolForSelectedNode(const UObject* SelectedNode);
 	UE_API void ShutdownActiveToolIfNeeded(EToolShutdownType ShutdownType);
 
@@ -101,6 +101,11 @@ public:
 	UE_API const FName* FindRecentlyUsedViewMode(const TArray<FName>& AvailableViewModes) const;
 
 	UE_API bool IsAnyNodeSelected() const;
+	UE_API TObjectPtr<UDataflowEdNode> GetSelectedNode() const;
+	UE_API FName GetSelectedNodeTypeName() const;
+
+	UE_API FText GetUserMessage() const { return UserMessage; }
+	UE_API bool HasUserMessage() const { return !UserMessage.IsEmpty(); }
 
 private:
 	friend class FDataflowEditorToolkit;
@@ -190,8 +195,12 @@ private:
 	UE_API const TArray<FName>& GetToolCategories();
 
 	bool CanUseTools() const;
+	bool CanSelectedNodeLaunchTool() const;
+	TSharedPtr<const FUICommandInfo> GetSelectedNodeToolCommand() const;
 	bool IsDataflowEditable() const;
 
+	TObjectPtr<UDataflowBaseContent> GetDataflowEditorContent() const;
+
 	/** Dataflow construction scene from the toolkit */
 	TWeakPtr<FDataflowConstructionScene> ConstructionScene;
 
@@ -254,6 +263,9 @@ private:
 	// Timestamp for telemetry
 	FDateTime LastModeStartTimestamp;
 
+	// Message to display to the user in the viewport
+	FText UserMessage;
+
 	// Each node has a set of view modes it can use. For each distinct set of available view modes, we store the last mode that the user has chosen.
 	// This will enable switching between different kinds of nodes and not have the view mode reset to a default every time.
 	TMap<TArray<FName>, FName> NodeViewModeHistory;
