Dataflow : update the dataflow Tool node change to also include the active snapshot

#Synced-CL 51755560

#Synced-CL 51755560

diff --git a/Engine/Plugins/Experimental/Dataflow/Source/DataflowNodes/Private/Dataflow/DataflowToolNode.cpp b/Engine/Plugins/Experimental/Dataflow/Source/DataflowNodes/Private/Dataflow/DataflowToolNode.cpp
--- a/Engine/Plugins/Experimental/Dataflow/Source/DataflowNodes/Private/Dataflow/DataflowToolNode.cpp
+++ b/Engine/Plugins/Experimental/Dataflow/Source/DataflowNodes/Private/Dataflow/DataflowToolNode.cpp
@@ -84,11 +84,6 @@ int32 FDataflowToolNodeSnapshotSet::FindOldestUnlockedSnapShot()
 	return OldestSnapshotIndex;
 }
 
-void FDataflowToolNodeSnapshotSet::SwapSnapshots(TArray<FDataflowToolNodeSnapshot>& InOutSnapshots)
-{
-	Swap(InOutSnapshots, Snapshots);
-}
-
 //////////////////////////////////////////////////////////////////////////////////////////////////////////////////////////
 
 FDataflowToolNode::FDataflowToolNode(const UE::Dataflow::FNodeParameters& InParam, FGuid InGuid)
@@ -105,9 +100,9 @@ bool FDataflowToolNode::RemoveSnapshot(int32 Index)
 	return Snapshots.RemoveSnapshot(Index);
 }
 
-void FDataflowToolNode::SwapSnapshots(TArray<FDataflowToolNodeSnapshot>& InOutSnapshots)
+void FDataflowToolNode::SwapSnapshots(FDataflowToolNodeSnapshotSet& InOutSnapshots)
 {
-	Snapshots.SwapSnapshots(InOutSnapshots);
+	Swap(Snapshots, InOutSnapshots);
 }
 
 TConstArrayView<FDataflowToolNodeSnapshot> FDataflowToolNode::GetSnapshots() const
@@ -124,7 +119,7 @@ const FDataflowToolNodeSnapshot* FDataflowToolNode::GetActiveSnapshot() const
 
 FDataflowToolNode::FSnapshotToolChange::FSnapshotToolChange::FSnapshotToolChange(FDataflowToolNode& Node)
 	: WeakNode(StaticCastWeakPtr<FDataflowToolNode>(Node.AsWeak()))
-	, SavedSnapshots(Node.GetSnapshots())
+	, SavedSnapshots(Node.Snapshots)
 {}
 
 FString FDataflowToolNode::FSnapshotToolChange::ToString() const
diff --git a/Engine/Plugins/Experimental/Dataflow/Source/DataflowNodes/Public/Dataflow/DataflowToolNode.h b/Engine/Plugins/Experimental/Dataflow/Source/DataflowNodes/Public/Dataflow/DataflowToolNode.h
--- a/Engine/Plugins/Experimental/Dataflow/Source/DataflowNodes/Public/Dataflow/DataflowToolNode.h
+++ b/Engine/Plugins/Experimental/Dataflow/Source/DataflowNodes/Public/Dataflow/DataflowToolNode.h
@@ -66,8 +66,6 @@ public:
 	UE_API bool RemoveSnapshot(int32 Index);
 	UE_API int32 FindOldestUnlockedSnapShot();
 
-	UE_API void SwapSnapshots(TArray<FDataflowToolNodeSnapshot>& InOutSnapshots);
-
 private:
 	UPROPERTY(EditAnywhere, Category = Snapshots)
 	int32 ActiveSnapshot = INDEX_NONE;
@@ -111,11 +109,11 @@ protected:
 
 	private:
 		TWeakPtr<FDataflowToolNode> WeakNode;
-		TArray<FDataflowToolNodeSnapshot> SavedSnapshots;
+		FDataflowToolNodeSnapshotSet SavedSnapshots;
 	};
 
 private:
-	UE_API void SwapSnapshots(TArray<FDataflowToolNodeSnapshot>& InOutSnapshots);
+	UE_API void SwapSnapshots(FDataflowToolNodeSnapshotSet& InOutSnapshots);
 
 	UPROPERTY(EditAnywhere, Category = Snapshots)
 	FDataflowToolNodeSnapshotSet Snapshots;
