Dataflow : add misisng regsitration of selection tool node and sattic mesh termninal

#Synced-CL 51755560

diff --git a/Engine/Plugins/Experimental/Dataflow/Source/DataflowNodes/Private/Dataflow/DataflowNodesPlugin.cpp b/Engine/Plugins/Experimental/Dataflow/Source/DataflowNodes/Private/Dataflow/DataflowNodesPlugin.cpp
--- a/Engine/Plugins/Experimental/Dataflow/Source/DataflowNodes/Private/Dataflow/DataflowNodesPlugin.cpp
+++ b/Engine/Plugins/Experimental/Dataflow/Source/DataflowNodes/Private/Dataflow/DataflowNodesPlugin.cpp
@@ -6,6 +6,7 @@
 #include "Dataflow/DataflowCollectionAttributeKeyNodes.h"
 #include "Dataflow/DataflowSkeletalMeshNodes.h"
 #include "Dataflow/DataflowStaticMeshNodes.h"
+#include "Dataflow/DataflowStaticMeshAssetNodes.h"
 #include "Dataflow/DataflowNodeFactory.h"
 #include "Dataflow/DataflowSelectionNodes.h"
 #include "Dataflow/DataflowContextOverridesNodes.h"
@@ -20,9 +21,10 @@
 #include "Dataflow/DataflowSamplerToAttributeNode.h"
 #include "Dataflow/DataflowGetCollectionBoundsNode.h"
 #include "Dataflow/SamplerNodes/DataflowSamplerTypes.h"
+#include "Dataflow/DataflowSelectionToolNode.h"
 #include "Dataflow/SamplerNodes/DataflowGradientSamplerNode.h"
 #include "Dataflow/SamplerNodes/DataflowTransformSamplerNode.h"
-#include "Dataflow/SamplerNodes/DataflowSamplerTypes.h"
+#include "Dataflow/Customization/DataflowToolNodeSnapshotCustomization.h"
 
 #define LOCTEXT_NAMESPACE "DataflowNodes"
 
@@ -56,7 +58,9 @@ void IDataflowNodesPlugin::StartupModule()
 {
 	UE::Dataflow::RegisterSkeletalMeshNodes();
 	UE::Dataflow::RegisterStaticMeshNodes();
+	UE::Dataflow::RegisterStaticMeshAssetNodes();
 	UE::Dataflow::RegisterSelectionNodes();
+	UE::Dataflow::RegisterSelectionToolNode();
 	UE::Dataflow::RegisterContextOverridesNodes();
 	UE::Dataflow::RegisterGetCollectionBoundsNode();
 
@@ -79,10 +83,26 @@ void IDataflowNodesPlugin::StartupModule()
 	UE::Dataflow::RegisterNodeFilter(FDataflowTerminalNode::StaticType());
 
 	FDataflowAddScalarVertexPropertyCallbackRegistry::Get().RegisterCallbacks(MakeUnique<FGeometryCollectionAddScalarVertexPropertyCallbacks>());
+
+#if WITH_EDITOR
+	if (FPropertyEditorModule* const PropertyModule = FModuleManager::GetModulePtr<FPropertyEditorModule>("PropertyEditor"))
+	{
+		PropertyModule->RegisterCustomPropertyTypeLayout(FDataflowToolNodeSnapshot::StaticStruct()->GetFName(), FOnGetPropertyTypeCustomizationInstance::CreateStatic(&FDataflowToolNodeSnapshotCustomization::MakeInstance));
+		PropertyModule->RegisterCustomPropertyTypeLayout(FDataflowToolNodeSnapshotSet::StaticStruct()->GetFName(), FOnGetPropertyTypeCustomizationInstance::CreateStatic(&FDataflowToolNodeSnapshotSetCustomization::MakeInstance));
+	}
+#endif
 }
 
 void IDataflowNodesPlugin::ShutdownModule()
 {
+#if WITH_EDITOR
+	if (FPropertyEditorModule* const PropertyModule = FModuleManager::GetModulePtr<FPropertyEditorModule>("PropertyEditor"))
+	{
+		PropertyModule->UnregisterCustomPropertyTypeLayout("DataflowToolNodeSnapshot");
+		PropertyModule->UnregisterCustomPropertyTypeLayout("DataflowToolNodeSnapshotSet");
+	}
+#endif
+
 	FDataflowAddScalarVertexPropertyCallbackRegistry::Get().DeregisterCallbacks(FGeometryCollectionAddScalarVertexPropertyCallbacks::Name);
 }
 
