Dataflow editor : fixing offset in Dataflow replay of caches
- this is because Chaos cache manager increment by Dt before recording and playing  but dataflow does not on replay
- fixing the chaos cache manager would have negative effect on existing caches
- fixing dataflow to properly account for that

#Synced-CL 51755560

diff --git a/Engine/Plugins/Experimental/Dataflow/Source/DataflowEditor/Private/Dataflow/DataflowSimulationScene.cpp b/Engine/Plugins/Experimental/Dataflow/Source/DataflowEditor/Private/Dataflow/DataflowSimulationScene.cpp
--- a/Engine/Plugins/Experimental/Dataflow/Source/DataflowEditor/Private/Dataflow/DataflowSimulationScene.cpp
+++ b/Engine/Plugins/Experimental/Dataflow/Source/DataflowEditor/Private/Dataflow/DataflowSimulationScene.cpp
@@ -381,7 +381,10 @@ void FDataflowSimulationScene::UpdateSimulationCache()
 			if(bIsSimulationLocked)
 			{
 				// Update the cached simulation at some point in time
-				const float ScrollBarTime = SimulationTime - TimeRange[0];
+				// We need to add delta time because the CHaos Cache Actor incremenet by Dt before recording the first frame
+				// The chaos cache manager replays are compensating with this issue and fixing it would break existing caches
+				// so we nee dto fix how we replay in Dataflow
+				const float ScrollBarTime = SimulationTime - TimeRange[0] + DeltaTime;
 				if (ScrollBarTime != CacheManager->StartTime)
 				{
 					CacheManager->SetStartTime(ScrollBarTime);
