Dataflow - Clear the embedded Dataflow object standalone and public flags on creation.

#rnx
#Synced-CL 51757399

diff --git a/Engine/Plugins/Experimental/Dataflow/Source/DataflowEditor/Private/Dataflow/AssetDefinition_DataflowAsset.cpp b/Engine/Plugins/Experimental/Dataflow/Source/DataflowEditor/Private/Dataflow/AssetDefinition_DataflowAsset.cpp
--- a/Engine/Plugins/Experimental/Dataflow/Source/DataflowEditor/Private/Dataflow/AssetDefinition_DataflowAsset.cpp
+++ b/Engine/Plugins/Experimental/Dataflow/Source/DataflowEditor/Private/Dataflow/AssetDefinition_DataflowAsset.cpp
@@ -414,6 +414,9 @@ namespace UE::DataflowAssetDefinitionHelpers
 			// Duplicate it and embed it into the new asset
 			if (UDataflow* const Dataflow = DuplicateObject(DataflowTemplate, NewAsset, TEXT("EmbeddedDataflow")))
 			{
+				// Clear standalone asset flags
+				Dataflow->ClearFlags(RF_Public | RF_Standalone);
+
 				// Mark as transactional
 				Dataflow->SetFlags(RF_Transactional);
 
