Dataflow Editor : add support for displaying face selection  on geometry collection

#Synced-CL 51755560

diff --git a/Engine/Plugins/Experimental/Dataflow/Source/DataflowEditor/Private/DataflowRendering/DataflowGeometryCollectionRenderableType.cpp b/Engine/Plugins/Experimental/Dataflow/Source/DataflowEditor/Private/DataflowRendering/DataflowGeometryCollectionRenderableType.cpp
--- a/Engine/Plugins/Experimental/Dataflow/Source/DataflowEditor/Private/DataflowRendering/DataflowGeometryCollectionRenderableType.cpp
+++ b/Engine/Plugins/Experimental/Dataflow/Source/DataflowEditor/Private/DataflowRendering/DataflowGeometryCollectionRenderableType.cpp
@@ -5,6 +5,7 @@
 #include "Components/DynamicMeshComponent.h"
 #include "Drawing/PointSetComponent.h"
 #include "Drawing/LineSetComponent.h"
+#include "Drawing/TriangleSetComponent.h"
 
 #include "Dataflow/DataflowRenderingViewMode.h"
 #include "DataflowRendering/DataflowRenderableComponents.h"
@@ -13,6 +14,9 @@
 #include "DataflowRendering/DataflowRenderableTypeRegistry.h"
 #include "DataflowRendering/DataflowRenderableTypeUtils.h"
 
+#include "DynamicMesh/NonManifoldMappingSupport.h"
+#include "DynamicMesh/DynamicMeshBulkEdit.h"
+
 #include "GeometryCollection/Facades/CollectionMeshFacade.h"
 #include "GeometryCollection/GeometryCollectionUtility.h"
 #include "GeometryCollection/GeometryCollectionAlgo.h"
@@ -30,9 +34,15 @@
 #include "PlanarCut.h" // for ConvertGeometryCollectionToDynamicMesh
 #include "Dataflow/DataflowEditorStyle.h"
 
-
 namespace GeometryCollectionSurfaceRenderableType::Private
 {
+	bool bFastDynamicMeshGeneration = false;
+	FAutoConsoleVariableRef CVarFastDynamicMeshGeneration(
+		TEXT("p.Dataflow.FastDynamicMeshGeneration"),
+		bFastDynamicMeshGeneration,
+		TEXT("When on, use a faster way to generate dynamic meshes - very experimental for now")
+	);
+
 	static const TCHAR* DefaultDataflowSelectionMaterialName = TEXT("/Engine/EditorMaterials/Dataflow/DataflowSelection");
 
 	static void AddLeavesToSelection(int32 TransformIndex, const TManagedArray<TSet<int32>>& ChildrenAttribute, FDataflowSelection& OutSelection)
@@ -65,6 +75,148 @@ namespace GeometryCollectionSurfaceRenderableType::Private
 			}
 		}
 	}
+
+	static void GeometryCollectionToDynamicMesh(const FManagedArrayCollection& Collection, int32 TransformIndex, FDynamicMesh3& OutMesh)
+	{
+		GeometryCollection::Facades::FCollectionMeshFacade MeshFacade(Collection);
+		if (!MeshFacade.IsValid())
+		{
+			return; // invalid facade - not a geometry collection
+		}
+
+		const int32 GeometryIndex = MeshFacade.TransformToGeometryIndexAttribute.IsValidIndex(TransformIndex)
+			? MeshFacade.TransformToGeometryIndexAttribute[TransformIndex]
+			: INDEX_NONE;
+
+		if (GeometryIndex == INDEX_NONE)
+		{
+			return; // no geometry to convert 
+		}
+
+		const int32 VtxStart = MeshFacade.VertexStartAttribute.IsValidIndex(GeometryIndex)
+			? MeshFacade.VertexStartAttribute[GeometryIndex]
+			: 0;
+		const int32 VtxCount = MeshFacade.VertexCountAttribute.IsValidIndex(GeometryIndex)
+			? MeshFacade.VertexCountAttribute[GeometryIndex]
+			: 0;
+
+		if (VtxCount == 0)
+		{
+			return; // no vertices
+		}
+
+		const int32 TriStart = MeshFacade.FaceStartAttribute.IsValidIndex(GeometryIndex)
+			? MeshFacade.FaceStartAttribute[GeometryIndex]
+			: 0;
+		const int32 TriCount = MeshFacade.FaceCountAttribute.IsValidIndex(GeometryIndex)
+			? MeshFacade.FaceCountAttribute[GeometryIndex]
+			: 0;
+
+		if (TriCount == 0)
+		{
+			return; // no triangles
+		}
+
+		int32 VisibleTriCount = TriCount;
+		if (MeshFacade.VisibleAttribute.Num() >= TriStart + TriCount)
+		{
+			int32 NumVisible = 0;
+			for (int32 Index = 0; Index < TriCount; ++Index)
+			{
+				if (MeshFacade.VisibleAttribute.Get()[TriStart + Index])
+				{
+					++NumVisible;
+				}
+			}
+			VisibleTriCount = NumVisible;
+		}
+
+		if (VisibleTriCount == 0)
+		{
+			return; // no visible triangles
+		}
+
+		TArrayView<const FIntVector> InTriangles = MeshFacade.GetTriangles(TransformIndex);
+
+		UE::Geometry::FDynamicMeshBulkEdit BulkEdit(OutMesh);
+
+		// Edge count is left to zero so it will be autogenerated ( faster than providing it in our case ) 
+		BulkEdit.Initialize(VtxCount, VisibleTriCount, /*EdgeCount*/0);
+
+		// Vertices
+		constexpr int VertexBatchSize = 1000;
+		TArrayView<const FVector3f> InPositions = MeshFacade.GetVertexPositions(TransformIndex);
+		ParallelFor(TEXT("CopyCollectionVerticesToDynamicMesh"), InPositions.Num(), VertexBatchSize,
+			[&BulkEdit, &InPositions](int32 VertexIndex)
+			{
+				BulkEdit.SetVertex(VertexIndex, FVector(InPositions[VertexIndex]));
+			});
+
+		// Triangles - because of visible tri attribute we cannot parralelize that to get the right index
+		int32 TriangleToAdd = 0;
+		for (int32 TriangleIndex = 0; TriangleIndex < TriCount; ++TriangleIndex)
+		{
+			const UE::Geometry::FIndex3i LocalTriangle(
+				InTriangles[TriangleIndex].X - VtxStart,
+				InTriangles[TriangleIndex].Y - VtxStart,
+				InTriangles[TriangleIndex].Z - VtxStart
+			);
+			const int32 GlobalTriIndex = TriStart + TriangleIndex;
+			bool bIsVisible = true;
+			if (MeshFacade.VisibleAttribute.IsValidIndex(GlobalTriIndex))
+			{
+				bIsVisible = MeshFacade.VisibleAttribute.Get()[GlobalTriIndex];
+			}
+			if (bIsVisible)
+			{
+				BulkEdit.SetTriVertexIndices(TriangleToAdd, LocalTriangle);
+				++TriangleToAdd;
+			}
+		};
+		
+		OutMesh.EnableAttributes();
+
+		TArray<UE::Tasks::FTask> OtherTasks;
+
+		// Normals
+		if (MeshFacade.NormalAttribute.Num() >= (VtxStart + VtxCount))
+		{
+			OtherTasks.Add(
+				UE::Tasks::Launch(TEXT("CopyCollectionNormalsToDynamicMesh"),
+					[&OutMesh, &VtxStart, &VtxCount, &InNormals = MeshFacade.NormalAttribute]()
+					{
+						OutMesh.EnableVertexNormals(FVector3f(FLinearColor::White));
+						for (int32 VertexIndex = 0; VertexIndex < VtxCount; ++VertexIndex)
+						{
+							OutMesh.SetVertexNormal(VertexIndex, InNormals[VtxStart+VertexIndex]);
+						}
+					}));
+		}
+
+		// Colors
+		if (MeshFacade.ColorAttribute.Num() >= (VtxStart + VtxCount))
+		{
+			OtherTasks.Add(
+				UE::Tasks::Launch(TEXT("CopyCollectionNormalsToDynamicMesh"),
+					[&OutMesh, &VtxStart, &VtxCount, &InColors = MeshFacade.ColorAttribute]()
+					{
+						OutMesh.Attributes()->EnablePrimaryColors();
+						if (UE::Geometry::FDynamicMeshColorOverlay* ColorOverlay = OutMesh.Attributes()->PrimaryColors())
+						{
+							ColorOverlay->CreatePerVertex(0.f);
+							for (int32 VertexIndex = 0; VertexIndex < VtxCount; ++VertexIndex)
+							{
+								ColorOverlay->SetElement(VertexIndex, FVector4f(InColors[VtxStart+VertexIndex]));
+							}
+						}
+					}));
+		}
+
+		BulkEdit.Finalize();
+
+		UE::Tasks::Wait(OtherTasks);
+
+	}
 }
 
 ////////////////////////////////////////////////////////////////////////////////////////////////////////////////////////////////
@@ -130,7 +282,7 @@ namespace UE::Dataflow::Private
 			FDataflowFaceSelection FaceSelection;
 			bool bValidFaceSelection = false;
 
-			FDataflowFaceSelection CurveSelection;
+			FDataflowCurveSelection CurveSelection;
 			bool bValidCurveSelection = false;
 		};
 
@@ -211,12 +363,28 @@ namespace UE::Dataflow::Private
 						ParallelFor(DynamicMeshes.Num(),
 							[&DynamicMeshes, &BoneTransforms, &GeomCollection](int32 TransformIndex)
 							{
-								FDynamicMesh3& DynamicMesh = DynamicMeshes[TransformIndex];
-								TArrayView<const int32> TransformIndicesToConvert(&TransformIndex, 1);
-
-								FTransform UnusedTransform;
-								constexpr bool bWeldEdges = false; // disabling this as this can take 43% of the conversion time for large meshes
-								::ConvertGeometryCollectionToDynamicMesh(DynamicMesh, UnusedTransform, false, *GeomCollection, bWeldEdges, BoneTransforms, true, TransformIndicesToConvert);
+								if (GeomCollection && GeomCollection->TransformToGeometryIndex.IsValidIndex(TransformIndex))
+								{
+									const int32 GeometryIndex = GeomCollection->TransformToGeometryIndex[TransformIndex];
+									FDynamicMesh3& DynamicMesh = DynamicMeshes[TransformIndex];
+
+									if (GeometryCollectionSurfaceRenderableType::Private::bFastDynamicMeshGeneration)
+									{
+										GeometryCollectionSurfaceRenderableType::Private::GeometryCollectionToDynamicMesh(*GeomCollection, TransformIndex, DynamicMesh);
+									}
+									else
+									{
+										TArrayView<const int32> TransformIndicesToConvert(&TransformIndex, 1);
+
+										FTransform UnusedTransform;
+										constexpr bool bWeldEdges = false; // disabling this as this can take 43% of the conversion time for large meshes
+										//constexpr bool bAllowInvisible = false;
+										//constexpr bool bCenterPivot = false;
+
+										::ConvertGeometryCollectionToDynamicMesh(DynamicMesh, UnusedTransform, false, *GeomCollection, bWeldEdges, BoneTransforms, true, TransformIndicesToConvert);
+										//ConvertGeometryCollectionGeometryToDynamicMesh(*GeomCollection, GeometryIndex, bAllowInvisible, bCenterPivot, DynamicMesh, UnusedTransform);
+									}
+								}
 							});
 					}
 					Instance.CacheValue(DynamicMeshes, MeshArrayCacheName);
@@ -298,6 +466,76 @@ namespace UE::Dataflow::Private
 					}
 				}
 
+				if (Selection.bValidFaceSelection)
+				{
+					const TManagedArray<FVector3f>* VertexAttribute = Collection.FindAttribute<FVector3f>(FGeometryCollection::VertexPositionAttribute, FGeometryCollection::VerticesGroup);
+					const TManagedArray<FVector3f>* NormalAttribute = Collection.FindAttribute<FVector3f>(FGeometryCollection::VertexNormalAttribute, FGeometryCollection::VerticesGroup);
+					const TManagedArray<int32>* VertexToTransformIndex = Collection.FindAttribute<int32>(FGeometryCollection::VertexBoneMapAttribute, FGeometryCollection::VerticesGroup);
+					const TManagedArray<FIntVector>* TriangleAttribute = Collection.FindAttribute<FIntVector>(FGeometryCollection::FaceIndicesAttribute, FGeometryCollection::FacesGroup);
+
+					if (VertexAttribute && NormalAttribute && TriangleAttribute)
+					{
+						const int32 NumSelected = Selection.FaceSelection.NumSelected();
+						if (NumSelected > 0 && Selection.FaceSelection.Num() == TriangleAttribute->Num())
+						{
+							const GeometryCollection::Facades::FCollectionTransformFacade TransformFacade(Collection);
+							const TArray<FTransform> CollectionSpaceTransforms = TransformFacade.ComputeCollectionSpaceTransforms();
+
+							const FName TriangleComponentName = Instance.GetComponentName(TEXT("FaceSelection"));
+							if (UTriangleSetComponent* TriangleComponent = OutComponents.AddNewComponent<UTriangleSetComponent>(TriangleComponentName))
+							{
+								TriangleComponent->AddTriangles(
+									1 /*NumIndices*/, // number of time the lambda will be called
+									[&VertexAttribute, &NormalAttribute, &VertexToTransformIndex, &TriangleAttribute, &Selection, &CollectionSpaceTransforms, &SelectionMaterial]
+									(int32 Iteration, TArray<FRenderableTriangle>& TrianglesOut)
+									{
+										auto GetPointFromIndex = [&VertexAttribute, &VertexToTransformIndex, &CollectionSpaceTransforms](int32 VtxIdx)
+											{
+												if (VertexAttribute->IsValidIndex(VtxIdx))
+												{
+													if (VertexToTransformIndex && VertexToTransformIndex->IsValidIndex(VtxIdx) && CollectionSpaceTransforms.IsValidIndex((*VertexToTransformIndex)[VtxIdx]))
+													{
+														const FTransform Transform = CollectionSpaceTransforms[(*VertexToTransformIndex)[VtxIdx]];
+														return Transform.TransformPosition(FVector((*VertexAttribute)[VtxIdx]));
+													}
+													return FVector((*VertexAttribute)[VtxIdx]);
+												}
+												return FVector::ZeroVector;
+											};
+
+										for (int32 TriIdx = 0; TriIdx < TriangleAttribute->Num(); ++TriIdx)
+										{
+											if (Selection.FaceSelection.IsSelected(TriIdx))
+											{
+												const FIntVector& Tri = (*TriangleAttribute)[TriIdx];
+												const FVector Vtx0 = GetPointFromIndex(Tri[0]);
+												const FVector Vtx1 = GetPointFromIndex(Tri[1]);
+												const FVector Vtx2 = GetPointFromIndex(Tri[2]);
+												
+												FRenderableTriangle TriangleToAdd;
+												TriangleToAdd.Material = SelectionMaterial;
+												TriangleToAdd.Vertex0.Position = Vtx0;
+												TriangleToAdd.Vertex0.Normal = FVector((*NormalAttribute)[Tri[0]]);
+												TriangleToAdd.Vertex0.Color = FColor::White;
+												TriangleToAdd.Vertex1.Position = Vtx1;
+												TriangleToAdd.Vertex1.Color = FColor::White;
+												TriangleToAdd.Vertex1.Normal = FVector((*NormalAttribute)[Tri[1]]);
+												TriangleToAdd.Vertex2.Position = Vtx2;
+												TriangleToAdd.Vertex2.Color = FColor::White;
+												TriangleToAdd.Vertex2.Normal = FVector((*NormalAttribute)[Tri[2]]);
+												TrianglesOut.Add(TriangleToAdd);
+											}
+										}
+									},
+									NumSelected /*TrianglesPerIndexHint*/
+								);
+								TriangleComponent->SetMaterial(0, SelectionMaterial);
+								TriangleComponent->SetOverlayMaterial(SelectionMaterial);
+							}
+						}
+					}
+				}
+
 				// Create line component from curve selection 
 				if (Selection.bValidCurveSelection)
 				{
