Dataflow : add the concept of Reference asset to a dataflow asset
- when reference asset is set, and the dataflow is shared , the dataflow asset wil only be editable by in the context of the reference asset and all other asset will have the graph grayed out
- also deprecate unused properties in the dataflow asset class

#Synced-CL 51967769

diff --git a/Engine/Plugins/Experimental/Dataflow/Source/DataflowEditor/Private/Dataflow/DataflowEditorToolkit.cpp b/Engine/Plugins/Experimental/Dataflow/Source/DataflowEditor/Private/Dataflow/DataflowEditorToolkit.cpp
--- a/Engine/Plugins/Experimental/Dataflow/Source/DataflowEditor/Private/Dataflow/DataflowEditorToolkit.cpp
+++ b/Engine/Plugins/Experimental/Dataflow/Source/DataflowEditor/Private/Dataflow/DataflowEditorToolkit.cpp
@@ -1229,7 +1229,7 @@ void FDataflowEditorToolkit::AddDataflowActionWidget(FToolMenuSection& InSection
 						{ 
 							if (TSharedPtr<FDataflowEditorToolkit> ToolkitPtr = WeakToolkit.Pin())
 							{
-								return ToolkitPtr->bAllowExternalDataflowEditing ? UnlockedText : LockedText;
+								return ToolkitPtr->AllowExternalDataflowEditing() ? UnlockedText : LockedText;
 							}
 							return UnlockedText;
 						}),
@@ -1237,7 +1237,7 @@ void FDataflowEditorToolkit::AddDataflowActionWidget(FToolMenuSection& InSection
 						{ 
 							if (TSharedPtr<FDataflowEditorToolkit> ToolkitPtr = WeakToolkit.Pin())
 							{
-								return ToolkitPtr->bAllowExternalDataflowEditing ? UnlockedTooltip : LockedTooltip;
+								return ToolkitPtr->AllowExternalDataflowEditing() ? UnlockedTooltip : LockedTooltip;
 							}
 							return UnlockedTooltip;
 						}),
@@ -1245,13 +1245,52 @@ void FDataflowEditorToolkit::AddDataflowActionWidget(FToolMenuSection& InSection
 						{ 
 							if (TSharedPtr<FDataflowEditorToolkit> ToolkitPtr = WeakToolkit.Pin())
 							{
-								return ToolkitPtr->bAllowExternalDataflowEditing ? UnlockedIcon : LockedIcon;
+								return ToolkitPtr->AllowExternalDataflowEditing() ? UnlockedIcon : LockedIcon;
 							}
 							return UnlockedIcon;
 						})
 				));
 			}
 		}));
+
+	// Open reference asset button 
+	InSection.AddDynamicEntry("DataflowActionOpenReferenceAsset", FNewToolMenuSectionDelegate::CreateLambda([](FToolMenuSection& DynamicSection)
+		{
+			using namespace UE::Dataflow::Private;
+			if (const TSharedPtr<FDataflowEditorToolkit> Toolkit = GetAssetEditorToolkitFromMenuContext(DynamicSection.FindContext<UAssetEditorToolkitMenuContext>()))
+			{
+				TWeakPtr<FDataflowEditorToolkit> WeakToolkit = Toolkit;
+
+				DynamicSection.AddEntry(FToolMenuEntry::InitToolBarButton(
+					TEXT("DataflowOpenReferenceAsset"),
+					FUIAction(
+						FExecuteAction::CreateLambda([WeakToolkit]()
+							{
+								if (TSharedPtr<FDataflowEditorToolkit> ToolkitPtr = WeakToolkit.Pin())
+								{
+									if (UObject* ReferenceAsset = ToolkitPtr->GetReferenceAsset())
+									{
+										GEditor->GetEditorSubsystem<UAssetEditorSubsystem>()->OpenEditorForAsset(ReferenceAsset);
+									}
+								}
+							}
+						),
+						FCanExecuteAction::CreateLambda([WeakToolkit]()
+							{
+								if (TSharedPtr<FDataflowEditorToolkit> ToolkitPtr = WeakToolkit.Pin())
+								{
+									return !ToolkitPtr->IsReferenceAssetNullOrSameAsBoundAsset();
+								}
+								return false;
+							}
+						)
+					),
+					TAttribute<FText>(LOCTEXT("OpenReferenceAssetLabel", "Open Reference Asset")),
+					LOCTEXT("OpenReferenceAssetTooltip", "Open reference asset in an editor"),
+					FSlateIcon(FAppStyle::GetAppStyleSetName(), "GraphEditor.EventGraph_16x")
+				));
+			}
+		}));
 }
 
 bool FDataflowEditorToolkit::IsEmbeddedDataflow() const
@@ -1271,14 +1310,47 @@ bool FDataflowEditorToolkit::IsExternalDataflow() const
 	return !IsEmbeddedDataflow();
 }
 
+UObject* FDataflowEditorToolkit::GetReferenceAsset() const
+{
+	if (const TObjectPtr<UDataflowBaseContent>& EditorContent = GetEditorContent())
+	{
+		if (const UDataflow* DataflowAsset = EditorContent->GetDataflowAsset())
+		{
+			return DataflowAsset->ReferenceAsset;
+		}
+	}
+	return nullptr;
+}
+
+bool FDataflowEditorToolkit::IsReferenceAssetNullOrSameAsBoundAsset() const
+{
+	UObject* ReferenceAsset = nullptr;
+	bool bIsReferenceAssetNullOrSameAsBoundAsset = true;
+	if (const TObjectPtr<UDataflowBaseContent>& EditorContent = GetEditorContent())
+	{
+		if (const UDataflow* DataflowAsset = EditorContent->GetDataflowAsset())
+		{
+			ReferenceAsset = DataflowAsset->ReferenceAsset;
+		}
+		TObjectPtr<UObject> BoundAsset = EditorContent->GetDataflowOwner();
+		bIsReferenceAssetNullOrSameAsBoundAsset = (ReferenceAsset == nullptr) || (BoundAsset == ReferenceAsset);
+	}
+	return bIsReferenceAssetNullOrSameAsBoundAsset;
+}
+
+bool FDataflowEditorToolkit::AllowExternalDataflowEditing() const
+{
+	return IsExternalDataflow() && bAllowExternalDataflowEditing && IsReferenceAssetNullOrSameAsBoundAsset();
+}
+
 bool FDataflowEditorToolkit::IsEditable() const
 {
-	return (IsEmbeddedDataflow() || bAllowExternalDataflowEditing || bDataflowForceGraphEditable);
+	return (IsEmbeddedDataflow() || AllowExternalDataflowEditing() || bDataflowForceGraphEditable);
 }
 
 bool FDataflowEditorToolkit::IsEditingLockButtonVisible() const
 {
-	return IsExternalDataflow() && !bDataflowForceGraphEditable;
+	return IsExternalDataflow() && IsReferenceAssetNullOrSameAsBoundAsset() && !bDataflowForceGraphEditable;
 }
 
 void FDataflowEditorToolkit::ConvertToEmbeddedDataflow()
@@ -1317,7 +1389,7 @@ void FDataflowEditorToolkit::ConvertToExternalDataflow()
 
 void FDataflowEditorToolkit::ToggleDataflowEditingLock()
 {
-	if (IsExternalDataflow())
+	if (IsExternalDataflow() && IsReferenceAssetNullOrSameAsBoundAsset())
 	{
 		bAllowExternalDataflowEditing = !bAllowExternalDataflowEditing;
 	}
diff --git a/Engine/Plugins/Experimental/Dataflow/Source/DataflowEditor/Public/Dataflow/DataflowEditorToolkit.h b/Engine/Plugins/Experimental/Dataflow/Source/DataflowEditor/Public/Dataflow/DataflowEditorToolkit.h
--- a/Engine/Plugins/Experimental/Dataflow/Source/DataflowEditor/Public/Dataflow/DataflowEditorToolkit.h
+++ b/Engine/Plugins/Experimental/Dataflow/Source/DataflowEditor/Public/Dataflow/DataflowEditorToolkit.h
@@ -349,6 +349,15 @@ private:
 	/** return true if the Dataflow is a separate asset or no dataflow is assigned to the owner asset */
 	bool IsExternalDataflow() const;
 
+	/** return true if the bound asset is the reference asset or if the reference asset is null */
+	bool IsReferenceAssetNullOrSameAsBoundAsset() const;
+
+	/** Get Reference asset from the dataflow graph */
+	UObject* GetReferenceAsset() const;
+
+	/** return true the current dataflow can be edited - condition are must be External and match reference asset (if reference asset is defined)*/
+	bool AllowExternalDataflowEditing() const;
+
 	/** whether the Lock/Unlock editing button is visible or not */
 	bool IsEditingLockButtonVisible() const;
 
diff --git a/Engine/Source/Runtime/Experimental/Dataflow/Engine/Public/Dataflow/DataflowObject.h b/Engine/Source/Runtime/Experimental/Dataflow/Engine/Public/Dataflow/DataflowObject.h
--- a/Engine/Source/Runtime/Experimental/Dataflow/Engine/Public/Dataflow/DataflowObject.h
+++ b/Engine/Source/Runtime/Experimental/Dataflow/Engine/Public/Dataflow/DataflowObject.h
@@ -145,10 +145,12 @@ public:
 	//~ End IDataflowGraphInterface interface
 
 public:
-	UPROPERTY(EditAnywhere, Category = "Evaluation")
+	UE_DEPRECATED(5.8, "bActive property is deprecated as it is no longer used")
+	UPROPERTY(meta = (DeprecatedProperty, DeprecationMessage = "bActive property is deprecated as it is no longer used"))
 	bool bActive = true;
 
-	UPROPERTY(EditAnywhere, Category = "Evaluation", AdvancedDisplay )
+	UE_DEPRECATED(5.8, "Targets property is deprecated as it is no longer used")
+	UPROPERTY(meta = (DeprecatedProperty, DeprecationMessage = "Targets property is deprecated as it is no longer used"))
 	TArray<TObjectPtr<UObject>> Targets;
 
 	UPROPERTY(EditAnywhere, Category = "Render")
@@ -157,9 +159,23 @@ public:
 	UPROPERTY(EditAnywhere, Category = "Evaluation",meta=(EditConditionHides))
 	EDataflowType Type = EDataflowType::Construction;
 
+	/**
+	* Dataflow variables
+	* Variables are used by the graph to parametrized the graph
+	* They can can default values and can be overriden on a per asset basis ( see Dataflow Instance )
+	*/
 	UPROPERTY()
 	FInstancedPropertyBag Variables;
 
+	/** 
+	* Reference asset to use for this Dataflow 
+	* The reference asset is used when opening the dataflow asset editor 
+	* If the dataflow is shared among multiple assets, and the reference asset is set, only this asset will be able to change the dataflow
+	* (for the other the dataflow graph will be read-only and only evaluation and overrides will be available)
+	*/
+	UPROPERTY(EditAnywhere, Category = Reference)
+	TObjectPtr<UObject> ReferenceAsset = nullptr;
+
 public:
 	/** UObject Interface */
 	static DATAFLOWENGINE_API void AddReferencedObjects(UObject* InThis, FReferenceCollector& Collector);
