Dataflow - Fixed object flags on newly un-embedded Dataflow graphs so that they show in the Content and can be saved.

This problem is a consequence of fixing the template flags for embedded assets at CL 51821154.

#rnx
#Synced-CL 51757399

diff --git a/Engine/Plugins/Experimental/Dataflow/Source/DataflowEditor/Private/Dataflow/AssetDefinition_DataflowAsset.cpp b/Engine/Plugins/Experimental/Dataflow/Source/DataflowEditor/Private/Dataflow/AssetDefinition_DataflowAsset.cpp
--- a/Engine/Plugins/Experimental/Dataflow/Source/DataflowEditor/Private/Dataflow/AssetDefinition_DataflowAsset.cpp
+++ b/Engine/Plugins/Experimental/Dataflow/Source/DataflowEditor/Private/Dataflow/AssetDefinition_DataflowAsset.cpp
@@ -117,7 +117,13 @@ namespace UE::DataflowAssetDefinitionHelpers
 			BoundAsset,
 			[&DataflowObject](UPackage* const NewPackage, FName NewAssetName)
 			{
-				return DuplicateObject<UDataflow>(DataflowObject, NewPackage, NewAssetName);
+				UDataflow* const DuplicateDataflowObject = DuplicateObject<UDataflow>(DataflowObject, NewPackage, NewAssetName);
+				if (DuplicateDataflowObject)
+				{
+					// Adjust flags, otherwise the bound Dataflow copy won't be showing in the package
+					DuplicateDataflowObject->SetFlags(RF_Public | RF_Standalone);
+				}
+				return DuplicateDataflowObject;
 			});
 
 		if (OutDataflowAsset)
