Dataflow editor : add icons to the node list

#Synced-CL 51997199

diff --git a/Engine/Plugins/Experimental/Dataflow/Source/DataflowEditor/Private/Dataflow/DataflowGraphEditor.cpp b/Engine/Plugins/Experimental/Dataflow/Source/DataflowEditor/Private/Dataflow/DataflowGraphEditor.cpp
--- a/Engine/Plugins/Experimental/Dataflow/Source/DataflowEditor/Private/Dataflow/DataflowGraphEditor.cpp
+++ b/Engine/Plugins/Experimental/Dataflow/Source/DataflowEditor/Private/Dataflow/DataflowGraphEditor.cpp
@@ -17,6 +17,7 @@
 #include "Dataflow/DataflowSchema.h"
 #include "Dataflow/DataflowSubGraph.h"
 #include "Dataflow/DataflowVariableNodes.h"
+#include "Dataflow/SDataflowGraphActionWidget.h"
 #include "HAL/PlatformApplicationMisc.h"
 #include "ScopedTransaction.h"
 #include "IStructureDetailsView.h"
@@ -381,7 +382,11 @@ FActionMenuContent SDataflowGraphEditor::OnCreateActionMenu(UEdGraph* Graph, con
 		.DraggedFromPins(DraggedPins)
 		.AutoExpandActionMenu(bAutoExpandActionMenu)
 		.OnClosedCallback(OnClosed)
-		;
+		.OnCreateWidgetForAction(SGraphActionMenu::FOnCreateWidgetForAction::CreateLambda(
+			[](const FCreateWidgetForActionData* const Data)
+			{
+				return SNew(SDataflowGraphActionWidget, Data);
+			}));
 
 	TSharedRef<SWidget> ActionMenuWithOptions =
 		SNew(SVerticalBox)
diff --git a/Engine/Plugins/Experimental/Dataflow/Source/DataflowEditor/Private/Dataflow/DataflowGraphSchemaAction.cpp b/Engine/Plugins/Experimental/Dataflow/Source/DataflowEditor/Private/Dataflow/DataflowGraphSchemaAction.cpp
--- a/Engine/Plugins/Experimental/Dataflow/Source/DataflowEditor/Private/Dataflow/DataflowGraphSchemaAction.cpp
+++ b/Engine/Plugins/Experimental/Dataflow/Source/DataflowEditor/Private/Dataflow/DataflowGraphSchemaAction.cpp
@@ -43,8 +43,18 @@ TSharedPtr<FAssetSchemaAction_Dataflow_CreateNode_DataflowEdNode> FAssetSchemaAc
 
 					const FText Category = FText::FromString(Param.Category.ToString().IsEmpty() ? FString("Dataflow") : Param.Category.ToString());
 					const FText Tags = FText::FromString(Param.Tags);
+
+					FName IconName;
+					if (Param.DefaultNodeObject)
+					{
+						if (const FString* IconNameStr = Param.DefaultNodeObject->TypedScriptStruct()->FindMetaData("Icon"))
+						{
+							IconName = FName(*IconNameStr);
+						}
+					}
+
 					TSharedPtr<FAssetSchemaAction_Dataflow_CreateNode_DataflowEdNode> NewNodeAction(
-						new FAssetSchemaAction_Dataflow_CreateNode_DataflowEdNode(NodeName, InNodeTypeName, Category, MenuDesc, ToolTip, Tags));
+						new FAssetSchemaAction_Dataflow_CreateNode_DataflowEdNode(NodeName, InNodeTypeName, Category, MenuDesc, ToolTip, Tags, IconName));
 					return NewNodeAction;
 				}
 			}
@@ -53,6 +63,14 @@ TSharedPtr<FAssetSchemaAction_Dataflow_CreateNode_DataflowEdNode> FAssetSchemaAc
 	return TSharedPtr<FAssetSchemaAction_Dataflow_CreateNode_DataflowEdNode>(nullptr);
 }
 
+FAssetSchemaAction_Dataflow_CreateNode_DataflowEdNode::FAssetSchemaAction_Dataflow_CreateNode_DataflowEdNode(const FName& InName, const FName& InType, FText InNodeCategory, FText InMenuDesc, FText InToolTip, FText InKeywords, FName InIconName)
+	: FEdGraphSchemaAction(InNodeCategory, InMenuDesc, InToolTip, 0, InKeywords)
+	, NodeName(InName)
+	, NodeTypeName(InType)
+	, IconName(InIconName)
+{
+}
+
 UEdGraphNode* FAssetSchemaAction_Dataflow_CreateNode_DataflowEdNode::PerformAction(class UEdGraph* ParentGraph, UEdGraphPin* FromPin, const FVector2f& Location, bool bSelectNewNode)
 {
 	// by default use the type name 
@@ -65,6 +83,15 @@ UEdGraphNode* FAssetSchemaAction_Dataflow_CreateNode_DataflowEdNode::PerformActi
 	return NewNode;
 }
 
+FSlateBrush const* FAssetSchemaAction_Dataflow_CreateNode_DataflowEdNode::GetPaletteIcon() const
+{
+	if (IconName.IsNone())
+	{
+		return nullptr;
+	}
+	return FAppStyle::GetBrush(IconName);
+}
+
 /////////////////////////////////////////////////////////////////////////////////////////////////////////////////////////
 
 FEdGraphSchemaAction_DataflowVariable::FEdGraphSchemaAction_DataflowVariable()
diff --git a/Engine/Plugins/Experimental/Dataflow/Source/DataflowEditor/Private/Dataflow/SDataflowGraphActionWidget.h b/Engine/Plugins/Experimental/Dataflow/Source/DataflowEditor/Private/Dataflow/SDataflowGraphActionWidget.h
new file mode 100644
--- /dev/null
+++ b/Engine/Plugins/Experimental/Dataflow/Source/DataflowEditor/Private/Dataflow/SDataflowGraphActionWidget.h
@@ -0,0 +1,72 @@
+// Copyright Epic Games, Inc. All Rights Reserved.
+
+#pragma once
+
+#include "Dataflow/DataflowGraphSchemaAction.h"
+
+#include "SGraphActionMenu.h"
+#include "EdGraph/EdGraphSchema.h"
+#include "Templates/SharedPointer.h"
+#include "Widgets/Images/SImage.h"
+
+class SDataflowGraphActionWidget : public SCompoundWidget
+{
+	SLATE_BEGIN_ARGS(SDataflowGraphActionWidget) {}
+		SLATE_ARGUMENT(TSharedPtr<SWidget>, NameWidget)
+	SLATE_END_ARGS()
+
+	void Construct(const FArguments& InArgs, const FCreateWidgetForActionData* const InCreateData)
+	{
+		check(InCreateData->Action.IsValid());
+		ActionPtr = InCreateData->Action;
+		MouseButtonDownDelegate = InCreateData->MouseButtonDownDelegate;
+
+		// Icon widget
+		TSharedRef<SWidget> IconWidget = SNullWidget::NullWidget;
+		if (const FSlateBrush* const PrimaryIcon = InCreateData->Action->GetPaletteIcon())
+		{
+			SAssignNew(IconWidget, SImage)
+			.DesiredSizeOverride(FVector2D(16))
+			.Image(PrimaryIcon);
+		}
+
+		// Name
+		const TSharedRef<SWidget> NameWidget =
+			SNew(STextBlock)
+			.Font(FCoreStyle::GetDefaultFontStyle("Regular", 9))
+			.Text(InCreateData->Action->GetMenuDescription())
+			;
+
+		this->ChildSlot
+		[
+			SNew(SHorizontalBox)
+			.ToolTipText(InCreateData->Action->GetTooltipDescription())
+			+ SHorizontalBox::Slot()
+			.AutoWidth()
+			[
+				IconWidget
+			]
+			+ SHorizontalBox::Slot()
+			.AutoWidth()
+			.Padding(4, 0, 0, 0)
+			[
+				NameWidget
+			]
+		];
+	}
+
+	virtual FReply OnMouseButtonDown(const FGeometry& MyGeometry, const FPointerEvent& MouseEvent) override
+	{
+		if (MouseButtonDownDelegate.Execute(ActionPtr))
+		{
+			return FReply::Handled();
+		}
+
+		return FReply::Unhandled();
+	}
+
+	/** The action that we want to display with this widget */
+	TWeakPtr<FEdGraphSchemaAction> ActionPtr;
+	/** Delegate executed when mouse button goes down */
+	FCreateWidgetMouseButtonDown MouseButtonDownDelegate;
+};
diff --git a/Engine/Plugins/Experimental/Dataflow/Source/DataflowEditor/Public/Dataflow/DataflowGraphSchemaAction.h b/Engine/Plugins/Experimental/Dataflow/Source/DataflowEditor/Public/Dataflow/DataflowGraphSchemaAction.h
--- a/Engine/Plugins/Experimental/Dataflow/Source/DataflowEditor/Public/Dataflow/DataflowGraphSchemaAction.h
+++ b/Engine/Plugins/Experimental/Dataflow/Source/DataflowEditor/Public/Dataflow/DataflowGraphSchemaAction.h
@@ -40,16 +40,18 @@ struct FAssetSchemaAction_Dataflow_CreateNode_DataflowEdNode : public FEdGraphSc
 public:
 	FAssetSchemaAction_Dataflow_CreateNode_DataflowEdNode() : FEdGraphSchemaAction() {}
 
-	FAssetSchemaAction_Dataflow_CreateNode_DataflowEdNode(const FName& InName, const FName& InType, FText InNodeCategory, FText InMenuDesc, FText InToolTip, FText InKeywords)
-		: FEdGraphSchemaAction(InNodeCategory, InMenuDesc, InToolTip, 0, InKeywords), NodeName(InName), NodeTypeName(InType) {}
+	FAssetSchemaAction_Dataflow_CreateNode_DataflowEdNode(const FName& InName, const FName& InType, FText InNodeCategory, FText InMenuDesc, FText InToolTip, FText InKeywords, FName InIconName);
 
 	static UE_API TSharedPtr<FAssetSchemaAction_Dataflow_CreateNode_DataflowEdNode> CreateAction(const UEdGraph* ParentGraph, const FName& NodeTypeName, const FName& InOverrideNodeName = NAME_None);
 
 	using FEdGraphSchemaAction::PerformAction; // Prevent hiding of deprecated base class function with FVector2D
 	UE_API virtual UEdGraphNode* PerformAction(class UEdGraph* ParentGraph, UEdGraphPin* FromPin, const FVector2f& Location, bool bSelectNewNode = true) override;
 
+	virtual FSlateBrush const* GetPaletteIcon() const override;
+
 	FName NodeName;
 	FName NodeTypeName;
+	FName IconName;
 };
 
 /**
