Dataflow:
- Added FDataflowPlane type
- Added rendering for FDataflowPlane and TArray<FDataflowPlane>

#virtualized
#Synced-CL 52027509

diff --git a/Engine/Plugins/Experimental/Dataflow/Source/DataflowEditor/Private/Dataflow/DataflowEditorModule.cpp b/Engine/Plugins/Experimental/Dataflow/Source/DataflowEditor/Private/Dataflow/DataflowEditorModule.cpp
--- a/Engine/Plugins/Experimental/Dataflow/Source/DataflowEditor/Private/Dataflow/DataflowEditorModule.cpp
+++ b/Engine/Plugins/Experimental/Dataflow/Source/DataflowEditor/Private/Dataflow/DataflowEditorModule.cpp
@@ -48,6 +48,7 @@
 #include "DataflowRendering/DataflowMaterialRenderableType.h"
 #include "DataflowRendering/DataflowTextureRenderableType.h"
 #include "DataflowRendering/DataflowSamplerRenderableType.h"
+#include "DataflowRendering/DataflowPlaneRenderableType.h"
 
 #include "ContentBrowserModule.h"
 #include "Modules/ModuleManager.h"
@@ -362,6 +363,7 @@ void FDataflowEditorModule::StartupModule()
 	UE::Dataflow::Private::RegisterMaterialRenderableTypes();
 	UE::Dataflow::Private::RegisterTextureRenderableTypes();
 	UE::Dataflow::Private::RegisterSamplerRenderableTypes();
+	UE::Dataflow::Private::RegisterPlaneRenderableTypes();
 
 	// register an browser extension to show when an asset is bound to a dataflow asset
 	if (FContentBrowserModule* ContentBrowserModule = FModuleManager::Get().GetModulePtr<FContentBrowserModule>(TEXT("ContentBrowser")))
diff --git a/Engine/Plugins/Experimental/Dataflow/Source/DataflowEditor/Private/DataflowRendering/DataflowPlaneRenderableType.cpp b/Engine/Plugins/Experimental/Dataflow/Source/DataflowEditor/Private/DataflowRendering/DataflowPlaneRenderableType.cpp
new file mode 100644
--- /dev/null
+++ b/Engine/Plugins/Experimental/Dataflow/Source/DataflowEditor/Private/DataflowRendering/DataflowPlaneRenderableType.cpp
@@ -0,0 +1,107 @@
+// Copyright Epic Games, Inc. All Rights Reserved.
+
+#include "DataflowRendering/DataflowPlaneRenderableType.h"
+#include "DataflowRendering/DataflowRenderableTypeUtils.h"
+
+#include "Components/BoxComponent.h"
+#include "Components/InstancedStaticMeshComponent.h"
+#include "Drawing/LineSetComponent.h"
+#include "Components/ArrowComponent.h"
+
+#include "Dataflow/DataflowRenderingViewMode.h"
+#include "Dataflow/DataflowEngineUtil.h"
+#include "Dataflow/DataflowPlane.h"
+
+#include "DataflowRendering/DataflowRenderableComponents.h"
+#include "DataflowRendering/DataflowRenderableType.h"
+#include "DataflowRendering/DataflowRenderableTypeInstance.h"
+#include "DataflowRendering/DataflowRenderableTypeRegistry.h"
+#include "DataflowRendering/DataflowRenderableTypeUtils.h"
+
+#include "Engine/StaticMesh.h"
+
+#include "UObject/ObjectPtr.h"
+
+#include "Materials/MaterialInstanceDynamic.h"
+
+namespace UE::Dataflow::Private
+{
+	///////////////////////////////////////////////////////////////////////////////////////////////////////////////
+
+	class FPlaneSurfaceRenderableType : public IRenderableType
+	{
+		UE_DATAFLOW_IRENDERABLE_PRIMARY_TYPE(FDataflowPlane, Plane);
+		UE_DATAFLOW_IRENDERABLE_RENDER_GROUP(Surface);
+		UE_DATAFLOW_IRENDERABLE_VIEW_MODE(FDataflowConstruction3DViewMode);
+		UE_DATAFLOW_IRENDERABLE_SETTINGS(UDataflowPlaneRenderSettings);
+
+	private:
+		virtual void GetPrimitiveComponents(const FRenderableTypeInstance& Instance, FRenderableComponents& OutComponents) const override
+		{
+			const FDataflowPlane Plane = GetPlane(Instance, {});
+
+			const FName PlaneComponentName = Instance.GetComponentName(TEXT("Plane"));
+			const FName ArrowComponentName = Instance.GetComponentName(TEXT("PlaneNormal"));
+
+			UStaticMesh* PlaneMesh = UE::Dataflow::RenderGeometry::GetDataflowPlane();
+			if (PlaneMesh)
+			{
+				const UDataflowPlaneRenderSettings* Settings = Instance.GetTypedRenderSettings<UDataflowPlaneRenderSettings>();
+				if (Settings)
+				{
+					if (UStaticMeshComponent* Component = OutComponents.AddNewComponent<UStaticMeshComponent>(PlaneComponentName))
+					{
+						Component->SetStaticMesh(PlaneMesh);
+
+						FTransform PlaneTransform = Plane.AsTransform();
+						PlaneTransform.SetScale3D(FVector(Settings->PlaneScale));
+
+						Component->SetWorldTransform(PlaneTransform);
+						Component->SetCollisionEnabled(ECollisionEnabled::NoCollision);
+
+						UMaterialInterface* Material = Settings->ColorSettings.bWireframe
+							? UE::Dataflow::RenderMaterial::GetDataflowColorWireframeMaterial(true)
+							: UE::Dataflow::RenderMaterial::GetDataflowColorMaterial(true);
+
+						UMaterialInstanceDynamic* MaterialInstance = nullptr;
+
+						if (Material)
+						{
+							MaterialInstance = UMaterialInstanceDynamic::Create(Material, Component);
+							if (MaterialInstance)
+							{
+								MaterialInstance->SetVectorParameterValue(FName("Color"), Settings->ColorSettings.Color);
+								MaterialInstance->SetScalarParameterValue(FName("Transparency"), Settings->ColorSettings.Transparency);
+							}
+						}
+
+						Component->SetMaterial(0, MaterialInstance);
+
+						if (UArrowComponent* ArrowComponent = OutComponents.AddNewComponent<UArrowComponent>(ArrowComponentName))
+						{
+							ArrowComponent->ArrowColor = Settings->ColorSettings.Color.ToFColor(true);
+							ArrowComponent->ArrowSize = Settings->NormalScale;
+
+							FTransform XToZTransform = FTransform(FRotator(90.f, 0.f, 0.f));
+							ArrowComponent->SetWorldTransform(XToZTransform * Plane.AsTransform());
+
+							ArrowComponent->SetMaterial(0, MaterialInstance);
+						}
+					}
+				}
+			}
+		}
+	};
+}
+
+///////////////////////////////////////////////////////////////////////////////////////////////////////////////
+
+namespace UE::Dataflow::Private
+{
+	void RegisterPlaneRenderableTypes()
+	{
+		UE_DATAFLOW_REGISTER_RENDERABLE_TYPE(FPlaneSurfaceRenderableType);
+	}
+}
+
+///////////////////////////////////////////////////////////////////////////////////////////////////////////////
diff --git a/Engine/Plugins/Experimental/Dataflow/Source/DataflowEditor/Private/DataflowRendering/DataflowPlaneRenderableType.h b/Engine/Plugins/Experimental/Dataflow/Source/DataflowEditor/Private/DataflowRendering/DataflowPlaneRenderableType.h
new file mode 100644
--- /dev/null
+++ b/Engine/Plugins/Experimental/Dataflow/Source/DataflowEditor/Private/DataflowRendering/DataflowPlaneRenderableType.h
@@ -0,0 +1,34 @@
+// Copyright Epic Games, Inc. All Rights Reserved.
+
+#pragma once
+
+#include "DataflowRendering/DataflowRenderableTypeSettings.h"
+#include "DataflowRendering/DataflowRenderableTypeUtils.h"
+
+#include "DataflowPlaneRenderableType.generated.h"
+
+UCLASS(MinimalAPI)
+class UDataflowPlaneRenderSettings : public UDataflowRenderableTypeSettings
+{
+	GENERATED_BODY()
+
+public:
+	UPROPERTY(EditAnywhere, Category = "Surface", meta = (ShowOnlyInnerProperties))
+	FDataflowSimpleColorCommonRenderSettings ColorSettings;
+
+	/** Length scalar for the normal */
+	UPROPERTY(EditAnywhere, Category = "Plane", meta = (UIMin = "1.0", UIMax = "10.0", ClampMin = "1.0", ClampMax = "10.0"))
+	float NormalScale = 1.f;
+
+	/** Plane size scalar */
+	UPROPERTY(EditAnywhere, Category = "Plane", meta = (UIMin = "1.0", UIMax = "10.0", ClampMin = "1.0", ClampMax = "10.0"))
+	float PlaneScale = 1.f;
+};
+
+namespace UE::Dataflow::Private
+{
+	void RegisterPlaneRenderableTypes();
+}
+
+
+
