Add MinSize support for layout stacks
- splitter slots do already support min size but no related was passed from the layout
- Use it for the dataflow editor timeline tab
- Tested in the dataflow editor

#Synced-CL 51997199

diff --git a/Engine/Plugins/Experimental/Dataflow/Source/DataflowEditor/Private/Dataflow/DataflowEditorToolkit.cpp b/Engine/Plugins/Experimental/Dataflow/Source/DataflowEditor/Private/Dataflow/DataflowEditorToolkit.cpp
--- a/Engine/Plugins/Experimental/Dataflow/Source/DataflowEditor/Private/Dataflow/DataflowEditorToolkit.cpp
+++ b/Engine/Plugins/Experimental/Dataflow/Source/DataflowEditor/Private/Dataflow/DataflowEditorToolkit.cpp
@@ -71,6 +71,7 @@
 #include "UObject/PackageReload.h"
 #include "Widgets/Notifications/SNotificationList.h"
 #include "Widgets/Docking/SDockTab.h"
+#include "Framework/Docking/SDockingTabStack.h"
 #include "Widgets/Layout/SSpacer.h"
 #include "Widgets/SDataflowSimulationTimeline.h"
 #include "Widgets/DataflowSimulationBinding.h"
@@ -115,6 +116,7 @@ const FName FDataflowEditorToolkit::OutputLogTabId(TEXT("DataflowEditor_OutputLo
 namespace UE::Dataflow::Private
 {
 	static const FName DataflowEditorToolkitName("DataflowEditor");
+	static const float DataflowLayoutTimelineMaxHeight = 75.f;
 
 	/** Create a debug draw component to visualize per node informations */
 	TUniquePtr<IDataflowDebugDrawInterface> CreateDebugDrawComponent(
@@ -389,6 +391,7 @@ FDataflowEditorToolkit::FDataflowEditorToolkit(UAssetEditor* InOwningAssetEditor
 					(
 					FTabManager::NewStack()
 							->SetSizeCoefficient(0.1f)		// Relative width of (simulation timeline) vs (Preview Viewport, Construction Viewport)
+							->SetMinSize(UE::Dataflow::Private::DataflowLayoutTimelineMaxHeight)
 							->AddTab(SimulationTimelineTabId, ETabState::OpenedTab)
 							->SetExtensionId("TimelineArea")
 							->SetHideTabWell(true)
@@ -479,6 +482,7 @@ FDataflowEditorToolkit::FDataflowEditorToolkit(UAssetEditor* InOwningAssetEditor
 					(
 					FTabManager::NewStack()
 							->SetSizeCoefficient(0.1f)		// Relative width of (simulation timeline) vs (Preview Viewport, Construction Viewport)
+							->SetMinSize(UE::Dataflow::Private::DataflowLayoutTimelineMaxHeight)
 							->AddTab(SimulationTimelineTabId, ETabState::OpenedTab)
 							->SetExtensionId("TimelineArea")
 							->SetHideTabWell(true)
@@ -890,6 +894,16 @@ void FDataflowEditorToolkit::PostInitAssetEditor()
 
 	// Force scenes to update loaded asset viewer settings
 	UAssetViewerSettings::Get()->OnAssetViewerSettingsChanged().Broadcast(NAME_None);
+
+	// Enforce the timeline min size even when restoring an older saved layout that predates the MinSize field.
+	// The SSplitter slot binds GetNodeMinSize() as a TAttribute, so setting it here takes effect immediately.
+	if (TSharedPtr<SDockTab> TimelineTab = TabManager->FindExistingLiveTab(SimulationTimelineTabId))
+	{
+		if (TSharedPtr<SDockingTabStack> TimelineStack = TimelineTab->GetParentDockTabStack())
+		{
+			TimelineStack->SetNodeMinSize(UE::Dataflow::Private::DataflowLayoutTimelineMaxHeight);
+		}
+	}
 }
 
 void FDataflowEditorToolkit::SetEvaluateGraphMode(EDataflowEditorEvaluationMode Mode)
diff --git a/Engine/Source/Runtime/Slate/Private/Framework/Docking/SDockingSplitter.cpp b/Engine/Source/Runtime/Slate/Private/Framework/Docking/SDockingSplitter.cpp
--- a/Engine/Source/Runtime/Slate/Private/Framework/Docking/SDockingSplitter.cpp
+++ b/Engine/Source/Runtime/Slate/Private/Framework/Docking/SDockingSplitter.cpp
@@ -26,6 +26,7 @@ void SDockingSplitter::AddChildNode( const TSharedRef<SDockingNode>& InChild, in
 	.Value( TAttribute<float>(InChild, &SDockingNode::GetSizeCoefficient) )
 	.OnSlotResized( SSplitter::FOnSlotResized::CreateSP( InChild, &SDockingNode::SetSizeCoefficient ) )
 	.SizeRule( TAttribute<SSplitter::ESizeRule>(InChild, &SDockingNode::GetSizeRule) )
+	.MinSize( TAttribute<float>(InChild, &SDockingNode::GetNodeMinSize) )
 	[
 		InChild
 	];	
@@ -64,6 +65,7 @@ void SDockingSplitter::ReplaceChild( const TSharedRef<SDockingNode>& InChildToRe
 	Slot.SetSizeValue(TAttribute<float>(Replacement, &SDockingNode::GetSizeCoefficient));
 	Slot.OnSlotResized().BindSP(Replacement, &SDockingNode::SetSizeCoefficient);
 	Slot.SetSizingRule(TAttribute<SSplitter::ESizeRule>(Replacement, &SDockingNode::GetSizeRule));
+	Slot.SetMinSize(TAttribute<float>(Replacement, &SDockingNode::GetNodeMinSize));
 	Slot[Replacement];
 
 	Replacement->SetParentNode( SharedThis(this) );
diff --git a/Engine/Source/Runtime/Slate/Private/Framework/Docking/SDockingTabStack.cpp b/Engine/Source/Runtime/Slate/Private/Framework/Docking/SDockingTabStack.cpp
--- a/Engine/Source/Runtime/Slate/Private/Framework/Docking/SDockingTabStack.cpp
+++ b/Engine/Source/Runtime/Slate/Private/Framework/Docking/SDockingTabStack.cpp
@@ -52,6 +52,7 @@ void SDockingTabStack::Construct( const FArguments& InArgs, const TSharedRef<FTa
 
 	Tabs = PersistentNode->Tabs;
 	this->SetSizeCoefficient(PersistentNode->GetSizeCoefficient());
+	NodeMinSize = PersistentNode->GetMinSize();
 
 	// the value of this is determined every time a tab is added
 	bShowingTitleBarArea = false;
@@ -923,7 +924,8 @@ TSharedPtr<FTabManager::FLayoutNode> SDockingTabStack::GatherPersistentLayout()
 		TSharedRef<FTabManager::FStack> PersistentStack =
 			FTabManager::NewStack()
 			->SetSizeCoefficient( this->GetSizeCoefficient() )
-			->SetHideTabWell( this->IsTabWellHidden() );
+			->SetHideTabWell( this->IsTabWellHidden() )
+			->SetMinSize( NodeMinSize );
 
 		TSharedPtr<SDockTab> ForegroundTab = TabWell->GetForegroundTab();
 		if(ForegroundTab.IsValid())
diff --git a/Engine/Source/Runtime/Slate/Private/Framework/Docking/TabManager.cpp b/Engine/Source/Runtime/Slate/Private/Framework/Docking/TabManager.cpp
--- a/Engine/Source/Runtime/Slate/Private/Framework/Docking/TabManager.cpp
+++ b/Engine/Source/Runtime/Slate/Private/Framework/Docking/TabManager.cpp
@@ -343,6 +343,13 @@ TSharedRef<FTabManager::FLayoutNode> FTabManager::FLayout::NewFromString_Helper(
 		TSharedRef<FTabManager::FStack> NewStack = FTabManager::NewStack();
 		NewStack->SetSizeCoefficient((float)JsonObject->GetNumberField(TEXT("SizeCoefficient")) );
 		NewStack->SetHideTabWell( JsonObject->GetBoolField(TEXT("HideTabWell")) );
+		{
+			double MinSize = 0.0;
+			if (JsonObject->TryGetNumberField(TEXT("MinSize"), MinSize))
+			{
+				NewStack->SetMinSize((float)MinSize);
+			}
+		}
 
 		if(JsonObject->HasField(TEXT("ForegroundTab")))
 		{
@@ -549,6 +556,10 @@ TSharedRef<FJsonObject> FTabManager::FLayout::PersistToString_Helper(const TShar
 	{
 		JsonObj->SetStringField( TEXT("Type"), TEXT("Stack") );
 		JsonObj->SetBoolField( TEXT("HideTabWell"), NodeAsStack->bHideTabWell );
+		if (NodeAsStack->MinSize != 0.f)
+		{
+			JsonObj->SetNumberField( TEXT("MinSize"), NodeAsStack->MinSize );
+		}
 
 		if (NodeAsStack->ForegroundTabId.ShouldSaveLayout())
 		{
diff --git a/Engine/Source/Runtime/Slate/Public/Framework/Docking/SDockingNode.h b/Engine/Source/Runtime/Slate/Public/Framework/Docking/SDockingNode.h
--- a/Engine/Source/Runtime/Slate/Public/Framework/Docking/SDockingNode.h
+++ b/Engine/Source/Runtime/Slate/Public/Framework/Docking/SDockingNode.h
@@ -108,6 +108,9 @@ public:
 	/** Should this node auto-size or be a percentage of its parent size; this setting is usually determined by users */
 	virtual SSplitter::ESizeRule GetSizeRule() const { return SSplitter::FractionOfParent; }
 
+	/** Minimum pixel size for this node when it is a slot in a parent splitter. Zero means no minimum. */
+	virtual float GetNodeMinSize() const { return 0.f; }
+
 	/** @return the numerator for the fraction of available space that this dock node should occupy. */
 	SLATE_API float GetSizeCoefficient() const;
 
diff --git a/Engine/Source/Runtime/Slate/Public/Framework/Docking/SDockingTabStack.h b/Engine/Source/Runtime/Slate/Public/Framework/Docking/SDockingTabStack.h
--- a/Engine/Source/Runtime/Slate/Public/Framework/Docking/SDockingTabStack.h
+++ b/Engine/Source/Runtime/Slate/Public/Framework/Docking/SDockingTabStack.h
@@ -59,6 +59,9 @@ public:
 		return SDockingNode::DockTabStack;
 	}
 
+	virtual float GetNodeMinSize() const override { return NodeMinSize; }
+	void SetNodeMinSize(float InMinSize) { NodeMinSize = InMinSize; }
+
 	SLATE_API void OnLastTabRemoved();
 
 	SLATE_API void OnTabClosed(const TSharedRef<SDockTab>& ClosedTab, SDockingNode::ELayoutModification RemovalMethod);
@@ -226,8 +229,11 @@ protected:
 	bool IsContentEnabled() const;
 
 private:
-	
-	/** Data that persists across sessions and when the widget associated with this node is removed. */	
+
+	/** Minimum pixel size enforced on the SSplitter slot containing this stack. Set from FTabManager::FStack::SetMinSize(). */
+	float NodeMinSize = 0.f;
+
+	/** Data that persists across sessions and when the widget associated with this node is removed. */
 	TArray<FTabManager::FTab> Tabs;
 
 
diff --git a/Engine/Source/Runtime/Slate/Public/Framework/Docking/TabManager.h b/Engine/Source/Runtime/Slate/Public/Framework/Docking/TabManager.h
--- a/Engine/Source/Runtime/Slate/Public/Framework/Docking/TabManager.h
+++ b/Engine/Source/Runtime/Slate/Public/Framework/Docking/TabManager.h
@@ -463,13 +463,17 @@ class FTabManager : public TSharedFromThis<FTabManager>
 
 			float GetSizeCoefficient() const { return SizeCoefficient; }
 
+		float GetMinSize() const { return MinSize; }
+
 		protected:
 			FLayoutNode()
 			: SizeCoefficient(1.0f)
+			, MinSize(0.f)
 			{
 			}
 
 			float SizeCoefficient;
+			float MinSize;
 		};
 
 		struct FTab
@@ -560,6 +564,12 @@ class FTabManager : public TSharedFromThis<FTabManager>
 					return SharedThis(this);
 				}
 
+				TSharedRef<FStack> SetMinSize( const float InMinSize )
+				{
+					MinSize = InMinSize;
+					return SharedThis(this);
+				}
+
 				TSharedRef<FStack> SetHideTabWell( const bool InHideTabWell )
 				{
 					bHideTabWell = InHideTabWell;
