Dataflow:
- Added FDataflowColorRampSamplerNode

#Synced-CL 52048340

diff --git a/Engine/Plugins/Experimental/Dataflow/Source/DataflowNodes/Private/Dataflow/SamplerNodes/DataflowColorRampSamplerNode.cpp b/Engine/Plugins/Experimental/Dataflow/Source/DataflowNodes/Private/Dataflow/SamplerNodes/DataflowColorRampSamplerNode.cpp
new file mode 100644
--- /dev/null
+++ b/Engine/Plugins/Experimental/Dataflow/Source/DataflowNodes/Private/Dataflow/SamplerNodes/DataflowColorRampSamplerNode.cpp
@@ -0,0 +1,77 @@
+// Copyright Epic Games, Inc. All Rights Reserved.
+
+#include "Dataflow/SamplerNodes/DataflowColorRampSamplerNode.h"
+#include "Dataflow/SamplerNodes/DataflowSamplerFunctions.h"
+#include "Dataflow/DataflowNodeFactory.h"
+
+FBox FDataflowColorRampVectorSampler::GetRenderBounds() const
+{
+	return FBox(FVector(-50.0), FVector(50.0));
+}
+
+void FDataflowColorRampVectorSampler::Sample(TArrayView<const FVector3f> Positions, TArrayView<FVector3f> OutValues) const
+{
+	if (Sampler.IsValid() && Positions.Num() == OutValues.Num())
+	{
+		TArray<float> OutValueFromInputSampler;
+		OutValueFromInputSampler.SetNumUninitialized(OutValues.Num());
+
+		Sampler->Sample(Positions, OutValueFromInputSampler);
+
+		if (Max - Min < SMALL_NUMBER)
+		{
+			for (int32 Idx = 0; Idx < OutValues.Num(); ++Idx)
+			{
+				OutValues[Idx] = FVector3f::ZeroVector;
+			}
+
+			return;
+		}
+
+		for (int32 Idx = 0; Idx < OutValues.Num(); ++Idx)
+		{
+
+			float NormalizedValue = (OutValueFromInputSampler[Idx] - Min) / (Max - Min);
+			NormalizedValue = FMath::Clamp(NormalizedValue, 0.f, 1.f);
+
+			FLinearColor Color = ColorRamp.GetLinearColorValue(NormalizedValue);
+
+			OutValues[Idx] = FVector3f(Color.R, Color.G, Color.B);
+		}
+	}
+	else
+	{
+		for (int32 Idx = 0; Idx < OutValues.Num(); ++Idx)
+		{
+			OutValues[Idx] = FVector3f::ZeroVector;
+		}
+	}
+}
+
+FDataflowColorRampSamplerNode::FDataflowColorRampSamplerNode(const UE::Dataflow::FNodeParameters& InParam, FGuid InGuid)
+	: FDataflowNode(InParam, InGuid)
+{
+	RegisterInputConnection(&Sampler);
+	RegisterOutputConnection(&ColorSampler);
+
+	constexpr bool bOnlyRGB = true;
+	ColorRampSampler.ColorRamp.SetColorAtTime(0.f, FColor::Purple, bOnlyRGB);
+	ColorRampSampler.ColorRamp.SetColorAtTime(1.0f, FLinearColor::Yellow, bOnlyRGB);
+}
+
+void FDataflowColorRampSamplerNode::Evaluate(UE::Dataflow::FContext& Context, const FDataflowOutput* Out) const
+{
+	if (Out->IsA(&ColorSampler))
+	{
+		const FDataflowFloatSampler& InFloatSampler = GetValue(Context, &Sampler);
+
+		TSharedRef<FDataflowColorRampVectorSampler> Impl = MakeShared<FDataflowColorRampVectorSampler>(ColorRampSampler);
+
+		Impl->Sampler = InFloatSampler.GetImpl();
+
+		FDataflowVectorSampler OutSampler(Impl.ToSharedPtr());
+
+		SetValue(Context, OutSampler, &ColorSampler);
+		return;
+	}
+}
diff --git a/Engine/Plugins/Experimental/Dataflow/Source/DataflowNodes/Private/Dataflow/SamplerNodes/DataflowColorRampSamplerNode.h b/Engine/Plugins/Experimental/Dataflow/Source/DataflowNodes/Private/Dataflow/SamplerNodes/DataflowColorRampSamplerNode.h
new file mode 100644
--- /dev/null
+++ b/Engine/Plugins/Experimental/Dataflow/Source/DataflowNodes/Private/Dataflow/SamplerNodes/DataflowColorRampSamplerNode.h
@@ -0,0 +1,67 @@
+// Copyright Epic Games, Inc. All Rights Reserved.
+
+#pragma once 
+
+#include "Dataflow/DataflowNode.h"
+#include "Dataflow/SamplerNodes/DataflowSamplerTypes.h"
+#include "Dataflow/DataflowImage.h"
+#include "Curves/LinearColorRamp.h"
+
+#include "DataflowColorRampSamplerNode.generated.h"
+
+USTRUCT()
+struct FDataflowColorRampVectorSampler : public FDataflowVectorSamplerBase
+{
+	GENERATED_BODY();
+
+	virtual ~FDataflowColorRampVectorSampler() override = default;
+	virtual void Sample(TArrayView<const FVector3f> Positions, TArrayView<FVector3f> OutValues) const override;
+	virtual FBox GetRenderBounds() const override;
+
+public:
+	/** Attribute minimum value for coloring. This value represents the minimum color on the color ramp. */
+	UPROPERTY(EditAnywhere, Category = "Color")
+	float Min = 0.f;
+
+	/** Attribute maximum value for coloring. This value represents the maximum color on the color ramp. */
+	UPROPERTY(EditAnywhere, Category = "Color")
+	float Max = 1.f;
+
+	/** Color ramp specifiyng the color for displayed points using a single value attribute. Only used with FloatSampler input. */
+	UPROPERTY(EditAnywhere, Category = "Color")
+	FLinearColorRamp ColorRamp;
+
+	TSharedPtr<const FDataflowFloatSamplerBase> Sampler;
+};
+
+/**
+ *
+ * ColorRamp Sampler
+ * Input(s) : Float Sampler
+ * Output(s): Vector sampler representing colors generated by the color ramp using the sampled float values
+ *
+ */
+USTRUCT(meta = (Icon = "Icons.EyeDropper"))
+struct FDataflowColorRampSamplerNode : public FDataflowNode
+{
+	GENERATED_BODY()
+	DATAFLOW_NODE_DEFINE_INTERNAL(FDataflowColorRampSamplerNode, "ColorRamp Sampler", "Samplers", "")
+
+public:
+	FDataflowColorRampSamplerNode(const UE::Dataflow::FNodeParameters& InParam, FGuid InGuid = FGuid::NewGuid());
+
+private:
+	/** Sampler input */
+	UPROPERTY(meta = (DataflowInput))
+	FDataflowFloatSampler Sampler;
+
+	/** Sampler output */
+	UPROPERTY(meta = (DataflowOutput))
+	FDataflowVectorSampler ColorSampler;
+
+	UPROPERTY(EditAnywhere, Category = Sampler, meta = (ShowOnlyInnerProperties, SkipInDisplayNameChain))
+	FDataflowColorRampVectorSampler ColorRampSampler;
+
+
+	void Evaluate(UE::Dataflow::FContext& Context, const FDataflowOutput* Out) const override;
+};
