Dataflow:
- Added FDataflowLerpSamplerNode

#Synced-CL 52048340

diff --git a/Engine/Plugins/Experimental/Dataflow/Source/DataflowNodes/Private/Dataflow/SamplerNodes/DataflowLerpSamplerNode.cpp b/Engine/Plugins/Experimental/Dataflow/Source/DataflowNodes/Private/Dataflow/SamplerNodes/DataflowLerpSamplerNode.cpp
new file mode 100644
--- /dev/null
+++ b/Engine/Plugins/Experimental/Dataflow/Source/DataflowNodes/Private/Dataflow/SamplerNodes/DataflowLerpSamplerNode.cpp
@@ -0,0 +1,170 @@
+// Copyright Epic Games, Inc. All Rights Reserved.
+
+#include "Dataflow/SamplerNodes/DataflowLerpSamplerNode.h"
+#include "Dataflow/SamplerNodes/DataflowSamplerFunctions.h"
+#include "Dataflow/DataflowNodeFactory.h"
+
+FBox FDataflowLerpFloatSampler::GetRenderBounds() const
+{
+	return FBox(FVector(-50.0), FVector(50.0));
+}
+
+void FDataflowLerpFloatSampler::Sample(TArrayView<const FVector3f> Positions, TArrayView<float> OutValues) const
+{
+	if (FloatSamplerA.IsValid() && Positions.Num() == OutValues.Num())
+	{
+		FloatSamplerA->Sample(Positions, OutValues);
+
+		if (FloatSamplerB.IsValid())
+		{
+			TArray<float> OutValueFromFloatSamplerB;
+			OutValueFromFloatSamplerB.SetNumUninitialized(OutValues.Num());
+
+			FloatSamplerB->Sample(Positions, OutValueFromFloatSamplerB);
+
+			if (BlendSampler.IsValid())
+			{
+				TArray<float> OutValueFromFloatSampler;
+				OutValueFromFloatSampler.SetNumUninitialized(OutValues.Num());
+
+				BlendSampler->Sample(Positions, OutValueFromFloatSampler);
+
+				for (int32 Idx = 0; Idx < OutValues.Num(); ++Idx)
+				{
+					OutValues[Idx] = FMath::Lerp(OutValues[Idx], OutValueFromFloatSamplerB[Idx], OutValueFromFloatSampler[Idx]);
+				}
+			}
+			else
+			{
+				for (int32 Idx = 0; Idx < OutValues.Num(); ++Idx)
+				{
+					OutValues[Idx] = FMath::Lerp(OutValues[Idx], OutValueFromFloatSamplerB[Idx], Blend);
+				}
+			}
+		}
+	}
+	else
+	{
+		for (int32 Idx = 0; Idx < OutValues.Num(); ++Idx)
+		{
+			OutValues[Idx] = 0.f;
+		}
+	}
+}
+
+FBox FDataflowLerpVectorSampler::GetRenderBounds() const
+{
+	return FBox(FVector(-50.0), FVector(50.0));
+}
+
+void FDataflowLerpVectorSampler::Sample(TArrayView<const FVector3f> Positions, TArrayView<FVector3f> OutValues) const
+{
+	if (VectorSamplerA.IsValid() && Positions.Num() == OutValues.Num())
+	{
+		VectorSamplerA->Sample(Positions, OutValues);
+
+		if (VectorSamplerB.IsValid())
+		{
+			TArray<FVector3f> OutValueFromFloatSamplerB;
+			OutValueFromFloatSamplerB.SetNumUninitialized(OutValues.Num());
+
+			VectorSamplerB->Sample(Positions, OutValueFromFloatSamplerB);
+
+			if (BlendSampler.IsValid())
+			{
+				TArray<float> OutValueFromFloatSampler;
+				OutValueFromFloatSampler.SetNumUninitialized(OutValues.Num());
+
+				BlendSampler->Sample(Positions, OutValueFromFloatSampler);
+
+				for (int32 Idx = 0; Idx < OutValues.Num(); ++Idx)
+				{
+					OutValues[Idx] = FMath::Lerp(OutValues[Idx], OutValueFromFloatSamplerB[Idx], OutValueFromFloatSampler[Idx]);
+
+				}
+			}
+			else
+			{
+				for (int32 Idx = 0; Idx < OutValues.Num(); ++Idx)
+				{
+					OutValues[Idx] = FMath::Lerp(OutValues[Idx], OutValueFromFloatSamplerB[Idx], Blend);
+				}
+			}
+		}
+	}
+	else
+	{
+		for (int32 Idx = 0; Idx < OutValues.Num(); ++Idx)
+		{
+			OutValues[Idx] = FVector3f::ZeroVector;
+		}
+	}
+}
+
+FDataflowLerpSamplerNode::FDataflowLerpSamplerNode(const UE::Dataflow::FNodeParameters& InParam, FGuid InGuid)
+	: FDataflowNode(InParam, InGuid)
+{
+	static const FName MainTypeGroup("Main");
+
+	RegisterInputConnection(&SamplerA)
+		.SetTypeDependencyGroup(MainTypeGroup);
+	RegisterInputConnection(&SamplerB)
+		.SetTypeDependencyGroup(MainTypeGroup);
+	RegisterInputConnection(&BlendSampler);
+
+	RegisterInputConnection(&Blend);
+
+	RegisterOutputConnection(&SamplerA)
+		.SetPassthroughInput(&SamplerA)
+		.SetTypeDependencyGroup(MainTypeGroup);
+}
+
+void FDataflowLerpSamplerNode::Evaluate(UE::Dataflow::FContext& Context, const FDataflowOutput* Out) const
+{
+	if (Out->IsA(&SamplerA))
+	{
+		const float InBlend = GetValue(Context, &Blend);
+
+		const FDataflowFloatSampler* InFloatSamplerA = GetValue(Context, &SamplerA).TryGet <FDataflowFloatSampler>();
+		const FDataflowFloatSampler* InFloatSamplerB = GetValue(Context, &SamplerB).TryGet <FDataflowFloatSampler>();
+		const FDataflowFloatSampler& InBlendSampler = GetValue(Context, &BlendSampler);
+
+		if (InFloatSamplerA && InFloatSamplerB)
+		{
+			TSharedRef<FDataflowLerpFloatSampler> Impl = MakeShared<FDataflowLerpFloatSampler>();
+
+			Impl->FloatSamplerA = InFloatSamplerA->GetImpl();
+			Impl->FloatSamplerB = InFloatSamplerB->GetImpl();
+			Impl->Blend = InBlend;
+
+			Impl->BlendSampler = InBlendSampler.GetImpl();
+
+			FDataflowFloatSampler OutSampler(Impl.ToSharedPtr());
+
+			SetValue(Context, OutSampler, &SamplerA);
+			return;
+		}
+
+		const FDataflowVectorSampler* VectorSamplerA = GetValue(Context, &SamplerA).TryGet<FDataflowVectorSampler>();
+		const FDataflowVectorSampler* VectorSamplerB = GetValue(Context, &SamplerB).TryGet<FDataflowVectorSampler>();
+
+		if (VectorSamplerA && VectorSamplerB)
+		{
+			TSharedRef<FDataflowLerpVectorSampler> Impl = MakeShared<FDataflowLerpVectorSampler>();
+
+			Impl->VectorSamplerA = VectorSamplerA->GetImpl();
+			Impl->VectorSamplerB = VectorSamplerB->GetImpl();
+			Impl->Blend = InBlend;
+
+			Impl->BlendSampler = InBlendSampler.GetImpl();
+
+			FDataflowVectorSampler OutSampler(Impl.ToSharedPtr());
+
+			SetValue(Context, OutSampler, &SamplerA);
+			return;
+		}
+
+		SafeForwardInput(Context, &SamplerA, &SamplerA);
+	}
+
+}
diff --git a/Engine/Plugins/Experimental/Dataflow/Source/DataflowNodes/Private/Dataflow/SamplerNodes/DataflowLerpSamplerNode.h b/Engine/Plugins/Experimental/Dataflow/Source/DataflowNodes/Private/Dataflow/SamplerNodes/DataflowLerpSamplerNode.h
new file mode 100644
--- /dev/null
+++ b/Engine/Plugins/Experimental/Dataflow/Source/DataflowNodes/Private/Dataflow/SamplerNodes/DataflowLerpSamplerNode.h
@@ -0,0 +1,84 @@
+// Copyright Epic Games, Inc. All Rights Reserved.
+
+#pragma once 
+
+#include "Dataflow/DataflowNode.h"
+#include "Dataflow/SamplerNodes/DataflowSamplerTypes.h"
+
+#include "DataflowLerpSamplerNode.generated.h"
+
+USTRUCT()
+struct FDataflowLerpFloatSampler : public FDataflowFloatSamplerBase
+{
+	GENERATED_BODY();
+
+	virtual ~FDataflowLerpFloatSampler() override = default;
+	virtual void Sample(TArrayView<const FVector3f> Positions, TArrayView<float> OutValues) const override;
+	virtual FBox GetRenderBounds() const override;
+
+public:
+	UPROPERTY()
+	float Blend = 0.f;
+
+	TSharedPtr<const FDataflowFloatSamplerBase> FloatSamplerA;
+
+	TSharedPtr<const FDataflowFloatSamplerBase> FloatSamplerB;
+
+	TSharedPtr<const FDataflowFloatSamplerBase> BlendSampler;
+};
+
+USTRUCT()
+struct FDataflowLerpVectorSampler : public FDataflowVectorSamplerBase
+{
+	GENERATED_BODY();
+
+	virtual ~FDataflowLerpVectorSampler() override = default;
+	virtual void Sample(TArrayView<const FVector3f> Positions, TArrayView<FVector3f> OutValues) const override;
+	virtual FBox GetRenderBounds() const override;
+
+public:
+	UPROPERTY()
+	float Blend = 0.f;
+
+	TSharedPtr<const FDataflowVectorSamplerBase> VectorSamplerA;
+
+	TSharedPtr<const FDataflowVectorSamplerBase> VectorSamplerB;
+
+	TSharedPtr<const FDataflowFloatSamplerBase> BlendSampler;
+};
+
+/**
+ *
+ * LerpSampler
+ * Input(s) : 2 Float or Vector Sampler, optional BlendSampler (if not connected blend value can be used)
+ * Output(s): Same type of sampler as the input ooutputting the lerped value
+ *
+ */
+USTRUCT(meta = (Icon = "Icons.EyeDropper"))
+struct FDataflowLerpSamplerNode : public FDataflowNode
+{
+	GENERATED_BODY()
+	DATAFLOW_NODE_DEFINE_INTERNAL(FDataflowLerpSamplerNode, "Lerp Sampler", "Samplers", "Linear Interpolation")
+
+public:
+	FDataflowLerpSamplerNode(const UE::Dataflow::FNodeParameters& InParam, FGuid InGuid = FGuid::NewGuid());
+
+private:
+	/** First input for lerp */
+	UPROPERTY(meta = (DataflowInput, DataflowOutput, DataflowPassthrough = SamplerA))
+	FDataflowSamplerTypes SamplerA;
+
+	/** Second input for lerp */
+	UPROPERTY(meta = (DataflowInput))
+	FDataflowSamplerTypes SamplerB;
+
+	/** Optional FloatSampler to control lerp per sampled points */
+	UPROPERTY(meta = (DataflowInput))
+	FDataflowFloatSampler BlendSampler;
+
+	/** If BlendSampler is not connected this value controls the lerp */
+	UPROPERTY(EditAnywhere, Category = Lerp, meta = (DataflowInput, UIMin = 0.f, ClampMin = 0.f, UIMax = 1.f, ClampMax = 1.f))
+	float Blend = 0.f;
+
+	void Evaluate(UE::Dataflow::FContext& Context, const FDataflowOutput* Out) const override;
+};
