Dataflow : properly invalid tool nodes when active snapshot is changed or if a snapshot is removed from the UI

#Synced-CL 52092384

diff --git a/Engine/Plugins/Experimental/Dataflow/Source/DataflowNodes/Private/Dataflow/Customization/DataflowToolNodeSnapshotCustomization.cpp b/Engine/Plugins/Experimental/Dataflow/Source/DataflowNodes/Private/Dataflow/Customization/DataflowToolNodeSnapshotCustomization.cpp
--- a/Engine/Plugins/Experimental/Dataflow/Source/DataflowNodes/Private/Dataflow/Customization/DataflowToolNodeSnapshotCustomization.cpp
+++ b/Engine/Plugins/Experimental/Dataflow/Source/DataflowNodes/Private/Dataflow/Customization/DataflowToolNodeSnapshotCustomization.cpp
@@ -310,7 +310,7 @@ TSharedPtr<SWidget> FDataflowToolNodeSnapshotSetCustomization::OnContextMenuOpen
 						{
 							if (FDataflowToolNodeSnapshotSet* SnapshotSet = reinterpret_cast<FDataflowToolNodeSnapshotSet*>(StructData))
 							{
-								SnapshotSet->SetActiveSnapshot(SnapshotIndex);
+								SnapshotSet->SetActiveSnapshot(SnapshotIndex, /*bNotify*/true);
 								RefreshView();
 							}
 						}
@@ -339,7 +339,7 @@ TSharedPtr<SWidget> FDataflowToolNodeSnapshotSetCustomization::OnContextMenuOpen
 						{
 							if (FDataflowToolNodeSnapshotSet* SnapshotSet = reinterpret_cast<FDataflowToolNodeSnapshotSet*>(StructData))
 							{
-								SnapshotSet->RemoveSnapshot(SnapshotIndex);
+								SnapshotSet->RemoveSnapshot(SnapshotIndex, /*bNotify*/true);
 								RefreshView();
 							}
 						}
diff --git a/Engine/Plugins/Experimental/Dataflow/Source/DataflowNodes/Private/Dataflow/DataflowToolNode.cpp b/Engine/Plugins/Experimental/Dataflow/Source/DataflowNodes/Private/Dataflow/DataflowToolNode.cpp
--- a/Engine/Plugins/Experimental/Dataflow/Source/DataflowNodes/Private/Dataflow/DataflowToolNode.cpp
+++ b/Engine/Plugins/Experimental/Dataflow/Source/DataflowNodes/Private/Dataflow/DataflowToolNode.cpp
@@ -11,6 +11,14 @@ uint64 FDataflowToolNodeSnapshot::GetAllocatedSize() const
 
 //////////////////////////////////////////////////////////////////////////////////////////////////////////////////////////
 
+void FDataflowToolNodeSnapshotSet::NotifyChanged()
+{
+	if (OnSnapshotSetChanged.IsBound())
+	{
+		OnSnapshotSetChanged.Broadcast(*this);
+	}
+}
+
 TConstArrayView<FDataflowToolNodeSnapshot> FDataflowToolNodeSnapshotSet::GetSnapshots() const
 {
 	return Snapshots;
@@ -26,29 +34,40 @@ const FDataflowToolNodeSnapshot* FDataflowToolNodeSnapshotSet::GetActiveSnapshot
 	
 }
 
-void FDataflowToolNodeSnapshotSet::SetActiveSnapshot(int32 ActiveIndex)
+void FDataflowToolNodeSnapshotSet::SetActiveSnapshot(int32 ActiveIndex, bool bNotify)
 {
 	if (Snapshots.IsValidIndex(ActiveIndex))
 	{
-		ActiveSnapshot = ActiveIndex;
+		if (ActiveSnapshot != ActiveIndex)
+		{
+			ActiveSnapshot = ActiveIndex;
+			if (bNotify)
+			{
+				NotifyChanged();
+			}
+		}
 	}
 }
 
-FDataflowToolNodeSnapshot& FDataflowToolNodeSnapshotSet::AddSnapshot()
+FDataflowToolNodeSnapshot& FDataflowToolNodeSnapshotSet::AddSnapshot(bool bNotify)
 {
 	// if all snapshots are lock we will exceed this number 
 	constexpr int32 MaxDesiredNumberOfSnapshots = 5;
 	if (Snapshots.Num() >= MaxDesiredNumberOfSnapshots)
 	{
-		RemoveSnapshot(FindOldestUnlockedSnapShot());
+		RemoveSnapshot(FindOldestUnlockedSnapShot(), /*bNotify*/false);
 	}
 	FDataflowToolNodeSnapshot& NewSnapshot = Snapshots.AddDefaulted_GetRef();
 	NewSnapshot.Name = TEXT("Tool Snapshot");
 	ActiveSnapshot = (Snapshots.Num() - 1);
+	if (bNotify)
+	{
+		NotifyChanged();
+	}
 	return NewSnapshot;
 }
 
-bool FDataflowToolNodeSnapshotSet::RemoveSnapshot(int32 Index)
+bool FDataflowToolNodeSnapshotSet::RemoveSnapshot(int32 Index, bool bNotify)
 {
 	if (Snapshots.IsValidIndex(Index))
 	{
@@ -61,7 +80,10 @@ bool FDataflowToolNodeSnapshotSet::RemoveSnapshot(int32 Index)
 		{
 			ActiveSnapshot = ActiveSnapshot - 1;
 		}
-
+		if (bNotify)
+		{
+			NotifyChanged();
+		}
 		return true;
 	}
 	return false;
@@ -88,7 +110,33 @@ int32 FDataflowToolNodeSnapshotSet::FindOldestUnlockedSnapShot()
 
 FDataflowToolNode::FDataflowToolNode(const UE::Dataflow::FNodeParameters& InParam, FGuid InGuid)
 	: FDataflowNode(InParam, InGuid)
-{}
+{
+	RegisterNotifyChanged();
+}
+
+FDataflowToolNode::~FDataflowToolNode()
+{
+	UnregisterNotifyChanged();
+}
+
+void FDataflowToolNode::RegisterNotifyChanged()
+{
+	UnregisterNotifyChanged();
+	OnNotifyChangedHandle = Snapshots.OnSnapshotSetChanged.AddRaw(this, &FDataflowToolNode::OnSnapshotSetChanged);
+}
+
+void FDataflowToolNode::UnregisterNotifyChanged()
+{
+	if (OnNotifyChangedHandle.IsValid())
+	{
+		Snapshots.OnSnapshotSetChanged.Remove(OnNotifyChangedHandle);
+	}
+}
+
+void FDataflowToolNode::OnSnapshotSetChanged(FDataflowToolNodeSnapshotSet& Set)
+{
+	Invalidate();
+}
 
 FDataflowToolNodeSnapshot& FDataflowToolNode::AddSnapshot()
 {
@@ -102,7 +150,9 @@ bool FDataflowToolNode::RemoveSnapshot(int32 Index)
 
 void FDataflowToolNode::SwapSnapshots(FDataflowToolNodeSnapshotSet& InOutSnapshots)
 {
+	UnregisterNotifyChanged();
 	Swap(Snapshots, InOutSnapshots);
+	RegisterNotifyChanged();
 }
 
 TConstArrayView<FDataflowToolNodeSnapshot> FDataflowToolNode::GetSnapshots() const
diff --git a/Engine/Plugins/Experimental/Dataflow/Source/DataflowNodes/Public/Dataflow/DataflowToolNode.h b/Engine/Plugins/Experimental/Dataflow/Source/DataflowNodes/Public/Dataflow/DataflowToolNode.h
--- a/Engine/Plugins/Experimental/Dataflow/Source/DataflowNodes/Public/Dataflow/DataflowToolNode.h
+++ b/Engine/Plugins/Experimental/Dataflow/Source/DataflowNodes/Public/Dataflow/DataflowToolNode.h
@@ -7,6 +7,7 @@
 #include "GeometryCollection/ManagedArrayCollection.h"
 #include "InteractiveToolChange.h"
 #include "UObject/NameTypes.h"
+#include "Delegates/IDelegateInstance.h"
 
 
 #include "DataflowToolNode.generated.h"
@@ -61,12 +62,17 @@ public:
 	UE_API TConstArrayView<FDataflowToolNodeSnapshot> GetSnapshots() const;
 	UE_API const FDataflowToolNodeSnapshot* GetActiveSnapshot() const;
 
-	UE_API FDataflowToolNodeSnapshot& AddSnapshot();
-	UE_API void SetActiveSnapshot(int32 ActiveIndex);
-	UE_API bool RemoveSnapshot(int32 Index);
+	UE_API FDataflowToolNodeSnapshot& AddSnapshot(bool bNotify = false);
+	UE_API void SetActiveSnapshot(int32 ActiveIndex, bool bNotify = false);
+	UE_API bool RemoveSnapshot(int32 Index, bool bNotify = false);
 	UE_API int32 FindOldestUnlockedSnapShot();
 
+	DECLARE_MULTICAST_DELEGATE_OneParam(FOnSnapshotSetChanged, FDataflowToolNodeSnapshotSet&);
+	FOnSnapshotSetChanged OnSnapshotSetChanged;
+
 private:
+	void NotifyChanged();
+
 	UPROPERTY(EditAnywhere, Category = Snapshots)
 	int32 ActiveSnapshot = INDEX_NONE;
 
@@ -85,6 +91,7 @@ struct FDataflowToolNode : public FDataflowNode
 	FDataflowToolNode() = default;
 
 	UE_API FDataflowToolNode(const UE::Dataflow::FNodeParameters& InParam, FGuid InGuid = FGuid::NewGuid());
+	UE_API virtual ~FDataflowToolNode();
 
 public:
 	UE_API TConstArrayView<FDataflowToolNodeSnapshot> GetSnapshots() const;
@@ -115,6 +122,11 @@ protected:
 private:
 	UE_API void SwapSnapshots(FDataflowToolNodeSnapshotSet& InOutSnapshots);
 
+	void RegisterNotifyChanged();
+	void UnregisterNotifyChanged();
+	FDelegateHandle OnNotifyChangedHandle;
+	void OnSnapshotSetChanged(FDataflowToolNodeSnapshotSet& Set);
+
 	UPROPERTY(EditAnywhere, Category = Snapshots)
 	FDataflowToolNodeSnapshotSet Snapshots;
 };
