Move dataflow attachement from rigid asset to base dataflow plugin to allow wider use

#Synced-CL 51721126

diff --git a/Engine/Plugins/Experimental/ChaosRigidAsset/Source/ChaosRigidAssetEditor/ChaosRigidAssetEditor.Build.cs b/Engine/Plugins/Experimental/ChaosRigidAsset/Source/ChaosRigidAssetEditor/ChaosRigidAssetEditor.Build.cs
--- a/Engine/Plugins/Experimental/ChaosRigidAsset/Source/ChaosRigidAssetEditor/ChaosRigidAssetEditor.Build.cs
+++ b/Engine/Plugins/Experimental/ChaosRigidAsset/Source/ChaosRigidAssetEditor/ChaosRigidAssetEditor.Build.cs
@@ -18,6 +18,7 @@ public class ChaosRigidAssetEditor : ModuleRules
 				"Slate",
 				"DataflowCore",
 				"DataflowEngine",
+				"DataflowEnginePlugin",
 				"DataflowEditor",
 				"ChaosRigidAssetEngine",
 				"ChaosRigidAssetNodes",
diff --git a/Engine/Plugins/Experimental/ChaosRigidAsset/Source/ChaosRigidAssetEditor/Private/ChaosRigidAssetEditorModule.cpp b/Engine/Plugins/Experimental/ChaosRigidAsset/Source/ChaosRigidAssetEditor/Private/ChaosRigidAssetEditorModule.cpp
--- a/Engine/Plugins/Experimental/ChaosRigidAsset/Source/ChaosRigidAssetEditor/Private/ChaosRigidAssetEditorModule.cpp
+++ b/Engine/Plugins/Experimental/ChaosRigidAsset/Source/ChaosRigidAssetEditor/Private/ChaosRigidAssetEditorModule.cpp
@@ -6,12 +6,13 @@
 #include "ContentBrowserModule.h"
 #include "PhysicsEngine/PhysicsAsset.h"
 #include "Dataflow/DataflowEditor.h"
-#include "DataflowAttachment.h"
+#include "Dataflow/DataflowAttachment.h"
 #include "DataflowRendering.h"
 #include "Dataflow/DataflowRenderingFactory.h"
 #include "Features/IModularFeatures.h"
 #include "Widgets/Notifications/SNotificationList.h"
 #include "Framework/Notifications/NotificationManager.h"
+#include "PhysicsAssetDataflowAttachment.h"
 
 IMPLEMENT_MODULE(FChaosRigidAssetEditorModule, ChaosRigidAssetEditor);
 
@@ -43,14 +44,14 @@ UPhysicsAsset* GetPhysAsset(const FAssetData& InAsset)
 	return nullptr;
 }
 
-UDataflowAttachment* GetDataflowAttachment(IInterface_AssetUserData* UserDataHolder)
+UPhysicsAssetDataflowAttachment* GetDataflowAttachment(IInterface_AssetUserData* UserDataHolder)
 {
 	if(!UserDataHolder)
 	{
 		return nullptr;
 	}
 
-	return UserDataHolder->GetAssetUserData<UDataflowAttachment>();
+	return UserDataHolder->GetAssetUserData<UPhysicsAssetDataflowAttachment>();
 }
 
 bool HasDataflowAttachment(IInterface_AssetUserData* UserDataHolder)
@@ -60,7 +61,7 @@ bool HasDataflowAttachment(IInterface_AssetUserData* UserDataHolder)
 		return false;
 	}
 
-	return UserDataHolder->HasAssetUserDataOfClass(UDataflowAttachment::StaticClass());
+	return UserDataHolder->HasAssetUserDataOfClass(UPhysicsAssetDataflowAttachment::StaticClass());
 }
 
 void AddDataflow(FAssetData InAsset)
@@ -72,7 +73,7 @@ void AddDataflow(FAssetData InAsset)
 		return;
 	}
 
-	Asset->AddAssetUserData(NewObject<UDataflowAttachment>(Asset));
+	Asset->AddAssetUserData(NewObject<UPhysicsAssetDataflowAttachment>(Asset));
 }
 
 void RemoveDataflow(FAssetData InAsset)
@@ -84,7 +85,7 @@ void RemoveDataflow(FAssetData InAsset)
 		return;
 	}
 
-	Asset->RemoveUserDataOfClass(UDataflowAttachment::StaticClass());
+	Asset->RemoveUserDataOfClass(UPhysicsAssetDataflowAttachment::StaticClass());
 }
 
 void SpawnEditorFor(UPhysicsAsset* InAsset)
diff --git a/Engine/Plugins/Experimental/ChaosRigidAsset/Source/ChaosRigidAssetEngine/ChaosRigidAssetEngine.Build.cs b/Engine/Plugins/Experimental/ChaosRigidAsset/Source/ChaosRigidAssetEngine/ChaosRigidAssetEngine.Build.cs
--- a/Engine/Plugins/Experimental/ChaosRigidAsset/Source/ChaosRigidAssetEngine/ChaosRigidAssetEngine.Build.cs
+++ b/Engine/Plugins/Experimental/ChaosRigidAsset/Source/ChaosRigidAssetEngine/ChaosRigidAssetEngine.Build.cs
@@ -15,6 +15,7 @@ public class ChaosRigidAssetEngine : ModuleRules
 				"Engine",
 				"DataflowCore",
 				"DataflowEngine",
+				"DataflowEnginePlugin"
 			}
 		);
 		
diff --git a/Engine/Plugins/Experimental/ChaosRigidAsset/Source/ChaosRigidAssetEngine/Private/DataflowAttachment.cpp b/Engine/Plugins/Experimental/ChaosRigidAsset/Source/ChaosRigidAssetEngine/Private/DataflowAttachment.cpp
deleted file mode 100644
diff --git a/Engine/Plugins/Experimental/ChaosRigidAsset/Source/ChaosRigidAssetEngine/Private/PhysicsAssetDataflowAttachment.cpp b/Engine/Plugins/Experimental/ChaosRigidAsset/Source/ChaosRigidAssetEngine/Private/PhysicsAssetDataflowAttachment.cpp
new file mode 100644
--- /dev/null
+++ b/Engine/Plugins/Experimental/ChaosRigidAsset/Source/ChaosRigidAssetEngine/Private/PhysicsAssetDataflowAttachment.cpp
@@ -0,0 +1,40 @@
+// Copyright Epic Games, Inc. All Rights Reserved.
+
+#include "PhysicsAssetDataflowAttachment.h"
+#include "PhysicsAssetDataflowContent.h"
+#include "PhysicsEngine/PhysicsAsset.h"
+
+
+void UPhysicsAssetDataflowAttachment::WriteDataflowContent(const TObjectPtr<UDataflowBaseContent>& DataflowContent) const
+{
+	if(const TObjectPtr<UPhysicsAssetDataflowContent> Content = Cast<UPhysicsAssetDataflowContent>(DataflowContent))
+	{
+		Content->SetDataflowAsset(GetDataflowInstance().GetDataflowAsset());
+		Content->SetDataflowTerminal(GetDataflowInstance().GetDataflowTerminal().ToString());
+	
+		if(UPhysicsAsset* PhysAsset = GetTypedOuter<UPhysicsAsset>())
+		{
+			Content->SetSkeletalMesh(PhysAsset->GetPreviewMesh());
+			Content->SetPhysicsAsset(PhysAsset);
+		}
+	}
+}
+
+void UPhysicsAssetDataflowAttachment::ReadDataflowContent(const TObjectPtr<UDataflowBaseContent>& DataflowContent)
+{
+
+}
+
+TObjectPtr<UDataflowBaseContent> UPhysicsAssetDataflowAttachment::CreateDataflowContent()
+{
+	TObjectPtr<UDataflowBaseContent> BaseContent = NewObject<UPhysicsAssetDataflowContent>();
+	
+	BaseContent->SetDataflowOwner(this);
+	BaseContent->SetTerminalAsset(this);
+	
+	WriteDataflowContent(BaseContent);
+	
+	return BaseContent;
+}
+
+
diff --git a/Engine/Plugins/Experimental/ChaosRigidAsset/Source/ChaosRigidAssetEngine/Public/DataflowAttachment.h b/Engine/Plugins/Experimental/ChaosRigidAsset/Source/ChaosRigidAssetEngine/Public/DataflowAttachment.h
deleted file mode 100644
diff --git a/Engine/Plugins/Experimental/ChaosRigidAsset/Source/ChaosRigidAssetEngine/Public/PhysicsAssetDataflowAttachment.h b/Engine/Plugins/Experimental/ChaosRigidAsset/Source/ChaosRigidAssetEngine/Public/PhysicsAssetDataflowAttachment.h
new file mode 100644
--- /dev/null
+++ b/Engine/Plugins/Experimental/ChaosRigidAsset/Source/ChaosRigidAssetEngine/Public/PhysicsAssetDataflowAttachment.h
@@ -0,0 +1,27 @@
+// Copyright Epic Games, Inc. All Rights Reserved.
+
+#pragma once
+
+#include "Dataflow/DataflowAttachment.h"
+#include "PhysicsAssetDataflowAttachment.generated.h"
+
+UCLASS()
+class UPhysicsAssetDataflowAttachment : public UDataflowAttachment
+{
+	GENERATED_BODY()
+
+public:
+
+	// IDataflowContentOwner
+	void WriteDataflowContent(const TObjectPtr<UDataflowBaseContent>& DataflowContent) const override;
+	void ReadDataflowContent(const TObjectPtr<UDataflowBaseContent>& DataflowContent) override;
+	////////////////////////////////////////////////////////////////////////////////////////////////////////////////
+
+protected:
+
+	// IDataflowContentOwner
+	TObjectPtr<UDataflowBaseContent> CreateDataflowContent() override;
+	////////////////////////////////////////////////////////////////////////////////////////////////////////////////
+
+private:
+};
diff --git a/Engine/Plugins/Experimental/ChaosRigidAsset/Source/ChaosRigidAssetNodes/ChaosRigidAssetNodes.Build.cs b/Engine/Plugins/Experimental/ChaosRigidAsset/Source/ChaosRigidAssetNodes/ChaosRigidAssetNodes.Build.cs
--- a/Engine/Plugins/Experimental/ChaosRigidAsset/Source/ChaosRigidAssetNodes/ChaosRigidAssetNodes.Build.cs
+++ b/Engine/Plugins/Experimental/ChaosRigidAsset/Source/ChaosRigidAssetNodes/ChaosRigidAssetNodes.Build.cs
@@ -15,6 +15,7 @@ public class ChaosRigidAssetNodes : ModuleRules
 				"Engine",
 				"DataflowCore",
 				"DataflowEngine",
+				"DataflowEnginePlugin",
 				"GeometryCore",
 				"DynamicMesh",
 				"ChaosRigidAssetEngine"
diff --git a/Engine/Plugins/Experimental/ChaosRigidAsset/Source/ChaosRigidAssetNodes/Private/PhysicsAssetDataflowNodes.cpp b/Engine/Plugins/Experimental/ChaosRigidAsset/Source/ChaosRigidAssetNodes/Private/PhysicsAssetDataflowNodes.cpp
--- a/Engine/Plugins/Experimental/ChaosRigidAsset/Source/ChaosRigidAssetNodes/Private/PhysicsAssetDataflowNodes.cpp
+++ b/Engine/Plugins/Experimental/ChaosRigidAsset/Source/ChaosRigidAssetNodes/Private/PhysicsAssetDataflowNodes.cpp
@@ -5,7 +5,7 @@
 #include "PhysicsEngine/PhysicsAsset.h"
 #include "PhysicsEngine/SkeletalBodySetup.h"
 #include "Misc/WildcardString.h"
-#include "DataflowAttachment.h"
+#include "Dataflow/DataflowAttachment.h"
 #include "Dataflow/DataflowDebugDrawInterface.h"
 #include "Dataflow/DataflowDebugDrawObject.h"
 #include "Animation/Skeleton.h"
diff --git a/Engine/Plugins/Experimental/ChaosRigidAsset/Source/ChaosRigidAssetNodes/Private/RigidDataflowNode.h b/Engine/Plugins/Experimental/ChaosRigidAsset/Source/ChaosRigidAssetNodes/Private/RigidDataflowNode.h
--- a/Engine/Plugins/Experimental/ChaosRigidAsset/Source/ChaosRigidAssetNodes/Private/RigidDataflowNode.h
+++ b/Engine/Plugins/Experimental/ChaosRigidAsset/Source/ChaosRigidAssetNodes/Private/RigidDataflowNode.h
@@ -5,7 +5,7 @@
 #include "Dataflow/DataflowNode.h"
 #include "Dataflow/DataflowObjectInterface.h"
 #include "ShapeElemNodes.h"
-#include "DataflowAttachment.h"
+#include "Dataflow/DataflowAttachment.h"
 
 #include "RigidDataflowNode.generated.h"
 
diff --git a/Engine/Plugins/Experimental/Dataflow/Source/DataflowEnginePlugin/Private/Dataflow/DataflowAttachment.cpp b/Engine/Plugins/Experimental/Dataflow/Source/DataflowEnginePlugin/Private/Dataflow/DataflowAttachment.cpp
new file mode 100644
--- /dev/null
+++ b/Engine/Plugins/Experimental/Dataflow/Source/DataflowEnginePlugin/Private/Dataflow/DataflowAttachment.cpp
@@ -0,0 +1,20 @@
+// Copyright Epic Games, Inc. All Rights Reserved.
+
+#include "Dataflow/DataflowAttachment.h"
+#include "PhysicsEngine/PhysicsAsset.h"
+
+const FDataflowInstance& UDataflowAttachment::GetDataflowInstance() const
+{
+	return Instance;
+}
+
+FDataflowInstance& UDataflowAttachment::GetDataflowInstance()
+{
+	return Instance;
+}
+
+void UDataflowAttachment::PostLoad()
+{
+	Super::PostLoad();
+	Instance.PostLoad();
+}
diff --git a/Engine/Plugins/Experimental/Dataflow/Source/DataflowEnginePlugin/Public/Dataflow/DataflowAttachment.h b/Engine/Plugins/Experimental/Dataflow/Source/DataflowEnginePlugin/Public/Dataflow/DataflowAttachment.h
new file mode 100644
--- /dev/null
+++ b/Engine/Plugins/Experimental/Dataflow/Source/DataflowEnginePlugin/Public/Dataflow/DataflowAttachment.h
@@ -0,0 +1,54 @@
+// Copyright Epic Games, Inc. All Rights Reserved.
+
+#pragma once
+
+#include "Engine/AssetUserData.h"
+#include "Dataflow/DataflowContent.h"
+#include "Dataflow/DataflowInstance.h"
+#include "UObject/ObjectPtr.h"
+
+#include "DataflowAttachment.generated.h"
+
+#define UE_API DATAFLOWENGINEPLUGIN_API
+
+/**
+ * Generic dataflow instance attachment implemented as a user asset data object so it can be embedded inside
+ * any asset that is a user-data container. This allows assets in modules that do not, or cannot know about
+ * dataflow (e.g. Engine module assets) to contain a dataflow attachment.
+ * Using a modular feature implementation this data can be grabbed as an asset opens and a plugin specific
+ * implementation can be used for the asset editor (see FDataflowPhysicsAssetEditorOverride)
+ */
+UCLASS(MinimalApi, Abstract)
+class UDataflowAttachment : public UAssetUserData, public IDataflowContentOwner, public IDataflowInstanceInterface
+{
+	GENERATED_BODY()
+
+public:
+
+	// IDataflowContentOwner
+	void WriteDataflowContent(const TObjectPtr<UDataflowBaseContent>& DataflowContent) const override PURE_VIRTUAL( UDataflowAttachment::WriteDataflowContent, );
+	void ReadDataflowContent(const TObjectPtr<UDataflowBaseContent>& DataflowContent) override PURE_VIRTUAL( UDataflowAttachment::ReadDataflowContent, );
+	////////////////////////////////////////////////////////////////////////////////////////////////////////////////
+	
+	// IDataflowInstanceInterface
+	UE_API const FDataflowInstance& GetDataflowInstance() const override;
+	UE_API FDataflowInstance& GetDataflowInstance() override;
+	////////////////////////////////////////////////////////////////////////////////////////////////////////////////
+
+	// UObject
+	UE_API virtual void PostLoad() override;
+	////////////////////////////////////////////////////////////////////////////////////////////////////////////////
+private:
+
+	// Dataflow instance for the attachment when editing.
+	UPROPERTY()
+	FDataflowInstance Instance;
+
+protected:
+	
+	// IDataflowContentOwner
+	TObjectPtr<UDataflowBaseContent> CreateDataflowContent() override PURE_VIRTUAL( UDataflowAttachment::CreateDataflowContent, return nullptr; );
+	////////////////////////////////////////////////////////////////////////////////////////////////////////////////
+};
+
+#undef UE_API
