Dataflow : better messageing =for when a tool does not launch properly

#Synced-CL 52092384

diff --git a/Engine/Plugins/Experimental/Dataflow/Source/DataflowEditor/Private/Dataflow/DataflowEditorMode.cpp b/Engine/Plugins/Experimental/Dataflow/Source/DataflowEditor/Private/Dataflow/DataflowEditorMode.cpp
--- a/Engine/Plugins/Experimental/Dataflow/Source/DataflowEditor/Private/Dataflow/DataflowEditorMode.cpp
+++ b/Engine/Plugins/Experimental/Dataflow/Source/DataflowEditor/Private/Dataflow/DataflowEditorMode.cpp
@@ -763,7 +763,14 @@ void UDataflowEditorMode::ModeTick(float DeltaTime)
 				if (Selection && Selection->Num() == 0)
 				{
 					// TODO : find a more accurate message ?
-					UserMessage = FText::Format(LOCTEXT("ToolNeedoneComponent", "You must select a component to launch the {0} tool"), Command->GetLabel());
+					if (!ConstructionScenePtr->HasRenderableGeometry())
+					{
+						UserMessage = FText::Format(LOCTEXT("NoComponentAvailable", "{0} tool cannot launch because there's no component in the construction view - check the node settings"), Command->GetLabel());
+					}
+					else
+					{
+						UserMessage = FText::Format(LOCTEXT("ToolNeedOneComponent", "You must select a component to launch the {0} tool"), Command->GetLabel());
+					}
 				}
 				else
 				{
