Dataflow : move vertex property group classes to their own files
- Also add ways to return information about each target group

#Synced-CL 52197375

diff --git a/Engine/Plugins/Experimental/Dataflow/Source/DataflowNodes/Private/Dataflow/DataflowCollectionAddScalarVertexPropertyNode.cpp b/Engine/Plugins/Experimental/Dataflow/Source/DataflowNodes/Private/Dataflow/DataflowCollectionAddScalarVertexPropertyNode.cpp
--- a/Engine/Plugins/Experimental/Dataflow/Source/DataflowNodes/Private/Dataflow/DataflowCollectionAddScalarVertexPropertyNode.cpp
+++ b/Engine/Plugins/Experimental/Dataflow/Source/DataflowNodes/Private/Dataflow/DataflowCollectionAddScalarVertexPropertyNode.cpp
@@ -11,82 +11,6 @@
 
 #define LOCTEXT_NAMESPACE "DataflowCollectionAddScalarVertexProperty"
 
-//
-// FDataflowAddScalarVertexPropertyCallbackRegistry
-//
-
-FDataflowAddScalarVertexPropertyCallbackRegistry& FDataflowAddScalarVertexPropertyCallbackRegistry::Get()
-{
-	return TLazySingleton<FDataflowAddScalarVertexPropertyCallbackRegistry>::Get();
-}
-
-void FDataflowAddScalarVertexPropertyCallbackRegistry::RegisterCallbacks(TUniquePtr<IDataflowAddScalarVertexPropertyCallbacks>&& Callbacks)
-{
-	AllCallbacks.Add(Callbacks->GetName(), MoveTemp(Callbacks));
-}
-
-void FDataflowAddScalarVertexPropertyCallbackRegistry::DeregisterCallbacks(const FName& CallbacksName)
-{
-	AllCallbacks.Remove(CallbacksName);
-}
-
-TArray<FName> FDataflowAddScalarVertexPropertyCallbackRegistry::GetTargetGroupNames() const
-{
-	TArray<FName> UniqueNames;
-
-	for (const TPair<FName, TUniquePtr<IDataflowAddScalarVertexPropertyCallbacks>>& CallbacksEntry : AllCallbacks)
-	{
-		for (const FName& GroupName : CallbacksEntry.Value->GetTargetGroupNames())
-		{
-			UniqueNames.AddUnique(GroupName);
-		}
-	}
-	return UniqueNames;
-}
-
-TArray<UE::Dataflow::FRenderingParameter> FDataflowAddScalarVertexPropertyCallbackRegistry::GetRenderingParameters() const
-{
-	TArray<UE::Dataflow::FRenderingParameter> UniqueParameters;
-
-	for (const TPair<FName, TUniquePtr<IDataflowAddScalarVertexPropertyCallbacks>>& CallbacksEntry : AllCallbacks)
-	{
-		for (const UE::Dataflow::FRenderingParameter& RenderingParameter : CallbacksEntry.Value->GetRenderingParameters())
-		{
-			UniqueParameters.AddUnique(RenderingParameter);
-		}
-	}
-	return UniqueParameters;
-}
-
-TArray<UE::Dataflow::FRenderingParameter> FDataflowAddScalarVertexPropertyCallbackRegistry::GetRenderingParameters(const FName& TargetGroup) const
-{
-	TArray<UE::Dataflow::FRenderingParameter> UniqueParameters;
-
-	for (const TPair<FName, TUniquePtr<IDataflowAddScalarVertexPropertyCallbacks>>& CallbacksEntry : AllCallbacks)
-	{
-		const TArray<UE::Dataflow::FRenderingParameter> RenderingParameters = CallbacksEntry.Value->GetRenderingParameters();
-		const TArray<FName> TargetGroups = CallbacksEntry.Value->GetTargetGroupNames();
-		if(RenderingParameters.Num() == TargetGroups.Num())
-		{
-			for(int32 TargetIndex = 0, NumTargets = TargetGroups.Num(); TargetIndex < NumTargets; ++TargetIndex)
-			{
-				if(TargetGroups[TargetIndex] == TargetGroup)
-				{
-					UniqueParameters.AddUnique(RenderingParameters[TargetIndex]);
-				}
-			}
-		}
-		else if(TargetGroups.Find(TargetGroup) != INDEX_NONE)
-		{
-			for (const UE::Dataflow::FRenderingParameter& RenderingParameter : CallbacksEntry.Value->GetRenderingParameters())
-			{
-				UniqueParameters.AddUnique(RenderingParameter);
-			}
-		}
-	}
-	return UniqueParameters;
-}
-
 
 ////////////////////////////////////////////////////////////////////////////////////////////////////////////////////////////////////////////////////////////
 // 
diff --git a/Engine/Plugins/Experimental/Dataflow/Source/DataflowNodes/Private/Dataflow/DataflowVertexPropertyGroup.cpp b/Engine/Plugins/Experimental/Dataflow/Source/DataflowNodes/Private/Dataflow/DataflowVertexPropertyGroup.cpp
new file mode 100644
--- /dev/null
+++ b/Engine/Plugins/Experimental/Dataflow/Source/DataflowNodes/Private/Dataflow/DataflowVertexPropertyGroup.cpp
@@ -0,0 +1,99 @@
+// Copyright Epic Games, Inc. All Rights Reserved.
+
+#include "Dataflow/DataflowVertexPropertyGroup.h"
+#include "Misc/LazySingleton.h"
+
+#define LOCTEXT_NAMESPACE "DataflowVertexPropertyGroup"
+
+FDataflowAddScalarVertexPropertyCallbackRegistry& FDataflowAddScalarVertexPropertyCallbackRegistry::Get()
+{
+	return TLazySingleton<FDataflowAddScalarVertexPropertyCallbackRegistry>::Get();
+}
+
+void FDataflowAddScalarVertexPropertyCallbackRegistry::RegisterCallbacks(TUniquePtr<IDataflowAddScalarVertexPropertyCallbacks>&& Callbacks)
+{
+	AllCallbacks.Add(Callbacks->GetName(), MoveTemp(Callbacks));
+}
+
+void FDataflowAddScalarVertexPropertyCallbackRegistry::DeregisterCallbacks(const FName& CallbacksName)
+{
+	AllCallbacks.Remove(CallbacksName);
+}
+
+TArray<FName> FDataflowAddScalarVertexPropertyCallbackRegistry::GetTargetGroupNames() const
+{
+	TArray<FName> UniqueNames;
+
+	for (const TPair<FName, TUniquePtr<IDataflowAddScalarVertexPropertyCallbacks>>& CallbacksEntry : AllCallbacks)
+	{
+		for (const FName& GroupName : CallbacksEntry.Value->GetTargetGroupNames())
+		{
+			UniqueNames.AddUnique(GroupName);
+		}
+	}
+	return UniqueNames;
+}
+
+FDataflowAddScalarVertexPropertyCallbackRegistry::FTargetGroupInfo FDataflowAddScalarVertexPropertyCallbackRegistry::GetTargetGroupInfo(FName TargetGroup) const
+{
+	TArray<FTargetGroupInfo> OutInfos;
+	for (const TPair<FName, TUniquePtr<IDataflowAddScalarVertexPropertyCallbacks>>& CallbacksEntry : AllCallbacks)
+	{
+		CallbacksEntry.Value->GetTargetGroupInfos(OutInfos);
+	}
+	for (const FTargetGroupInfo& Info : OutInfos)
+	{
+		if (Info.TargetGroup == TargetGroup)
+		{
+			return Info;
+		}
+	}
+	return {};
+}
+
+TArray<UE::Dataflow::FRenderingParameter> FDataflowAddScalarVertexPropertyCallbackRegistry::GetRenderingParameters() const
+{
+	TArray<UE::Dataflow::FRenderingParameter> UniqueParameters;
+
+	for (const TPair<FName, TUniquePtr<IDataflowAddScalarVertexPropertyCallbacks>>& CallbacksEntry : AllCallbacks)
+	{
+		for (const UE::Dataflow::FRenderingParameter& RenderingParameter : CallbacksEntry.Value->GetRenderingParameters())
+		{
+			UniqueParameters.AddUnique(RenderingParameter);
+		}
+	}
+	return UniqueParameters;
+}
+
+TArray<UE::Dataflow::FRenderingParameter> FDataflowAddScalarVertexPropertyCallbackRegistry::GetRenderingParameters(const FName& TargetGroup) const
+{
+	TArray<UE::Dataflow::FRenderingParameter> UniqueParameters;
+
+	for (const TPair<FName, TUniquePtr<IDataflowAddScalarVertexPropertyCallbacks>>& CallbacksEntry : AllCallbacks)
+	{
+		const TArray<UE::Dataflow::FRenderingParameter> RenderingParameters = CallbacksEntry.Value->GetRenderingParameters();
+		const TArray<FName> TargetGroups = CallbacksEntry.Value->GetTargetGroupNames();
+		if (RenderingParameters.Num() == TargetGroups.Num())
+		{
+			for (int32 TargetIndex = 0, NumTargets = TargetGroups.Num(); TargetIndex < NumTargets; ++TargetIndex)
+			{
+				if (TargetGroups[TargetIndex] == TargetGroup)
+				{
+					UniqueParameters.AddUnique(RenderingParameters[TargetIndex]);
+				}
+			}
+		}
+		else if (TargetGroups.Find(TargetGroup) != INDEX_NONE)
+		{
+			for (const UE::Dataflow::FRenderingParameter& RenderingParameter : CallbacksEntry.Value->GetRenderingParameters())
+			{
+				UniqueParameters.AddUnique(RenderingParameter);
+			}
+		}
+	}
+	return UniqueParameters;
+}
+
+
+#undef LOCTEXT_NAMESPACE 
+
diff --git a/Engine/Plugins/Experimental/Dataflow/Source/DataflowNodes/Public/Dataflow/DataflowCollectionAddScalarVertexPropertyNode.h b/Engine/Plugins/Experimental/Dataflow/Source/DataflowNodes/Public/Dataflow/DataflowCollectionAddScalarVertexPropertyNode.h
--- a/Engine/Plugins/Experimental/Dataflow/Source/DataflowNodes/Public/Dataflow/DataflowCollectionAddScalarVertexPropertyNode.h
+++ b/Engine/Plugins/Experimental/Dataflow/Source/DataflowNodes/Public/Dataflow/DataflowCollectionAddScalarVertexPropertyNode.h
@@ -7,54 +7,12 @@
 #include "GeometryCollection/GeometryCollection.h"
 #include "GeometryCollection/ManagedArrayCollection.h"
 #include "Dataflow/DataflowCollectionAttributeKeyNodes.h"
+#include "Dataflow/DataflowVertexPropertyGroup.h"
 
 #include "DataflowCollectionAddScalarVertexPropertyNode.generated.h"
 
 #define UE_API DATAFLOWNODES_API
 
-
-class IDataflowAddScalarVertexPropertyCallbacks
-{
-public:
-	virtual ~IDataflowAddScalarVertexPropertyCallbacks() = default;
-	virtual FName GetName() const = 0;
-	virtual TArray<FName> GetTargetGroupNames() const = 0;
-	virtual TArray<UE::Dataflow::FRenderingParameter> GetRenderingParameters() const = 0;
-};
-
-class FDataflowAddScalarVertexPropertyCallbackRegistry
-{
-public:
-
-	DATAFLOWNODES_API static FDataflowAddScalarVertexPropertyCallbackRegistry& Get();
-
-	DATAFLOWNODES_API void RegisterCallbacks(TUniquePtr<IDataflowAddScalarVertexPropertyCallbacks>&& Callbacks);
-
-	DATAFLOWNODES_API void DeregisterCallbacks(const FName& CallbacksName);
-
-	DATAFLOWNODES_API TArray<FName> GetTargetGroupNames() const;
-
-	DATAFLOWNODES_API TArray<UE::Dataflow::FRenderingParameter> GetRenderingParameters() const;
-	
-	DATAFLOWNODES_API TArray<UE::Dataflow::FRenderingParameter> GetRenderingParameters(const FName& TargetGroup) const;
-
-private:
-
-	TMap<FName, TUniquePtr<IDataflowAddScalarVertexPropertyCallbacks>> AllCallbacks;
-};
-
-/*
-* Custom type so that we can use property type customization
-*/
-USTRUCT()
-struct FScalarVertexPropertyGroup
-{
-	GENERATED_USTRUCT_BODY()
-
-	UPROPERTY(EditAnywhere, Category = "Vertex Group")
-	FName Name = FGeometryCollection::VerticesGroup;
-};
-
 /** How the map stored on the AddWeightMapNode should be applied to an existing map. If no map exists, it is treated as zero.*/
 UENUM()
 enum class EDataflowWeightMapOverrideType : uint8
diff --git a/Engine/Plugins/Experimental/Dataflow/Source/DataflowNodes/Public/Dataflow/DataflowVertexPropertyGroup.h b/Engine/Plugins/Experimental/Dataflow/Source/DataflowNodes/Public/Dataflow/DataflowVertexPropertyGroup.h
new file mode 100644
--- /dev/null
+++ b/Engine/Plugins/Experimental/Dataflow/Source/DataflowNodes/Public/Dataflow/DataflowVertexPropertyGroup.h
@@ -0,0 +1,80 @@
+// Copyright Epic Games, Inc. All Rights Reserved.
+
+#pragma once 
+
+#include "Containers/Array.h"
+#include "Containers/Map.h"
+#include "Dataflow/DataflowNode.h"
+#include "GeometryCollection/GeometryCollection.h"
+#include "UObject/NameTypes.h"
+
+#include "DataflowVertexPropertyGroup.generated.h"
+
+#define UE_API DATAFLOWNODES_API
+
+/**
+* interface for defining vertex property groups
+* this is used by nodes to know how to access collection data from a specific vertex group
+*/
+class IDataflowAddScalarVertexPropertyCallbacks
+{
+public:
+	using FAttributeKey = FDataflowNode::FAttributeKey;
+
+	struct FTargetGroupInfo
+	{
+		FName TargetGroup;
+		FAttributeKey PositionAttributeKey;
+		FAttributeKey IndicesAttributeKey;
+		FAttributeKey MappingFrom2DTo3DAttributeKey;
+	};
+
+	virtual ~IDataflowAddScalarVertexPropertyCallbacks() = default;
+	virtual FName GetName() const = 0;
+	virtual TArray<FName> GetTargetGroupNames() const = 0;
+	virtual void GetTargetGroupInfos(TArray<FTargetGroupInfo>& OutInfos) const {};
+	virtual TArray<UE::Dataflow::FRenderingParameter> GetRenderingParameters() const = 0;
+};
+
+/**
+* Registry for implementations of IDataflowAddScalarVertexPropertyCallbacks
+*/
+class FDataflowAddScalarVertexPropertyCallbackRegistry
+{
+public:
+
+	using FTargetGroupInfo = IDataflowAddScalarVertexPropertyCallbacks::FTargetGroupInfo;
+
+	UE_API static FDataflowAddScalarVertexPropertyCallbackRegistry& Get();
+
+	UE_API void RegisterCallbacks(TUniquePtr<IDataflowAddScalarVertexPropertyCallbacks>&& Callbacks);
+
+	UE_API void DeregisterCallbacks(const FName& CallbacksName);
+
+	UE_API TArray<FName> GetTargetGroupNames() const;
+
+	UE_API FTargetGroupInfo GetTargetGroupInfo(FName TargetGroup) const;
+
+	UE_API TArray<UE::Dataflow::FRenderingParameter> GetRenderingParameters() const;
+
+	UE_API TArray<UE::Dataflow::FRenderingParameter> GetRenderingParameters(const FName& TargetGroup) const;
+
+private:
+
+	TMap<FName, TUniquePtr<IDataflowAddScalarVertexPropertyCallbacks>> AllCallbacks;
+};
+
+/*
+* Custom type so that we can use property type customization
+* Used by Dataflow nodes as a property type - displays as a combobox with the registered groups
+*/
+USTRUCT()
+struct FScalarVertexPropertyGroup
+{
+	GENERATED_USTRUCT_BODY()
+
+	UPROPERTY(EditAnywhere, Category = "Vertex Group")
+	FName Name = FGeometryCollection::VerticesGroup;
+};
+
+#undef UE_API 
