Dataflow add renderable type for physics asset

#Synced-CL 52199989

diff --git a/Engine/Plugins/ChaosClothAssetDataflowNodes/Source/ChaosClothAssetDataflowNodes/Private/ChaosClothAsset/SetPhysicsAssetNode.cpp b/Engine/Plugins/ChaosClothAssetDataflowNodes/Source/ChaosClothAssetDataflowNodes/Private/ChaosClothAsset/SetPhysicsAssetNode.cpp
--- a/Engine/Plugins/ChaosClothAssetDataflowNodes/Source/ChaosClothAssetDataflowNodes/Private/ChaosClothAsset/SetPhysicsAssetNode.cpp
+++ b/Engine/Plugins/ChaosClothAssetDataflowNodes/Source/ChaosClothAssetDataflowNodes/Private/ChaosClothAsset/SetPhysicsAssetNode.cpp
@@ -16,16 +16,18 @@ FChaosClothAssetSetPhysicsAssetNode::FChaosClothAssetSetPhysicsAssetNode(const U
 	RegisterInputConnection(&Collection);
 	RegisterInputConnection(&PhysicsAsset);
 	RegisterOutputConnection(&Collection, &Collection);
+	RegisterOutputConnection(&PhysicsAsset, &PhysicsAsset)
+		.SetCanHidePin(true).SetPinIsHidden(true); // we only need the output for the construction view rendering
 }
 
 void FChaosClothAssetSetPhysicsAssetNode::Evaluate(UE::Dataflow::FContext& Context, const FDataflowOutput* Out) const
 {
-	if (Out->IsA<FManagedArrayCollection>(&Collection))
+	if (Out->IsA(&Collection))
 	{
 		using namespace UE::Chaos::ClothAsset;
 
 		// Evaluate in collection
-		FManagedArrayCollection InCollection = GetValue<FManagedArrayCollection>(Context, &Collection);
+		FManagedArrayCollection InCollection = GetValue(Context, &Collection);
 		const TSharedRef<FManagedArrayCollection> ClothCollection = MakeShared<FManagedArrayCollection>(MoveTemp(InCollection));
 
 		if (FCollectionClothFacade(ClothCollection).IsValid())  // Can only act on the collection if it is a valid cloth collection
@@ -37,6 +39,10 @@ void FChaosClothAssetSetPhysicsAssetNode::Evaluate(UE::Dataflow::FContext& Conte
 
 		SetValue(Context, MoveTemp(*ClothCollection), &Collection);
 	}
+	else if (Out->IsA(&PhysicsAsset))
+	{
+		SafeForwardInput(Context, &PhysicsAsset, &PhysicsAsset);
+	}
 }
 
 #undef LOCTEXT_NAMESPACE
diff --git a/Engine/Plugins/ChaosClothAssetDataflowNodes/Source/ChaosClothAssetDataflowNodes/Public/ChaosClothAsset/SetPhysicsAssetNode.h b/Engine/Plugins/ChaosClothAssetDataflowNodes/Source/ChaosClothAssetDataflowNodes/Public/ChaosClothAsset/SetPhysicsAssetNode.h
--- a/Engine/Plugins/ChaosClothAssetDataflowNodes/Source/ChaosClothAssetDataflowNodes/Public/ChaosClothAsset/SetPhysicsAssetNode.h
+++ b/Engine/Plugins/ChaosClothAssetDataflowNodes/Source/ChaosClothAssetDataflowNodes/Public/ChaosClothAsset/SetPhysicsAssetNode.h
@@ -22,7 +22,7 @@ public:
 	FManagedArrayCollection Collection;
 
 	/** The physics asset to assign to the Cloth Collection. */
-	UPROPERTY(EditAnywhere, Category = "Set Physics Asset", Meta = (DataflowInput))
+	UPROPERTY(EditAnywhere, Category = "Set Physics Asset", Meta = (DataflowInput, DataflowOutput, DataflowPassthrough = "PhysicsAsset"))
 	TObjectPtr<UPhysicsAsset> PhysicsAsset;
 
 	FChaosClothAssetSetPhysicsAssetNode(const UE::Dataflow::FNodeParameters& InParam, FGuid InGuid = FGuid::NewGuid());
diff --git a/Engine/Plugins/Experimental/Dataflow/Source/DataflowEditor/Private/Dataflow/DataflowEditorModule.cpp b/Engine/Plugins/Experimental/Dataflow/Source/DataflowEditor/Private/Dataflow/DataflowEditorModule.cpp
--- a/Engine/Plugins/Experimental/Dataflow/Source/DataflowEditor/Private/Dataflow/DataflowEditorModule.cpp
+++ b/Engine/Plugins/Experimental/Dataflow/Source/DataflowEditor/Private/Dataflow/DataflowEditorModule.cpp
@@ -49,6 +49,7 @@
 #include "DataflowRendering/DataflowTextureRenderableType.h"
 #include "DataflowRendering/DataflowSamplerRenderableType.h"
 #include "DataflowRendering/DataflowPlaneRenderableType.h"
+#include "DataflowRendering/DataflowPhysicsAssetRenderableType.h"
 
 #include "ContentBrowserModule.h"
 #include "Modules/ModuleManager.h"
@@ -364,6 +365,7 @@ void FDataflowEditorModule::StartupModule()
 	UE::Dataflow::Private::RegisterTextureRenderableTypes();
 	UE::Dataflow::Private::RegisterSamplerRenderableTypes();
 	UE::Dataflow::Private::RegisterPlaneRenderableTypes();
+	UE::Dataflow::Private::RegisterPhysicsAssetRenderableTypes();
 
 	// register an browser extension to show when an asset is bound to a dataflow asset
 	if (FContentBrowserModule* ContentBrowserModule = FModuleManager::Get().GetModulePtr<FContentBrowserModule>(TEXT("ContentBrowser")))
diff --git a/Engine/Plugins/Experimental/Dataflow/Source/DataflowEditor/Private/DataflowRendering/DataflowPhysicsAssetRenderableType.cpp b/Engine/Plugins/Experimental/Dataflow/Source/DataflowEditor/Private/DataflowRendering/DataflowPhysicsAssetRenderableType.cpp
new file mode 100644
--- /dev/null
+++ b/Engine/Plugins/Experimental/Dataflow/Source/DataflowEditor/Private/DataflowRendering/DataflowPhysicsAssetRenderableType.cpp
@@ -0,0 +1,392 @@
+// Copyright Epic Games, Inc. All Rights Reserved.
+
+#include "DataflowRendering/DataflowPhysicsAssetRenderableType.h"
+
+#include "AnimationRuntime.h"
+#include "Components/DynamicMeshComponent.h"
+#include "Components/StaticMeshComponent.h"
+#include "Dataflow/DataflowEngineUtil.h"
+#include "Dataflow/DataflowRenderingViewMode.h"
+#include "DataflowRendering/DataflowRenderableComponents.h"
+#include "DataflowRendering/DataflowRenderableType.h"
+#include "DataflowRendering/DataflowRenderableTypeInstance.h"
+#include "DataflowRendering/DataflowRenderableTypeRegistry.h"
+#include "DataflowRendering/DataflowRenderableTypeUtils.h"
+#include "DynamicMesh/DynamicMesh3.h"
+#include "Engine/SkeletalMesh.h"
+#include "Generators/CapsuleGenerator.h"
+#include "Materials/MaterialInstanceDynamic.h"
+#include "PhysicsEngine/AggregateGeom.h"
+#include "PhysicsEngine/PhysicsAsset.h"
+#include "PhysicsEngine/SkeletalBodySetup.h"
+#include "ReferenceSkeleton.h"
+#include "UObject/ObjectPtr.h"
+
+UDataflowPhysicsAssetRenderSettings::UDataflowPhysicsAssetRenderSettings(const FObjectInitializer& ObjectInitializer)
+{
+	ColorSettings.Color = FLinearColor(0.5f, 1.0f, 1.0f);
+	ColorSettings.Transparency = 0.5;
+}
+
+namespace UE::Dataflow::Private
+{
+namespace
+{
+
+// Compute component-space reference pose transforms.
+// Returns an empty array if PreviewSkeletalMesh is null.
+static TArray<FTransform> ComputeReferencePoseCS(const USkeletalMesh* Mesh)
+{
+	TArray<FTransform> Out;
+	if (Mesh)
+	{
+		const FReferenceSkeleton& RefSkel = Mesh->GetRefSkeleton();
+		FAnimationRuntime::FillUpComponentSpaceTransforms(RefSkel, RefSkel.GetRefBonePose(), Out);
+	}
+	return Out;
+}
+
+// create a material instance based on the settings 
+static UMaterialInstanceDynamic* CreateMaterialInstance(UObject* Owner, const UDataflowPhysicsAssetRenderSettings* Settings)
+{
+	UMaterialInstanceDynamic* MaterialInstance = nullptr;
+
+	UMaterialInterface* Material = Settings && Settings->ColorSettings.bWireframe
+		? UE::Dataflow::RenderMaterial::GetDataflowColorWireframeMaterial()
+		: UE::Dataflow::RenderMaterial::GetDataflowColorMaterial()
+		;
+	if (Material)
+	{
+		MaterialInstance = UMaterialInstanceDynamic::Create(Material, Owner);
+		if (MaterialInstance)
+		{
+			static const FName ColorParameterName("Color");
+			static const FName TransparencyParameterName("Transparency");
+			MaterialInstance->SetVectorParameterValue(ColorParameterName, Settings? Settings->ColorSettings.Color: FLinearColor::White);
+			MaterialInstance->SetScalarParameterValue(TransparencyParameterName, Settings ? Settings->ColorSettings.Transparency: 0.f);
+		}
+	}
+	return MaterialInstance;
+}
+
+// Build a uniform capsule mesh centered at the origin, aligned along +Z.
+// Total height = Length + 2 * Radius.
+// FCapsuleGenerator places its cylinder from Z=0 to Z=SegmentLength;
+// we shift the vertices to center the mesh at Z=0.
+static UE::Geometry::FDynamicMesh3 BuildCapsuleMesh(float Radius, float Length)
+{
+	UE::Geometry::FCapsuleGenerator Generator;
+	Generator.Radius         = Radius;
+	Generator.SegmentLength  = Length;
+	Generator.NumHemisphereArcSteps = 6;
+	Generator.NumCircleSteps        = 12;
+	Generator.Generate();
+
+	UE::Geometry::FDynamicMesh3 Mesh(&Generator);
+
+	const double OffsetZ = -(double)Length * 0.5;
+	for (int32 VID : Mesh.VertexIndicesItr())
+	{
+		FVector3d V = Mesh.GetVertex(VID);
+		V.Z += OffsetZ;
+		Mesh.SetVertex(VID, V);
+	}
+	return Mesh;
+}
+
+// Build a tapered capsule mesh centered at the origin, aligned along +Z.
+// Radius0 = bottom end radius, Radius1 = top end radius.
+// Length = cylinder section length (total height = Length + Radius0 + Radius1).
+static UE::Geometry::FDynamicMesh3 BuildTaperedCapsuleMesh(float Radius0, float Radius1, float Length,
+	int32 ArcSteps = 6, int32 CircleSteps = 12)
+{
+	UE::Geometry::FDynamicMesh3 Mesh;
+	const double HalfLength = (double)Length * 0.5;
+	const double DTheta     = 2.0 * UE_DOUBLE_PI / CircleSteps;
+
+	// Profile: (radius, z) pairs bottom to top.
+	// Bottom hemisphere: phi sweeps 0 (pole) -> PI/2 (equator)
+	// Straight taper section: two equatorial rings already at ends of arcs
+	// Top hemisphere:    phi sweeps PI/2 (equator) -> 0 (pole, approached from PI/2)
+	TArray<TTuple<double, double>> Profile;
+
+	// Bottom pole
+	Profile.Add(MakeTuple(0.0, -HalfLength - (double)Radius0));
+	// Bottom hemisphere (excluding pole, including equator)
+	for (int32 i = 1; i <= ArcSteps; ++i)
+	{
+		const double Phi = (UE_DOUBLE_PI * 0.5) * (double)i / (double)ArcSteps;
+		Profile.Add(MakeTuple((double)Radius0 * FMath::Sin(Phi), -HalfLength - (double)Radius0 * FMath::Cos(Phi)));
+	}
+	// Top equator (separate ring for the taper)
+	Profile.Add(MakeTuple((double)Radius1, HalfLength));
+	// Top hemisphere (excluding equator, including pole)
+	for (int32 i = ArcSteps - 1; i >= 0; --i)
+	{
+		const double Phi = (UE_DOUBLE_PI * 0.5) * (double)i / (double)ArcSteps;
+		Profile.Add(MakeTuple((double)Radius1 * FMath::Sin(Phi), HalfLength + (double)Radius1 * FMath::Cos(Phi)));
+	}
+
+	// Revolve profile to produce vertex rings.
+	// Pole entries (r==0) produce a single vertex; all others produce a full ring.
+	const int32 N = Profile.Num();
+	TArray<TArray<int32>> Rings;
+	Rings.SetNum(N);
+
+	for (int32 PIdx = 0; PIdx < N; ++PIdx)
+	{
+		const double R = Profile[PIdx].Get<0>();
+		const double Z = Profile[PIdx].Get<1>();
+		if (FMath::IsNearlyZero(R))
+		{
+			Rings[PIdx].Add(Mesh.AppendVertex(FVector3d(0.0, 0.0, Z)));
+		}
+		else
+		{
+			for (int32 CIdx = 0; CIdx < CircleSteps; ++CIdx)
+			{
+				const double Theta = CIdx * DTheta;
+				Rings[PIdx].Add(Mesh.AppendVertex(FVector3d(R * FMath::Cos(Theta), R * FMath::Sin(Theta), Z)));
+			}
+		}
+	}
+
+	// Stitch adjacent rings with triangles.
+	for (int32 PIdx = 0; PIdx + 1 < N; ++PIdx)
+	{
+		const TArray<int32>& RingA = Rings[PIdx];
+		const TArray<int32>& RingB = Rings[PIdx + 1];
+
+		if (RingA.Num() == 1)
+		{
+			// Fan from pole down to ring
+			const int32 Pole = RingA[0];
+			for (int32 CIdx = 0; CIdx < CircleSteps; ++CIdx)
+			{
+				Mesh.AppendTriangle(Pole, RingB[CIdx], RingB[(CIdx + 1) % CircleSteps]);
+			}
+		}
+		else if (RingB.Num() == 1)
+		{
+			// Fan from ring up to pole
+			const int32 Pole = RingB[0];
+			for (int32 CIdx = 0; CIdx < CircleSteps; ++CIdx)
+			{
+				Mesh.AppendTriangle(RingA[CIdx], Pole, RingA[(CIdx + 1) % CircleSteps]);
+			}
+		}
+		else
+		{
+			// Quad strip
+			for (int32 CIdx = 0; CIdx < CircleSteps; ++CIdx)
+			{
+				const int32 Next = (CIdx + 1) % CircleSteps;
+				Mesh.AppendTriangle(RingA[CIdx],  RingB[CIdx],  RingB[Next]);
+				Mesh.AppendTriangle(RingA[CIdx],  RingB[Next],  RingA[Next]);
+			}
+		}
+	}
+
+	return Mesh;
+}
+
+// Build a mesh from a convex element's pre-computed triangle data.
+// Returns an empty mesh if the data is missing.
+static UE::Geometry::FDynamicMesh3 BuildConvexMesh(const FKConvexElem& Elem)
+{
+	UE::Geometry::FDynamicMesh3 Mesh;
+
+	if (Elem.VertexData.IsEmpty() || Elem.IndexData.Num() < 3)
+	{
+		return Mesh;
+	}
+
+	for (const FVector& V : Elem.VertexData)
+	{
+		Mesh.AppendVertex(FVector3d(V));
+	}
+
+	for (int32 Tri = 0; Tri + 2 < Elem.IndexData.Num(); Tri += 3)
+	{
+		Mesh.AppendTriangle(Elem.IndexData[Tri], Elem.IndexData[Tri + 1], Elem.IndexData[Tri + 2]);
+	}
+
+	return Mesh;
+}
+
+} // anonymous namespace
+
+///////////////////////////////////////////////////////////////////////////////////////////////////////////////
+
+class FPhysicsAssetSurfaceRenderableType : public IRenderableType
+{
+	UE_DATAFLOW_IRENDERABLE_PRIMARY_TYPE(TObjectPtr<UPhysicsAsset>, PhysicsAsset);
+	UE_DATAFLOW_IRENDERABLE_RENDER_GROUP(Surface);
+	UE_DATAFLOW_IRENDERABLE_VIEW_MODE(FDataflowConstruction3DViewMode);
+	UE_DATAFLOW_IRENDERABLE_SETTINGS(UDataflowPhysicsAssetRenderSettings);
+
+	virtual void GetPrimitiveComponents(const FRenderableTypeInstance& Instance, FRenderableComponents& OutComponents) const override
+	{
+		const UDataflowPhysicsAssetRenderSettings* Settings = Instance.GetTypedRenderSettings<UDataflowPhysicsAssetRenderSettings>();
+
+		const TObjectPtr<const UPhysicsAsset> PhysicsAsset = GetPhysicsAsset(Instance, nullptr);
+		if (!PhysicsAsset)
+		{
+			return;
+		}
+
+		static FName BaseParentName = TEXT("PhysicsAsset");
+		UPrimitiveComponent* ParentComponent = OutComponents.AddNewComponent<UStaticMeshComponent>(BaseParentName);
+		if (!ParentComponent)
+		{
+			return;
+		}
+
+		UMaterialInstanceDynamic* MaterialInstance = CreateMaterialInstance(ParentComponent, Settings);
+
+		// Compute reference-pose bone transforms. Shapes without a valid preview mesh
+		// are placed at the origin (BoneCS = Identity).
+		const USkeletalMesh* PreviewMesh = PhysicsAsset->PreviewSkeletalMesh.Get();
+		const TArray<FTransform> CSTransforms = ComputeReferencePoseCS(PreviewMesh);
+
+		int32 ShapeCounter = 0;
+
+		for (int32 BodyIdx = 0; BodyIdx < PhysicsAsset->SkeletalBodySetups.Num(); ++BodyIdx)
+		{
+			const USkeletalBodySetup* BodySetup = PhysicsAsset->SkeletalBodySetups[BodyIdx];
+			if (!BodySetup)
+			{
+				continue;
+			}
+
+			// Bone component-space transform (identity if mesh is absent or bone not found)
+			const int32 BoneIdx = PreviewMesh
+				? PreviewMesh->GetRefSkeleton().FindBoneIndex(BodySetup->BoneName)
+				: INDEX_NONE;
+			const FTransform BoneCS = (BoneIdx != INDEX_NONE) 
+				? CSTransforms[BoneIdx] 
+				: FTransform::Identity;
+
+			static const FName UnknownBoneName(TEXT("[Unknown]"));
+			const FName BoneName = (BoneIdx != INDEX_NONE)
+				? PreviewMesh->GetRefSkeleton().GetBoneName(BoneIdx)
+				: UnknownBoneName;
+
+			const FKAggregateGeom& Geom = BodySetup->AggGeom;
+
+			// --- Spheres ---
+			for (int32 EIdx = 0; EIdx < Geom.SphereElems.Num(); ++EIdx, ++ShapeCounter)
+			{
+				const FKSphereElem& Elem = Geom.SphereElems[EIdx];
+				UStaticMesh* SphereMesh = UE::Dataflow::RenderGeometry::GetDataflowSphere();
+				if (!SphereMesh)
+				{
+					continue;
+				}
+
+				const FName SphereComponentName(FString::Printf(TEXT("[%d]_%s_(Sphere)"), ShapeCounter, *BoneName.ToString()));
+				if (UStaticMeshComponent* Component = OutComponents.AddNewComponent<UStaticMeshComponent>(SphereComponentName, ParentComponent))
+				{
+					Component->SetStaticMesh(SphereMesh);
+					Component->SetCollisionEnabled(ECollisionEnabled::NoCollision);
+
+					FTransform WorldTransform(FQuat::Identity, BoneCS.TransformPosition(FVector(Elem.Center)));
+					WorldTransform.SetScale3D(FVector(Elem.Radius / 50.f));
+					Component->SetWorldTransform(WorldTransform);
+					Component->SetMaterial(0, MaterialInstance);
+				}
+			}
+
+			// --- Boxes ---
+			for (int32 EIdx = 0; EIdx < Geom.BoxElems.Num(); ++EIdx, ++ShapeCounter)
+			{
+				const FKBoxElem& Elem = Geom.BoxElems[EIdx];
+				UStaticMesh* BoxMesh = UE::Dataflow::RenderGeometry::GetDataflowBox();
+				if (!BoxMesh)
+				{
+					continue;
+				}
+
+				const FName BoxComponentName(FString::Printf(TEXT("[%d]_%s_(Box)"), ShapeCounter, *BoneName.ToString()));
+				if (UStaticMeshComponent* Component = OutComponents.AddNewComponent<UStaticMeshComponent>(BoxComponentName, ParentComponent))
+				{
+					Component->SetStaticMesh(BoxMesh);
+					Component->SetCollisionEnabled(ECollisionEnabled::NoCollision);
+
+					// Elem.GetTransform() gives position+rotation in bone-local space.
+					// X/Y/Z are already full-extents; the unit box mesh has a extent of 100 n all directions.
+					FTransform WorldTransform = Elem.GetTransform() * BoneCS;
+					WorldTransform.SetScale3D(FVector(Elem.X, Elem.Y, Elem.Z) / 100.f);
+					Component->SetWorldTransform(WorldTransform);
+					Component->SetMaterial(0, MaterialInstance);
+				}
+			}
+
+			// --- Capsules (Sphyl) ---
+			for (int32 EIdx = 0; EIdx < Geom.SphylElems.Num(); ++EIdx, ++ShapeCounter)
+			{
+				const FKSphylElem& Elem = Geom.SphylElems[EIdx];
+
+				UE::Geometry::FDynamicMesh3 CapsuleMesh = BuildCapsuleMesh(Elem.Radius, Elem.Length);
+
+				const FName CapsuleComponentName(FString::Printf(TEXT("[%d]_%s_(Capsule)"), ShapeCounter, *BoneName.ToString()));
+				if (UDynamicMeshComponent* Component = OutComponents.AddNewComponent<UDynamicMeshComponent>(CapsuleComponentName, ParentComponent))
+				{
+					Component->SetMesh(MoveTemp(CapsuleMesh));
+					Component->SetCollisionEnabled(ECollisionEnabled::NoCollision);
+					Component->SetWorldTransform(Elem.GetTransform() * BoneCS);
+					Component->SetOverrideRenderMaterial(MaterialInstance);
+				}
+			}
+
+			// --- Tapered capsules ---
+			for (int32 EIdx = 0; EIdx < Geom.TaperedCapsuleElems.Num(); ++EIdx, ++ShapeCounter)
+			{
+				const FKTaperedCapsuleElem& Elem = Geom.TaperedCapsuleElems[EIdx];
+
+				UE::Geometry::FDynamicMesh3 TaperedMesh = BuildTaperedCapsuleMesh(Elem.Radius0, Elem.Radius1, Elem.Length);
+
+				const FName TaperedCapsuleComponentName(FString::Printf(TEXT("[%d]_%s_(Tapered_Capsule)"), ShapeCounter, *BoneName.ToString()));
+				if (UDynamicMeshComponent* Component = OutComponents.AddNewComponent<UDynamicMeshComponent>(TaperedCapsuleComponentName, ParentComponent))
+				{
+					Component->SetMesh(MoveTemp(TaperedMesh));
+					Component->SetCollisionEnabled(ECollisionEnabled::NoCollision);
+					Component->SetWorldTransform(Elem.GetTransform() * BoneCS);
+					Component->SetOverrideRenderMaterial(MaterialInstance);
+				}
+			}
+
+			// --- Convex hulls ---
+			for (int32 EIdx = 0; EIdx < Geom.ConvexElems.Num(); ++EIdx, ++ShapeCounter)
+			{
+				const FKConvexElem& Elem = Geom.ConvexElems[EIdx];
+
+				UE::Geometry::FDynamicMesh3 ConvexMesh = BuildConvexMesh(Elem);
+				if (ConvexMesh.TriangleCount() == 0)
+				{
+					continue;
+				}
+
+				const FName ConvexComponentName(FString::Printf(TEXT("[%d]_%s_(Convex)"), ShapeCounter, *BoneName.ToString()));
+				if (UDynamicMeshComponent* Component = OutComponents.AddNewComponent<UDynamicMeshComponent>(ConvexComponentName, ParentComponent))
+				{
+					Component->SetMesh(MoveTemp(ConvexMesh));
+					Component->SetCollisionEnabled(ECollisionEnabled::NoCollision);
+					Component->SetWorldTransform(Elem.GetTransform() * BoneCS);
+					Component->SetOverrideRenderMaterial(MaterialInstance);
+				}
+			}
+		}
+	}
+};
+
+///////////////////////////////////////////////////////////////////////////////////////////////////////////////
+
+void RegisterPhysicsAssetRenderableTypes()
+{
+	UE_DATAFLOW_REGISTER_RENDERABLE_TYPE(FPhysicsAssetSurfaceRenderableType);
+}
+
+///////////////////////////////////////////////////////////////////////////////////////////////////////////////
+
+} // namespace UE::Dataflow::Private
diff --git a/Engine/Plugins/Experimental/Dataflow/Source/DataflowEditor/Private/DataflowRendering/DataflowPhysicsAssetRenderableType.h b/Engine/Plugins/Experimental/Dataflow/Source/DataflowEditor/Private/DataflowRendering/DataflowPhysicsAssetRenderableType.h
new file mode 100644
--- /dev/null
+++ b/Engine/Plugins/Experimental/Dataflow/Source/DataflowEditor/Private/DataflowRendering/DataflowPhysicsAssetRenderableType.h
@@ -0,0 +1,25 @@
+// Copyright Epic Games, Inc. All Rights Reserved.
+
+#pragma once
+
+#include "DataflowRendering/DataflowRenderableTypeSettings.h"
+#include "DataflowRendering/DataflowRenderableTypeUtils.h"
+
+#include "DataflowPhysicsAssetRenderableType.generated.h"
+
+UCLASS(MinimalAPI)
+class UDataflowPhysicsAssetRenderSettings : public UDataflowRenderableTypeSettings
+{
+	GENERATED_BODY()
+
+	UDataflowPhysicsAssetRenderSettings(const FObjectInitializer& ObjectInitializer = FObjectInitializer::Get());
+
+public:
+	UPROPERTY(EditAnywhere, Category = "Surface", meta = (ShowOnlyInnerProperties))
+	FDataflowSimpleColorCommonRenderSettings ColorSettings;
+};
+
+namespace UE::Dataflow::Private
+{
+	void RegisterPhysicsAssetRenderableTypes();
+}
diff --git a/Engine/Plugins/Experimental/Dataflow/Source/DataflowEditor/Public/DataflowRendering/DataflowRenderableComponents.h b/Engine/Plugins/Experimental/Dataflow/Source/DataflowEditor/Public/DataflowRendering/DataflowRenderableComponents.h
--- a/Engine/Plugins/Experimental/Dataflow/Source/DataflowEditor/Public/DataflowRendering/DataflowRenderableComponents.h
+++ b/Engine/Plugins/Experimental/Dataflow/Source/DataflowEditor/Public/DataflowRendering/DataflowRenderableComponents.h
@@ -37,6 +37,10 @@ namespace UE::Dataflow
 			T* NewComponent = NewObject<T>(ParentObject, ObjectName);
 			if (NewComponent)
 			{
+				if (USceneComponent* ParentObjectAsComp = Cast<USceneComponent>(ParentObject))
+				{
+					NewComponent->SetupAttachment(ParentObjectAsComp);
+				}
 				Components.Add(NewComponent);
 			}
 			return NewComponent;
