Dataflow Editor : update weightmap tool nodes to support snapshots
- Editable node now inherits from FDataflowToolNode class
- Update APIs of the nodes to accommodate the new snapshot paradigm
- Update both Cloth Weightmap node and general collection paint weightmap node

#Synced-CL 52199989

diff --git a/Engine/Plugins/ChaosClothAssetDataflowNodes/Source/ChaosClothAssetDataflowNodes/Private/ChaosClothAsset/WeightMapNode.cpp b/Engine/Plugins/ChaosClothAssetDataflowNodes/Source/ChaosClothAssetDataflowNodes/Private/ChaosClothAsset/WeightMapNode.cpp
--- a/Engine/Plugins/ChaosClothAssetDataflowNodes/Source/ChaosClothAssetDataflowNodes/Private/ChaosClothAsset/WeightMapNode.cpp
+++ b/Engine/Plugins/ChaosClothAssetDataflowNodes/Source/ChaosClothAssetDataflowNodes/Private/ChaosClothAsset/WeightMapNode.cpp
@@ -204,8 +204,11 @@ void FChaosClothAssetWeightMapNode::Evaluate(UE::Dataflow::FContext& Context, co
 				}
 				else
 				{
+					TArray<float> WeightMapValues;
+					GetWeightsFromActiveSnapshot(&ClothFacade, WeightMapValues);
+
 					constexpr bool bIsSim = true;
-					CheckSourceVertexWeights(ClothSimWeights, GetVertexWeights(), bIsSim);
+					CheckSourceVertexWeights(ClothSimWeights, WeightMapValues, bIsSim);
 					CalculateFinalVertexWeightValues(ClothFacade.GetWeightMap(InInputName), ClothSimWeights);
 				}
 			}
@@ -227,8 +230,11 @@ void FChaosClothAssetWeightMapNode::Evaluate(UE::Dataflow::FContext& Context, co
 				}
 				else
 				{
+					TArray<float> WeightMapValues;
+					GetWeightsFromActiveSnapshot(&ClothFacade, WeightMapValues);
+
 					constexpr bool bIsSim = false;
-					CheckSourceVertexWeights(ClothRenderWeights, GetVertexWeights(), bIsSim);
+					CheckSourceVertexWeights(ClothRenderWeights, WeightMapValues, bIsSim);
 					CalculateFinalVertexWeightValues(ClothFacade.GetUserDefinedAttribute<float>(InInputName, ClothCollectionGroup::RenderVertices), ClothRenderWeights);
 				}
 			}
@@ -318,8 +324,13 @@ void FChaosClothAssetWeightMapNode::GetVertexAttributeValues(UE::Dataflow::FCont
 			FallbackInputValues.Init(0, NumValues);
 			InputAttributeValues = FallbackInputValues;
 		}
+
+		// This keeps backward compatibility to load what's on the node 
+		TArray<float> WeightMapValues;
+		GetWeightsFromActiveSnapshot(&ClothFacade, WeightMapValues);
+
 		OutValues.Init(0, NumValues);
-		UE::Chaos::ClothAsset::Private::CalculateFinalVertexWeightValues(InputAttributeValues, OutValues, MapOverrideType, GetVertexWeights());
+		UE::Chaos::ClothAsset::Private::CalculateFinalVertexWeightValues(InputAttributeValues, OutValues, MapOverrideType, WeightMapValues);
 	}
 }
 
@@ -358,20 +369,27 @@ void FChaosClothAssetWeightMapNode::SetVertexAttributeValues(UE::Dataflow::FCont
 			NumFinalWeights = ClothFacade.GetNumRenderVertices();
 		}
 
+		TArray<float> ValuesToStore;
 		if (InWeightIndices.Num() == InValues.Num())
 		{
 			TArray<float> ValuesToApply;
 			ValuesToApply.Init(0.f, NumFinalWeights);
 			for (int32 Index = 0; Index < InWeightIndices.Num(); ++Index)
 			{
-				ValuesToApply[InWeightIndices[Index]] = InValues[Index];
+				const int32 WeightIndex = InWeightIndices[Index];
+				if (ValuesToApply.IsValidIndex(WeightIndex))
+				{
+					ValuesToApply[InWeightIndices[Index]] = InValues[Index];
+				}
 			}
-			UE::Chaos::ClothAsset::Private::SetVertexWeights(InputAttributeValues, ValuesToApply, MapOverrideType, GetVertexWeights());
+			UE::Chaos::ClothAsset::Private::SetVertexWeights(InputAttributeValues, ValuesToApply, MapOverrideType, ValuesToStore);
 		}
 		else
 		{
-			UE::Chaos::ClothAsset::Private::SetVertexWeights(InputAttributeValues, InValues, MapOverrideType, GetVertexWeights());
+			UE::Chaos::ClothAsset::Private::SetVertexWeights(InputAttributeValues, InValues, MapOverrideType, ValuesToStore);
 		}
+
+		StoreWeightsInSnapshot(&ClothFacade, ValuesToStore, AddSnapshot());
 	}
 }
 
@@ -397,77 +415,216 @@ void FChaosClothAssetWeightMapNode::GetExtraVertexMapping(UE::Dataflow::FContext
 	}
 }
 
-void FChaosClothAssetWeightMapNode::SwapStoredAttributeValuesWith(TArray<float>& OtherValues)
-{
-	Swap(VertexWeights, OtherValues);
-}
-
 void FChaosClothAssetWeightMapNode::SetVertexWeights(const TConstArrayView<float> InputMap, const TArray<float>& FinalValues)
 {
-	UE::Chaos::ClothAsset::Private::SetVertexWeights(InputMap, FinalValues, MapOverrideType, GetVertexWeights());
+	TArray<float> WeightsToStore;
+	UE::Chaos::ClothAsset::Private::SetVertexWeights(InputMap, FinalValues, MapOverrideType, WeightsToStore);
+
+	StoreWeightsInSnapshot(/*ClothFacade*/nullptr, WeightsToStore, AddSnapshot());
 }
 
 void FChaosClothAssetWeightMapNode::CalculateFinalVertexWeightValues(const TConstArrayView<float> InputMap, TArrayView<float> FinalOutputMap) const
 {
-	UE::Chaos::ClothAsset::Private::CalculateFinalVertexWeightValues(InputMap, FinalOutputMap, MapOverrideType, GetVertexWeights());
+	TArray<float> WeightMapValues;
+	GetWeightsFromActiveSnapshot(/*ClothFacade*/nullptr, WeightMapValues);
+
+	UE::Chaos::ClothAsset::Private::CalculateFinalVertexWeightValues(InputMap, FinalOutputMap, MapOverrideType, WeightMapValues);
 }
 
-// Object encapsulating a change to the WeightMap node's values. Used for Undo/Redo.
-class FChaosClothAssetWeightMapNode::FWeightMapNodeChange final : public FToolCommandChange
+namespace ChaosClothAssetWeightMapNode::Private
 {
+	static FName VertexGroupName = TEXT("Vertices");
+	static FName VertexPositionsAttributeName = TEXT("Positions");
+	static FName VertexWeightsAttributeName = TEXT("Weights");
 
-public:
-
-	FWeightMapNodeChange(const FChaosClothAssetWeightMapNode& Node) :
-		NodeGuid(Node.GetGuid()),
-		SavedWeights(Node.GetVertexWeights()),
-		SavedMapOverrideType(Node.MapOverrideType),
-		SavedWeightMapName(Node.OutputName.StringValue)
-	{}
-
-private:
-
-	FGuid NodeGuid;
-	TArray<float> SavedWeights;
+	static FName FaceGroupName = TEXT("Faces");
+	static FName FaceIndicesAttributeName = TEXT("Indices");
 
-	EChaosClothAssetWeightMapOverrideType SavedMapOverrideType;
-	FString SavedWeightMapName;
+	template<typename T>
+	void CopyData(TConstArrayView<T> From, TManagedArray<T>& To)
+	{
+		if (ensure(From.Num() == To.Num()))
+		{
+			for (int32 Index = 0; Index < From.Num(); ++Index)
+			{
+				To[Index] = From[Index];
+			}
+		}
+	};
+}
 
-	virtual FString ToString() const final
+void FChaosClothAssetWeightMapNode::GetWeightsFromActiveSnapshot(UE::Chaos::ClothAsset::FCollectionClothConstFacade* ClothFacade, TArray<float>& OutValues) const
+{
+	if (const FDataflowToolNodeSnapshot* ActiveSnapshot = GetActiveSnapshot())
 	{
-		return TEXT("FChaosClothAssetWeightMapNode::FWeightMapNodeChange");
+		GetWeightsFromSnapshot(ClothFacade, *ActiveSnapshot, OutValues);
 	}
-
-	virtual void Apply(UObject* Object) final
+	else
 	{
-		SwapApplyRevert(Object);
+		// This keeps backward compatibility to load what's on the node before we add snapshots
+		OutValues = GetLegacyVertexWeights();
 	}
+}
+
+void FChaosClothAssetWeightMapNode::GetWeightsFromSnapshot(UE::Chaos::ClothAsset::FCollectionClothConstFacade* ClothFacade, const FDataflowToolNodeSnapshot& Snapshot, TArray<float>& OutValues) const
+{
+	using namespace UE::Chaos::ClothAsset;
 
-	virtual void Revert(UObject* Object) final
+	OutValues.Reset();
+
+	int32 InVertexCount = 0;
+	int32 InFaceCount = 0;
+	if (ClothFacade && ClothFacade->IsValid())
 	{
-		SwapApplyRevert(Object);
+		if (MeshTarget == EChaosClothAssetWeightMapMeshTarget::Simulation)
+		{
+			InVertexCount = ClothFacade->GetNumSimVertices3D();
+			InFaceCount = ClothFacade->GetNumSimFaces();
+		}
+		else if (MeshTarget == EChaosClothAssetWeightMapMeshTarget::Render)
+		{
+			InVertexCount = ClothFacade->GetNumRenderVertices();
+			InFaceCount = ClothFacade->GetNumRenderFaces();
+		}
 	}
 
-	void SwapApplyRevert(UObject* Object)
+	const TManagedArray<float>* VtxWeights = Snapshot.Data.FindAttribute<float>(ChaosClothAssetWeightMapNode::Private::VertexWeightsAttributeName, ChaosClothAssetWeightMapNode::Private::VertexGroupName);
+	if (VtxWeights)
 	{
-		if (UDataflow* const Dataflow = Cast<UDataflow>(Object))
+		bool bMatchOriginalTopology = true;
+		bMatchOriginalTopology &= (Snapshot.Data.NumElements(ChaosClothAssetWeightMapNode::Private::VertexGroupName) == InVertexCount);
+		bMatchOriginalTopology &= (Snapshot.Data.NumElements(ChaosClothAssetWeightMapNode::Private::FaceGroupName) == InFaceCount);
+		// todo : add more checks
+
+		const TManagedArray<FVector3f>* VtxPositions = Snapshot.Data.FindAttribute<FVector3f>(ChaosClothAssetWeightMapNode::Private::VertexPositionsAttributeName, ChaosClothAssetWeightMapNode::Private::VertexGroupName);
+		const TManagedArray<FIntVector>* FaceIndices = Snapshot.Data.FindAttribute<FIntVector>(ChaosClothAssetWeightMapNode::Private::FaceIndicesAttributeName, ChaosClothAssetWeightMapNode::Private::FaceGroupName);
+		const bool bIsValidSourceMesh = (VtxPositions && FaceIndices);
+
+		const bool bIsValidTargetMesh = (InVertexCount > 0 && InFaceCount > 0);
+
+		if (!bMatchOriginalTopology)
 		{
-			if (const TSharedPtr<FDataflowNode> BaseNode = Dataflow->GetDataflow()->FindBaseNode(NodeGuid))
+			if (bIsValidSourceMesh && bIsValidTargetMesh)
 			{
-				if (FChaosClothAssetWeightMapNode* const Node = BaseNode->AsType<FChaosClothAssetWeightMapNode>())
+				// We need to transfer the transfer the data over to the new topology 
+				if (MeshTarget == EChaosClothAssetWeightMapMeshTarget::Simulation)
+				{
+					// Remap the weights
+					OutValues.SetNumZeroed(ClothFacade->GetNumSimVertices3D());
+
+					// in this automatic mode we only transfer on 3D model because the patterns may have changed and be in different place
+					// we also only store the original 3D Mesh 
+					FClothGeometryTools::TransferWeightMap(
+						VtxPositions->GetConstArray(),
+						FaceIndices->GetConstArray(),
+						VtxWeights->GetConstArray(),
+						ClothFacade->GetSimPosition3D(),
+						ClothFacade->GetSimNormal(),
+						ClothFacade->GetSimIndices3D(),
+						MakeArrayView(OutValues));
+				}
+				else if (MeshTarget == EChaosClothAssetWeightMapMeshTarget::Render)
 				{
-					Swap(Node->GetVertexWeights(), SavedWeights);
-					Swap(Node->MapOverrideType, SavedMapOverrideType);
-					Swap(Node->OutputName.StringValue, SavedWeightMapName);
+					OutValues.SetNumZeroed(ClothFacade->GetNumRenderVertices());
 
-					Node->Invalidate();
+					FClothGeometryTools::TransferWeightMap(
+						VtxPositions->GetConstArray(),
+						FaceIndices->GetConstArray(),
+						VtxWeights->GetConstArray(),
+						ClothFacade->GetRenderPosition(),
+						ClothFacade->GetRenderNormal(),
+						ClothFacade->GetRenderIndices(),
+						MakeArrayView(OutValues));
 				}
 			}
 		}
+		else 
+		{
+			// we simply read the weights as they are 
+			OutValues.Append(VtxWeights->GetConstArray());
+		}
+	}
+}
+
+void FChaosClothAssetWeightMapNode::StoreWeightsInSnapshot(UE::Chaos::ClothAsset::FCollectionClothConstFacade* ClothFacade, const TArray<float>& InValues, FDataflowToolNodeSnapshot& OutSnapshot) const
+{
+	using namespace UE::Chaos::ClothAsset;
+
+	OutSnapshot.Data.Reset();
+
+	if (InValues.Num())
+	{
+		TManagedArray<float>& VtxWeights = OutSnapshot.Data.AddAttribute<float>(ChaosClothAssetWeightMapNode::Private::VertexWeightsAttributeName, ChaosClothAssetWeightMapNode::Private::VertexGroupName);
+
+		// for backward compatibility with cloth tools, we need to support case where the Cloth facade is not available 
+		// all new Dataflow workflow should have a valid cloth collection though
+		if (ClothFacade && ClothFacade->IsValid())
+		{
+			TManagedArray<FVector3f>& VtxPositions = OutSnapshot.Data.AddAttribute<FVector3f>(ChaosClothAssetWeightMapNode::Private::VertexPositionsAttributeName, ChaosClothAssetWeightMapNode::Private::VertexGroupName);
+			TManagedArray<FIntVector>& FaceIndices = OutSnapshot.Data.AddAttribute<FIntVector>(ChaosClothAssetWeightMapNode::Private::FaceIndicesAttributeName, ChaosClothAssetWeightMapNode::Private::FaceGroupName);
+
+			if (MeshTarget == EChaosClothAssetWeightMapMeshTarget::Simulation)
+			{
+				OutSnapshot.Data.Resize(ClothFacade->GetNumSimVertices3D(), ChaosClothAssetWeightMapNode::Private::VertexGroupName);
+				OutSnapshot.Data.Resize(ClothFacade->GetNumSimFaces(), ChaosClothAssetWeightMapNode::Private::FaceGroupName);
+
+				ChaosClothAssetWeightMapNode::Private::CopyData(ClothFacade->GetSimPosition3D(), VtxPositions);
+				ChaosClothAssetWeightMapNode::Private::CopyData(MakeArrayView(InValues), VtxWeights);
+				ChaosClothAssetWeightMapNode::Private::CopyData(ClothFacade->GetSimIndices3D(), FaceIndices);
+			}
+			else if (MeshTarget == EChaosClothAssetWeightMapMeshTarget::Render)
+			{
+				OutSnapshot.Data.Resize(ClothFacade->GetNumRenderVertices(), ChaosClothAssetWeightMapNode::Private::VertexGroupName);
+				OutSnapshot.Data.Resize(ClothFacade->GetNumRenderFaces(), ChaosClothAssetWeightMapNode::Private::FaceGroupName);
+
+				ChaosClothAssetWeightMapNode::Private::CopyData(ClothFacade->GetRenderPosition(), VtxPositions);
+				ChaosClothAssetWeightMapNode::Private::CopyData(MakeArrayView(InValues), VtxWeights);
+				ChaosClothAssetWeightMapNode::Private::CopyData(ClothFacade->GetRenderIndices(), FaceIndices);
+			}
+		}
+		else
+		{
+			// backward compatibility path , we only store the weights
+			OutSnapshot.Data.Resize(InValues.Num(), ChaosClothAssetWeightMapNode::Private::VertexGroupName);
+			ChaosClothAssetWeightMapNode::Private::CopyData(MakeArrayView(InValues), VtxWeights);
+		}
+	}
+}
+
+/////////////////////////////////////////////////////////////////////////////////////////////////////////////////////////////////////////////////////////////////////////
+
+// Object encapsulating a change to the WeightMap node's values. Used for Undo/Redo.
+class FChaosClothAssetWeightMapNode::FWeightMapNodeChange final 
+	: public FDataflowToolNode::FSnapshotToolChange
+{
+public:
+	FWeightMapNodeChange(FChaosClothAssetWeightMapNode& Node) 
+		: FSnapshotToolChange(Node)
+		, SavedMapOverrideType(Node.MapOverrideType)
+		, SavedWeightMapName(Node.OutputName.StringValue)
+	{}
+
+private:
+	EChaosClothAssetWeightMapOverrideType SavedMapOverrideType;
+	FString SavedWeightMapName;
+
+	virtual FString ToString() const final
+	{
+		return TEXT("FChaosClothAssetWeightMapNode::FWeightMapNodeChange");
+	}
+
+	virtual void SwapApplyRevert(UObject* Object, FDataflowToolNode& Node) override final
+	{
+		if (FChaosClothAssetWeightMapNode* TypedNode = Node.AsType<FChaosClothAssetWeightMapNode>())
+		{
+			Swap(TypedNode->MapOverrideType, SavedMapOverrideType);
+			Swap(TypedNode->OutputName.StringValue, SavedWeightMapName);
+		}
+		FDataflowToolNode::FSnapshotToolChange::SwapApplyRevert(Object, Node);
 	}
 };
 
-TUniquePtr<FToolCommandChange> FChaosClothAssetWeightMapNode::MakeWeightMapNodeChange(const FChaosClothAssetWeightMapNode& Node)
+TUniquePtr<FToolCommandChange> FChaosClothAssetWeightMapNode::MakeWeightMapNodeChange(FChaosClothAssetWeightMapNode& Node)
 {
 	return MakeUnique<FChaosClothAssetWeightMapNode::FWeightMapNodeChange>(Node);
 }
diff --git a/Engine/Plugins/ChaosClothAssetDataflowNodes/Source/ChaosClothAssetDataflowNodes/Public/ChaosClothAsset/WeightMapNode.h b/Engine/Plugins/ChaosClothAssetDataflowNodes/Source/ChaosClothAssetDataflowNodes/Public/ChaosClothAsset/WeightMapNode.h
--- a/Engine/Plugins/ChaosClothAssetDataflowNodes/Source/ChaosClothAssetDataflowNodes/Public/ChaosClothAsset/WeightMapNode.h
+++ b/Engine/Plugins/ChaosClothAssetDataflowNodes/Source/ChaosClothAssetDataflowNodes/Public/ChaosClothAsset/WeightMapNode.h
@@ -12,6 +12,12 @@
 
 namespace UE_DEPRECATED(5.5, "Use UE::Dataflow instead.") Dataflow {}
 
+namespace UE::Chaos::ClothAsset
+{
+	class FCollectionClothFacade;
+	class FCollectionClothConstFacade;
+}
+
 /** Which mesh to update with the node's weight map */
 UENUM()
 enum class EChaosClothAssetWeightMapMeshTarget : uint8
@@ -74,10 +80,9 @@ private:
 	friend class UClothEditorWeightMapPaintTool;
 
 	class FWeightMapNodeChange;
-	static TUniquePtr<class FToolCommandChange> CHAOSCLOTHASSETDATAFLOWNODES_API MakeWeightMapNodeChange(const FChaosClothAssetWeightMapNode& Node);
+	static TUniquePtr<class FToolCommandChange> CHAOSCLOTHASSETDATAFLOWNODES_API MakeWeightMapNodeChange(FChaosClothAssetWeightMapNode& Node);
 
-	const TArray<float>& GetVertexWeights() const { return VertexWeights; }
-	TArray<float>& GetVertexWeights() { return VertexWeights; }
+	const TArray<float>& GetLegacyVertexWeights() const { return VertexWeights; }
 
 	// These methods are exported for UClothEditorWeightMapPaintTool which lives in a different module.
 	FName CHAOSCLOTHASSETDATAFLOWNODES_API GetInputName(UE::Dataflow::FContext& Context) const;
@@ -98,12 +103,26 @@ private:
 	virtual void GetVertexAttributeValues(UE::Dataflow::FContext& Context, TArray<float>& OutValues) const override;
 	virtual void SetVertexAttributeValues(UE::Dataflow::FContext& Context, const TArray<float>& InValues, const TArray<int32>& InWeightIndices) override;
 	virtual void GetExtraVertexMapping(UE::Dataflow::FContext& Context, FName SelectedViewMode, TArray<int32>& OutMappingToWeight, TArray<TArray<int32>>& OutMappingFromWeight) const override;
-	virtual const TArray<float>& GetStoredAttributeValues() const override { return VertexWeights; }
-	virtual void SwapStoredAttributeValuesWith(TArray<float>& OtherValues) override;
 	//~ End - FDataflowVertexAttributeEditableNode API
 
 	void OnTransfer(UE::Dataflow::FContext& Context);
 
+	/** 
+	* Get weights values from a specific snaphot
+	* If there's data in the snapshot, the original topology wioll be compared to the one in the optional ClothFacade, 
+	* and data will be projected to match the new topology if needed
+	*/
+	void GetWeightsFromSnapshot(UE::Chaos::ClothAsset::FCollectionClothConstFacade* ClothFacade, const FDataflowToolNodeSnapshot& Snapshot, TArray<float>& OutValues) const;
+	void GetWeightsFromActiveSnapshot(UE::Chaos::ClothAsset::FCollectionClothConstFacade* ClothFacade, TArray<float>& OutValues) const;
+
+	/**
+	* store weights values in a specific snapshot 
+	* if the cloth facade is available, the topology information will be also stored in the snapshot
+	*/
+	void StoreWeightsInSnapshot(UE::Chaos::ClothAsset::FCollectionClothConstFacade* ClothFacade, const TArray<float>& InValues, FDataflowToolNodeSnapshot& OutSnapshot) const;
+
+	// This property should be technically deprecated but we need to keep it because we cannot post serialize fix it as we need a dataflow Context to create a snapshot 
+	// deprecated methods are still using it and we still use it to load the data when there's no pre-existing snapshot 
 	UPROPERTY()
 	TArray<float> VertexWeights;
 
diff --git a/Engine/Plugins/Experimental/Dataflow/Source/DataflowEditor/Private/DataflowEditorTools/DataflowEditorWeightMapPaintTool.cpp b/Engine/Plugins/Experimental/Dataflow/Source/DataflowEditor/Private/DataflowEditorTools/DataflowEditorWeightMapPaintTool.cpp
--- a/Engine/Plugins/Experimental/Dataflow/Source/DataflowEditor/Private/DataflowEditorTools/DataflowEditorWeightMapPaintTool.cpp
+++ b/Engine/Plugins/Experimental/Dataflow/Source/DataflowEditor/Private/DataflowEditorTools/DataflowEditorWeightMapPaintTool.cpp
@@ -428,15 +428,19 @@ void UDataflowEditorWeightMapPaintTool::Setup()
 			FCollectionAttributeKey AttributeKey = WeightMapNodeToUpdate->GetWeightAttributeKey(*DataflowContext);
 
 			// Fill the attribute values
+			PRAGMA_DISABLE_DEPRECATION_WARNINGS
 			WeightMapNodeToUpdate->FillAttributeWeights(DataflowEditorContextObject->GetSelectedCollection(), AttributeKey, SetupWeights);
+			PRAGMA_ENABLE_DEPRECATION_WARNINGS
 		}
 	}
 
 	// Copy weights from selected node to the preview mesh
 	TArray<float> CurrentWeights;
 	CurrentWeights.SetNumZeroed(NumExpectedWeights);
+	PRAGMA_DISABLE_DEPRECATION_WARNINGS
 	WeightMapNodeToUpdate->ExtractVertexWeights(SetupWeights, TArrayView<float>(CurrentWeights));
-	
+	PRAGMA_ENABLE_DEPRECATION_WARNINGS
+
 	if (bHaveDynamicMeshToWeightConversion)
 	{
 		if (WeightToDynamicMesh.Num() == CurrentWeights.Num())	// Only copy node weights if they match the number of mesh vertices
@@ -1830,6 +1834,7 @@ void UDataflowEditorWeightMapPaintTool::UpdateSelectedNode()
 
 	check(WeightMapNodeToUpdate);
 
+	PRAGMA_DISABLE_DEPRECATION_WARNINGS
 	if (bHaveDynamicMeshToWeightConversion)
 	{
 		WeightMapNodeToUpdate->ReportVertexWeights(SetupWeights, CurrentWeights, DynamicMeshToWeight);
@@ -1838,6 +1843,7 @@ void UDataflowEditorWeightMapPaintTool::UpdateSelectedNode()
 	{
 		WeightMapNodeToUpdate->ReportVertexWeights(SetupWeights, CurrentWeights, {});
 	}
+	PRAGMA_ENABLE_DEPRECATION_WARNINGS
 	
 	WeightMapNodeToUpdate->Name = UpdateWeightMapProperties->Name;
 
diff --git a/Engine/Plugins/Experimental/Dataflow/Source/DataflowNodes/Private/Dataflow/DataflowCollectionAddScalarVertexPropertyNode.cpp b/Engine/Plugins/Experimental/Dataflow/Source/DataflowNodes/Private/Dataflow/DataflowCollectionAddScalarVertexPropertyNode.cpp
--- a/Engine/Plugins/Experimental/Dataflow/Source/DataflowNodes/Private/Dataflow/DataflowCollectionAddScalarVertexPropertyNode.cpp
+++ b/Engine/Plugins/Experimental/Dataflow/Source/DataflowNodes/Private/Dataflow/DataflowCollectionAddScalarVertexPropertyNode.cpp
@@ -7,61 +7,12 @@
 #include "InteractiveToolChange.h"
 #include "Misc/LazySingleton.h"
 
-#include UE_INLINE_GENERATED_CPP_BY_NAME(DataflowCollectionAddScalarVertexPropertyNode)
-
-#define LOCTEXT_NAMESPACE "DataflowCollectionAddScalarVertexProperty"
-
+#include "Dataflow/DataflowCollectionAttributeUtils.h"
 
-////////////////////////////////////////////////////////////////////////////////////////////////////////////////////////////////////////////////////////////
-// 
-// Object encapsulating a change to the vertex attribute node's values. Used for Undo/Redo.
-class FDataflowVertexAttributeEditableNode::FEditNodeToolChange final : public FToolCommandChange
-{
-public:
-	FEditNodeToolChange(const FDataflowVertexAttributeEditableNode& Node) :
-		NodeGuid(Node.GetGuid()),
-		SavedAttributeValues(Node.GetStoredAttributeValues())
-	{}
-
-private:
-	virtual FString ToString() const final
-	{
-		return TEXT("FDataflowVertexAttributeEditableNode::FEditNodeToolChange");
-	}
 
-	virtual void Apply(UObject* Object) final
-	{
-		SwapApplyRevert(Object);
-	}
-
-	virtual void Revert(UObject* Object) final
-	{
-		SwapApplyRevert(Object);
-	}
-
-	void SwapApplyRevert(UObject* Object)
-	{
-		if (UDataflow* const Dataflow = Cast<UDataflow>(Object))
-		{
-			if (const TSharedPtr<FDataflowNode> BaseNode = Dataflow->GetDataflow()->FindBaseNode(NodeGuid))
-			{
-				if (FDataflowVertexAttributeEditableNode* const Node = BaseNode->AsType<FDataflowVertexAttributeEditableNode>())
-				{
-					Node->SwapStoredAttributeValuesWith(SavedAttributeValues);
-					Node->Invalidate();
-				}
-			}
-		}
-	}
-
-	FGuid NodeGuid;
-	TArray<float> SavedAttributeValues;
-};
+#include UE_INLINE_GENERATED_CPP_BY_NAME(DataflowCollectionAddScalarVertexPropertyNode)
 
-TUniquePtr<FToolCommandChange> FDataflowVertexAttributeEditableNode::MakeEditNodeToolChange()
-{
-	return MakeUnique<FDataflowVertexAttributeEditableNode::FEditNodeToolChange>(*this);
-}
+#define LOCTEXT_NAMESPACE "DataflowCollectionAddScalarVertexProperty"
 
 ////////////////////////////////////////////////////////////////////////////////////////////////////////////////////////////////////////////////////////////
 //
@@ -100,18 +51,22 @@ void FDataflowCollectionAddScalarVertexPropertyNode::Evaluate(UE::Dataflow::FCon
 			const FName InGroup(Key.Group);
 			TManagedArray<float>& ScalarWeights = InCollection.AddAttribute<float>(InName, InGroup);
 
-			if (VertexWeights.Num() > 0 && VertexWeights.Num() != ScalarWeights.Num())
+			TArray<float> StoredWeights;
+			StoredWeights.SetNumZeroed(ScalarWeights.Num());
+			GetWeightsFromActiveSnapshot(InCollection, InGroup, StoredWeights);
+
+			if (StoredWeights.Num() > 0 && StoredWeights.Num() != ScalarWeights.Num())
 			{
 				FDataflowTools::LogAndToastWarning(*this,
 					LOCTEXT("VertexCountMismatchHeadline", "Vertex count mismatch."),
 					FText::Format(LOCTEXT("VertexCountMismatchDetails", "Vertex weights in the node: {0}\n Vertices in group \"{1}\" in the Collection: {2}"),
-						VertexWeights.Num(),
+						StoredWeights.Num(),
 						FText::FromName(InGroup),
 						ScalarWeights.Num()));
 			}
 
-			const TArray<float> SetupWeights = ScalarWeights.GetConstArray();
-			ExtractVertexWeights(SetupWeights, TArrayView<float>(ScalarWeights.GetData(), ScalarWeights.Num()));
+			TConstArrayView<float> InputWeights = MakeArrayView(ScalarWeights.GetConstArray());
+			ComputeFinalWeights(InputWeights, StoredWeights, MakeArrayView(ScalarWeights.GetData(), ScalarWeights.Num()));
 			
 		}
 
@@ -123,6 +78,7 @@ void FDataflowCollectionAddScalarVertexPropertyNode::Evaluate(UE::Dataflow::FCon
 	}
 }
 
+// DEPRECATED 5.8
 bool FDataflowCollectionAddScalarVertexPropertyNode::FillAttributeWeights( 
 	const TSharedPtr<const FManagedArrayCollection> SelectedCollection, const FCollectionAttributeKey& InAttributeKey, TArray<float>& OutAttributeValues) const
 {
@@ -157,13 +113,17 @@ FCollectionAttributeKey FDataflowCollectionAddScalarVertexPropertyNode::GetWeigh
 
 void FDataflowCollectionAddScalarVertexPropertyNode::GetSupportedViewModes(UE::Dataflow::FContext& Context, TArray<FName>& OutViewModeNames) const
 {
+	for (UE::Dataflow::FRenderingParameter& RenderParameter : FDataflowCollectionAddScalarVertexPropertyNode::GetRenderParametersImpl())
+	{
+		OutViewModeNames.AddUnique(RenderParameter.ViewMode);
+	}
 	// matches UE::Dataflow::FDataflowConstruction3DViewMode::Name - but we cannot include the module to get it directly 
-	static const FName DataflowConstruction3DViewModeName = TEXT("3DView");
-	static const FName Cloth3DSimView = TEXT("Cloth3DSimView");
-	static const FName ClothRenderView = TEXT("ClothRenderView");
-	OutViewModeNames.Add(DataflowConstruction3DViewModeName);
-	OutViewModeNames.Add(Cloth3DSimView);
-	OutViewModeNames.Add(ClothRenderView);
+	//static const FName DataflowConstruction3DViewModeName = TEXT("3DView");
+	//static const FName Cloth3DSimView = TEXT("Cloth3DSimView");
+	//static const FName ClothRenderView = TEXT("ClothRenderView");
+	//OutViewModeNames.Add(DataflowConstruction3DViewModeName);
+	//OutViewModeNames.Add(Cloth3DSimView);
+	//OutViewModeNames.Add(ClothRenderView);
 }
 
 void FDataflowCollectionAddScalarVertexPropertyNode::GetVertexAttributeValues(UE::Dataflow::FContext& Context, TArray<float>& OutValues) const
@@ -173,19 +133,25 @@ void FDataflowCollectionAddScalarVertexPropertyNode::GetVertexAttributeValues(UE
 	const FName InGroupName(InKey.Group);
 	const FManagedArrayCollection& InCollection = GetValue(Context, &Collection, Collection);
 
-	const int32 NumVertices = InCollection.NumElements(InGroupName);
-	OutValues.SetNumUninitialized(NumVertices);
+	UE::Dataflow::FDataflowCollectionFacade Facade(InCollection, InGroupName);
+	const int32 NumVertices = Facade.GetVertexPositions().Num();
+	OutValues.SetNumZeroed(NumVertices);
+
+	TArray<float> StoredWeights;
+	StoredWeights.SetNumZeroed(NumVertices);
+	GetWeightsFromActiveSnapshot(InCollection, InGroupName, StoredWeights);
 
-	if (const TManagedArray<float>* AttributeArray = InCollection.FindAttributeTyped<float>(InAttribName, InGroupName))
+	if (const TManagedArray<float>* WeightsAttribute = InCollection.FindAttributeTyped<float>(InAttribName, InGroupName))
 	{
-		ExtractVertexWeights(AttributeArray->GetConstArray(), OutValues);
+		ComputeFinalWeights(WeightsAttribute->GetConstArray(), StoredWeights, OutValues);
 	}
 	else
 	{
 		TArray<float> ZeroSetupWeights;
 		ZeroSetupWeights.Init(0.0f, NumVertices);
-		ExtractVertexWeights(ZeroSetupWeights, OutValues);
+		ComputeFinalWeights(ZeroSetupWeights, StoredWeights, OutValues);
 	}
+	
 }
 
 void FDataflowCollectionAddScalarVertexPropertyNode::SetVertexAttributeValues(UE::Dataflow::FContext& Context, const TArray<float>& InValues, const TArray<int32>& InWeightIndices)
@@ -195,55 +161,114 @@ void FDataflowCollectionAddScalarVertexPropertyNode::SetVertexAttributeValues(UE
 	const FName InGroupName(InKey.Group);
 	const FManagedArrayCollection& InCollection = GetValue(Context, &Collection, Collection);
 
-	if (const TManagedArray<float>* AttributeArray = InCollection.FindAttributeTyped<float>(InAttribName, InGroupName))
+	UE::Dataflow::FDataflowCollectionFacade Facade(InCollection, InGroupName);
+	const int32 NumVertices = Facade.GetVertexPositions().Num();
+
+	TArray<float> WeightsToStore;
+	WeightsToStore.SetNum(NumVertices);
+
+	const TManagedArray<float>* AttributeArray = InCollection.FindAttributeTyped<float>(InAttribName, InGroupName);
+	if (AttributeArray && AttributeArray->Num() == NumVertices)
 	{
-		ReportVertexWeights(AttributeArray->GetConstArray(), InValues, InWeightIndices);
+		ComputeWeightsToStore(AttributeArray->GetConstArray(), InValues, InWeightIndices, WeightsToStore);
 	}
 	else
 	{
 		TArray<float> ZeroSetupWeights;
-		const int32 NumVertices = InCollection.NumElements(InGroupName);
 		ZeroSetupWeights.Init(0.0f, NumVertices);
-		ReportVertexWeights(ZeroSetupWeights, InValues, InWeightIndices);
+		ComputeWeightsToStore(ZeroSetupWeights, InValues, InWeightIndices, WeightsToStore);
 	}
 
+	FDataflowToolNodeSnapshot& NewSnapShot = AddSnapshot();
+	NewSnapShot.Description = FString::Printf(TEXT("%s - %d vertices"), *InKey.Group, NumVertices);
+	StoreWeightsInSnapshot(InCollection, WeightsToStore, InGroupName, NewSnapShot);
+
 	Invalidate();
 }
 
-void FDataflowCollectionAddScalarVertexPropertyNode::SwapStoredAttributeValuesWith(TArray<float>& OtherValues)
+
+void FDataflowCollectionAddScalarVertexPropertyNode::GetExtraVertexMapping(UE::Dataflow::FContext& Context, FName SelectedViewMode, TArray<int32>& OutMappingToWeight, TArray<TArray<int32>>& OutMappingFromWeight) const
 {
-	Swap(VertexWeights, OtherValues);
+	OutMappingToWeight.Reset();
+	OutMappingFromWeight.Reset();
+
+	const FCollectionAttributeKey InKey = GetWeightAttributeKey(Context);
+	const FName InAttribName(InKey.Attribute);
+	const FName InGroupName(InKey.Group);
+	FManagedArrayCollection InCollection = GetValue<FManagedArrayCollection>(Context, &Collection);
+
+	UE::Dataflow::FDataflowCollectionFacade Facade(InCollection, InGroupName);
+	TConstArrayView<int32> Mapping2DTo3D = Facade.Get2Dto3DMapping();
+	if (Mapping2DTo3D.IsEmpty())
+	{
+		return;
+	}
+
+	OutMappingToWeight = Mapping2DTo3D;
+
+	// recompute the reverse mapping 
+	OutMappingFromWeight.SetNum(Facade.GetVertexPositions().Num());
+	for (int32 Index2D = 0; Index2D < Mapping2DTo3D.Num(); ++Index2D)
+	{
+		const int32 Index3D = Mapping2DTo3D[Index2D];
+		if (OutMappingFromWeight.IsValidIndex(Index3D))
+		{
+			OutMappingFromWeight[Index3D].Add(Index2D);
+		}
+	}
 }
 
-void FDataflowCollectionAddScalarVertexPropertyNode::ReportVertexWeights(const TArray<float>& SetupWeights, const TArray<float>& FinalWeights,
-		const TArray<int32>& WeightIndices)
+// DEPRECATED 5.8
+void FDataflowCollectionAddScalarVertexPropertyNode::ReportVertexWeights(const TArray<float>& SetupWeights, const TArray<float>& FinalWeights, const TArray<int32>& WeightIndices)
 {
-	check((WeightIndices.Num() == FinalWeights.Num()) || (WeightIndices.IsEmpty() && (SetupWeights.Num() == FinalWeights.Num())));
-	VertexWeights.SetNumZeroed(SetupWeights.Num());
+	VertexWeights.SetNum(SetupWeights.Num());
+	ComputeWeightsToStore(SetupWeights, FinalWeights, WeightIndices, VertexWeights);
+}
+
+void FDataflowCollectionAddScalarVertexPropertyNode::ComputeWeightsToStore(TConstArrayView<float> InputWeights, TConstArrayView<float> FinalWeights, TConstArrayView<int32> WeightIndices, TArrayView<float> WeightsToStore)
+{
+	check((WeightIndices.Num() == FinalWeights.Num()) || (WeightIndices.IsEmpty() && (InputWeights.Num() == FinalWeights.Num())));
+	check(WeightsToStore.Num() == InputWeights.Num());
+
+	// zero the values first 
+	for (float& WeightToStore : WeightsToStore)
+	{
+		WeightToStore = 0;
+	}
+
 	for (int32 WeightIndex = 0, NumWeights = FinalWeights.Num(); WeightIndex <NumWeights; ++WeightIndex)
 	{
 		const int32 VertexIndex = WeightIndices.IsEmpty() ? WeightIndex : WeightIndices[WeightIndex];
-		switch (OverrideType)
+		if (WeightsToStore.IsValidIndex(VertexIndex))
 		{
+			switch (OverrideType)
+			{
 			case EDataflowWeightMapOverrideType::ReplaceAll:
-				VertexWeights[VertexIndex] = FinalWeights[WeightIndex];
+				WeightsToStore[VertexIndex] = FinalWeights[WeightIndex];
 				break;
 			case EDataflowWeightMapOverrideType::ReplaceChanged:
-				VertexWeights[VertexIndex] = (SetupWeights[VertexIndex] == FinalWeights[WeightIndex]) ? ReplaceChangedPassthroughValue : FinalWeights[WeightIndex];
+				WeightsToStore[VertexIndex] = (InputWeights[VertexIndex] == FinalWeights[WeightIndex]) ? ReplaceChangedPassthroughValue : FinalWeights[WeightIndex];
 				break;
 			case EDataflowWeightMapOverrideType::AddDifference:
-				VertexWeights[VertexIndex] = FinalWeights[WeightIndex] - SetupWeights[VertexIndex];
+				WeightsToStore[VertexIndex] = FinalWeights[WeightIndex] - InputWeights[VertexIndex];
 				break;
 			default: unimplemented();
+			}
 		}
 	}
 }
 
+// DEPRECATED 5.8
 void FDataflowCollectionAddScalarVertexPropertyNode::ExtractVertexWeights(const TArray<float>& SetupWeights, TArrayView<float> FinalWeights) const
 {
-	check(SetupWeights.Num() == FinalWeights.Num());
+	ComputeFinalWeights(SetupWeights, VertexWeights, FinalWeights);
+}
 
-	int32 NumWeights = FMath::Min(VertexWeights.Num(), SetupWeights.Num());
+void FDataflowCollectionAddScalarVertexPropertyNode::ComputeFinalWeights(TConstArrayView<float> InputWeights, TConstArrayView<float> StoredWeights, TArrayView<float> FinalWeights) const
+{
+	check(InputWeights.Num() == FinalWeights.Num());
+
+	int32 NumWeights = FMath::Min(StoredWeights.Num(), InputWeights.Num());
 	if(NumWeights != 0)
 	{ 
 		for (int32 WeightIndex = 0; WeightIndex < NumWeights; ++WeightIndex)
@@ -251,13 +276,13 @@ void FDataflowCollectionAddScalarVertexPropertyNode::ExtractVertexWeights(const
 			switch (OverrideType)
 			{
 				case EDataflowWeightMapOverrideType::ReplaceAll:
-					FinalWeights[WeightIndex] = FMath::Clamp(VertexWeights[WeightIndex], 0.f, 1.f);
+					FinalWeights[WeightIndex] = FMath::Clamp(StoredWeights[WeightIndex], 0.f, 1.f);
 					break;
 				case EDataflowWeightMapOverrideType::ReplaceChanged:
-					FinalWeights[WeightIndex] = FMath::Clamp( (VertexWeights[WeightIndex] == ReplaceChangedPassthroughValue) ? SetupWeights[WeightIndex] : VertexWeights[WeightIndex], 0.0f, 1.0f);
+					FinalWeights[WeightIndex] = FMath::Clamp( (StoredWeights[WeightIndex] == ReplaceChangedPassthroughValue) ? InputWeights[WeightIndex] : StoredWeights[WeightIndex], 0.0f, 1.0f);
 					break;
 				case EDataflowWeightMapOverrideType::AddDifference:
-					FinalWeights[WeightIndex] = FMath::Clamp(SetupWeights[WeightIndex] + VertexWeights[WeightIndex], 0.f, 1.f);
+					FinalWeights[WeightIndex] = FMath::Clamp(InputWeights[WeightIndex] + StoredWeights[WeightIndex], 0.f, 1.f);
 					break;
 				default: unimplemented();
 			}
@@ -265,14 +290,137 @@ void FDataflowCollectionAddScalarVertexPropertyNode::ExtractVertexWeights(const
 	}
 	else
 	{
-		NumWeights = SetupWeights.Num();
+		NumWeights = InputWeights.Num();
 		for (int32 WeightIndex = 0; WeightIndex < NumWeights; ++WeightIndex)
 		{
-			FinalWeights[WeightIndex] = SetupWeights[WeightIndex];
+			FinalWeights[WeightIndex] = InputWeights[WeightIndex];
 		}
 	}
 }
 
+namespace DataflowCollectionAddScalarVertexPropertyNode::Private
+{
+	static FName VertexGroupName = TEXT("Vertices");
+	static FName VertexPositionsAttributeName = TEXT("Positions");
+	static FName VertexWeightsAttributeName = TEXT("Weights");
+
+	static FName FaceGroupName = TEXT("Faces");
+	static FName FaceIndicesAttributeName = TEXT("Indices");
+
+	template<typename T>
+	void CopyData(TConstArrayView<T> From, TManagedArray<T>& To)
+	{
+		if (ensure(From.Num() == To.Num()))
+		{
+			for (int32 Index = 0; Index < From.Num(); ++Index)
+			{
+				To[Index] = From[Index];
+			}
+		}
+	};
+}
+
+void FDataflowCollectionAddScalarVertexPropertyNode::GetWeightsFromActiveSnapshot(const FManagedArrayCollection& InCollection, FName TargetGroupName, TArrayView<float> OutValues) const
+{
+	if (const FDataflowToolNodeSnapshot* ActiveSnapshot = GetActiveSnapshot())
+	{
+		GetWeightsFromSnapshot(InCollection, *ActiveSnapshot, TargetGroupName, OutValues);
+	}
+	else
+	{
+		// This keeps backward compatibility to load what's on the node before we add snapshots
+		// if topology as changed since, for now let's simply do our best 
+		const int32 NumValuesToCopy = FMath::Min(VertexWeights.Num(), OutValues.Num());
+		for (int32 Index = 0; Index < NumValuesToCopy; ++Index)
+		{
+			OutValues[Index] = VertexWeights[Index];
+		}
+	}
+}
+
+void FDataflowCollectionAddScalarVertexPropertyNode::GetWeightsFromSnapshot(const FManagedArrayCollection& InCollection, const FDataflowToolNodeSnapshot& Snapshot, FName TargetGroupName, TArrayView<float> OutValues) const
+{
+	namespace NodePrivate = DataflowCollectionAddScalarVertexPropertyNode::Private;
+	
+	using namespace UE::Dataflow;
+
+	const TManagedArray<float>* VtxWeights = Snapshot.Data.FindAttribute<float>(NodePrivate::VertexWeightsAttributeName, NodePrivate::VertexGroupName);
+	if (VtxWeights)
+	{
+		FDataflowCollectionFacade Facade(InCollection, TargetGroupName);
+		TConstArrayView<FVector3f> CollectionVtxPositions = Facade.GetVertexPositions();
+		TConstArrayView<FIntVector> CollectionTriangles = Facade.GetTriangleIndices();
+
+		const int32 NumVertices = CollectionVtxPositions.Num();
+		const int32 NumTriangles = CollectionTriangles.Num();
+
+		bool bMatchOriginalTopology = true;
+		bMatchOriginalTopology &= (VtxWeights->Num() == OutValues.Num());
+		bMatchOriginalTopology &= (Snapshot.Data.NumElements(NodePrivate::VertexGroupName) == NumVertices);
+		bMatchOriginalTopology &= (Snapshot.Data.NumElements(NodePrivate::FaceGroupName) == NumTriangles);
+		// todo : add more checks
+
+		const TManagedArray<FVector3f>* VtxPositions = Snapshot.Data.FindAttribute<FVector3f>(NodePrivate::VertexPositionsAttributeName, NodePrivate::VertexGroupName);
+		const TManagedArray<FIntVector>* FaceIndices = Snapshot.Data.FindAttribute<FIntVector>(NodePrivate::FaceIndicesAttributeName, NodePrivate::FaceGroupName);
+		const bool bIsValidSourceMesh = (VtxPositions && FaceIndices);
+
+		const bool bIsValidTargetMesh = (NumVertices > 0 && NumTriangles > 0);
+
+		if (!bMatchOriginalTopology)
+		{
+			if (bIsValidSourceMesh && bIsValidTargetMesh)
+			{
+				// TODO(Dataflow) need to transfer the weights from one topo to another 
+				// in this automatic mode we only transfer on 3D model because the patterns may have changed and be in different place
+				// we also only store the original 3D Mesh 
+			}
+			else
+			{
+				// no transfer since there's no geometry to transfer to or from
+			}
+		}
+		else
+		{
+			// We simply read the weights as they are 
+			check(VtxWeights->Num() == OutValues.Num());
+			for (int32 Index = 0; Index < VtxWeights->Num(); ++Index)
+			{
+				OutValues[Index] = (*VtxWeights)[Index];
+			}
+		}
+	}
+}
+
+void FDataflowCollectionAddScalarVertexPropertyNode::StoreWeightsInSnapshot(const FManagedArrayCollection& InCollection, TConstArrayView<float> InValues, FName TargetGroupName, FDataflowToolNodeSnapshot& OutSnapshot) const
+{
+	namespace NodePrivate = DataflowCollectionAddScalarVertexPropertyNode::Private;
+	using namespace UE::Dataflow;
+
+	OutSnapshot.Data.Reset();
+
+	FDataflowCollectionFacade Facade(InCollection, TargetGroupName);
+	TConstArrayView<FVector3f> CollectionVtxPositions = Facade.GetVertexPositions();
+	TConstArrayView<FIntVector> CollectionTriangles = Facade.GetTriangleIndices();
+
+	if (InValues.Num() && InValues.Num() == CollectionVtxPositions.Num())
+	{
+		OutSnapshot.Data.AddGroup(NodePrivate::VertexGroupName);
+		OutSnapshot.Data.Resize(CollectionVtxPositions.Num(), NodePrivate::VertexGroupName);
+
+		OutSnapshot.Data.AddGroup(NodePrivate::FaceGroupName);
+		OutSnapshot.Data.Resize(CollectionTriangles.Num(), NodePrivate::FaceGroupName);
+
+		TManagedArray<float>& VtxWeights = OutSnapshot.Data.AddAttribute<float>(NodePrivate::VertexWeightsAttributeName, NodePrivate::VertexGroupName);
+		NodePrivate::CopyData(MakeArrayView(InValues), VtxWeights);
+
+		TManagedArray<FVector3f>& VtxPositions = OutSnapshot.Data.AddAttribute<FVector3f>(NodePrivate::VertexPositionsAttributeName, NodePrivate::VertexGroupName);
+		NodePrivate::CopyData(CollectionVtxPositions, VtxPositions);
+
+		TManagedArray<FIntVector>& FaceIndices = OutSnapshot.Data.AddAttribute<FIntVector>(NodePrivate::FaceIndicesAttributeName, NodePrivate::FaceGroupName);
+		NodePrivate::CopyData(CollectionTriangles, FaceIndices);
+	}
+}
+
 FDataflowNode::FAttributeKey FDataflowCollectionAddScalarVertexPropertyNode::GetVertexAttributeToVisualize(UE::Dataflow::FContext& Context, const FName OutputName, const FName RenderGroup) const
 {
 	const FCollectionAttributeKey CollectionAttributeKey = GetWeightAttributeKey(Context);
diff --git a/Engine/Plugins/Experimental/Dataflow/Source/DataflowNodes/Private/Dataflow/DataflowVertexAttributeEditableNode.cpp b/Engine/Plugins/Experimental/Dataflow/Source/DataflowNodes/Private/Dataflow/DataflowVertexAttributeEditableNode.cpp
--- a/Engine/Plugins/Experimental/Dataflow/Source/DataflowNodes/Private/Dataflow/DataflowVertexAttributeEditableNode.cpp
+++ b/Engine/Plugins/Experimental/Dataflow/Source/DataflowNodes/Private/Dataflow/DataflowVertexAttributeEditableNode.cpp
@@ -10,19 +10,28 @@
 #include UE_INLINE_GENERATED_CPP_BY_NAME(DataflowVertexAttributeEditableNode)
 
 FDataflowVertexAttributeEditableNode::FDataflowVertexAttributeEditableNode(const UE::Dataflow::FNodeParameters& InParam, FGuid InGuid)
-	: FDataflowNode(InParam, InGuid)
+	: FDataflowToolNode(InParam, InGuid)
 {}
 
-const TArray<float>& FDataflowVertexAttributeEditableNode::GetStoredAttributeValues() const 
-{ 
-	static const TArray<float> ConstEmptyArray;
-	ensure(false);
-	return ConstEmptyArray;
-}
-
 void FDataflowVertexAttributeEditableNode::GetExtraVertexMapping(UE::Dataflow::FContext& Context, FName SelectedViewMode, TArray<int32>& OutMappingToWeight, TArray<TArray<int32>>& OutMappingFromWeight) const
 {
 	OutMappingToWeight.Reset();
 	OutMappingFromWeight.Reset();
 }
 
+////////////////////////////////////////////////////////////////////////////////////////////////////////////////////////////////////////////////////////////
+// 
+// Object encapsulating a change to the vertex attribute node's values. Used for Undo/Redo.
+class FDataflowVertexAttributeEditableNode::FEditNodeToolChange final : public FDataflowToolNode::FSnapshotToolChange
+{
+public:
+	FEditNodeToolChange(FDataflowVertexAttributeEditableNode& Node)
+		: FDataflowToolNode::FSnapshotToolChange(Node)
+	{
+	}
+};
+
+TUniquePtr<FToolCommandChange> FDataflowVertexAttributeEditableNode::MakeEditNodeToolChange()
+{
+	return MakeUnique<FDataflowVertexAttributeEditableNode::FEditNodeToolChange>(*this);
+}
diff --git a/Engine/Plugins/Experimental/Dataflow/Source/DataflowNodes/Public/Dataflow/DataflowCollectionAddScalarVertexPropertyNode.h b/Engine/Plugins/Experimental/Dataflow/Source/DataflowNodes/Public/Dataflow/DataflowCollectionAddScalarVertexPropertyNode.h
--- a/Engine/Plugins/Experimental/Dataflow/Source/DataflowNodes/Public/Dataflow/DataflowCollectionAddScalarVertexPropertyNode.h
+++ b/Engine/Plugins/Experimental/Dataflow/Source/DataflowNodes/Public/Dataflow/DataflowCollectionAddScalarVertexPropertyNode.h
@@ -37,7 +37,6 @@ struct FDataflowCollectionAddScalarVertexPropertyNode : public FDataflowVertexAt
 
 	
 public:
-
 	UPROPERTY(Meta = (DataflowInput, DataflowOutput, DataflowPassthrough = "Collection"))
 	FManagedArrayCollection Collection;
 
@@ -48,8 +47,10 @@ public:
 	UPROPERTY(meta = (DisplayName = "Attribute Key", DataflowInput, DataflowOutput, DataflowPassthrough = "AttributeKey"))
 	FCollectionAttributeKey AttributeKey;
 
+	// This property should be technically deprecated but we need to keep it because we cannot post serialize fix it as we need a dataflow Context to create a snapshot 
+	// deprecated methods are still using it and we still use it to load the data when there's no pre-existing snapshot 
 	UPROPERTY()
-	TArray<float> VertexWeights;
+	TArray<float> VertexWeights; 
 
 	UPROPERTY(EditAnywhere, Category = "Vertex Attributes", meta = (DisplayName = "Vertex Group"))
 	FScalarVertexPropertyGroup TargetGroup;
@@ -61,12 +62,15 @@ public:
 	UE_API FDataflowCollectionAddScalarVertexPropertyNode(const UE::Dataflow::FNodeParameters& InParam, FGuid InGuid = FGuid::NewGuid());
 
 	/** Report the vertex weights back onto the property ones */
+	UE_DEPRECATED(5.8, "ReportVertexWeights is now deprecated, use the FDataflowVertexAttributeEditableNode API instead")
 	UE_API void ReportVertexWeights(const TArray<float>& SetupWeights, const TArray<float>& FinalWeights, const TArray<int32>& WeightIndices);
 
 	/** Extract the vertex weights back from the property ones */
+	UE_DEPRECATED(5.8, "ExtractVertexWeights is now deprecated, use the FDataflowVertexAttributeEditableNode API instead")
 	UE_API void ExtractVertexWeights(const TArray<float>& SetupWeights, TArrayView<float> FinalWeights) const;
 
 	/** Fill the weight attribute values from the collection */
+	UE_DEPRECATED(5.8, "FillAttributeWeights is now deprecated, use the FDataflowVertexAttributeEditableNode API instead")
 	UE_API bool FillAttributeWeights(const TSharedPtr<const FManagedArrayCollection> SelectedCollection, const FCollectionAttributeKey& InAttributeKey, TArray<float>& OutAttributeValues) const;
 
 	/** Get the weights attribute key to retrieve/set the weight values*/
@@ -84,9 +88,23 @@ private:
 	virtual void GetSupportedViewModes(UE::Dataflow::FContext& Context, TArray<FName>& OutViewModeNames) const override;
 	virtual void GetVertexAttributeValues(UE::Dataflow::FContext& Context, TArray<float>& OutValues) const override;
 	virtual void SetVertexAttributeValues(UE::Dataflow::FContext& Context, const TArray<float>& InValues, const TArray<int32>& InWeightIndices) override;
-	virtual const TArray<float>& GetStoredAttributeValues() const override { return VertexWeights; }
-	virtual void SwapStoredAttributeValuesWith(TArray<float>& OtherValues) override;
+	virtual void GetExtraVertexMapping(UE::Dataflow::FContext& Context, FName SelectedViewMode, TArray<int32>& OutMappingToWeight, TArray<TArray<int32>>& OutMappingFromWeight) const override;
 	//~ End - FDataflowVertexAttributeEditableNode API
+
+	void GetWeightsFromActiveSnapshot(const FManagedArrayCollection& Collection, FName TargetGroupName, TArrayView<float> OutValues) const;
+	void GetWeightsFromSnapshot(const FManagedArrayCollection& Collection, const FDataflowToolNodeSnapshot& Snapshot, FName TargetGroupName, TArrayView<float> OutValues) const;
+	void StoreWeightsInSnapshot(const FManagedArrayCollection& Collection, TConstArrayView<float> InValues, FName TargetGroupName, FDataflowToolNodeSnapshot& OutSnapshot) const;
+
+	/** use the input and stored weights and process final weights based on the selected override type */
+	void ComputeFinalWeights(TConstArrayView<float> InputWeights, TConstArrayView<float> StoredWeights, TArrayView<float> FinalWeights) const;
+
+	/** 
+	* Compute weights to store from the final weights, input weights and the selected override type
+	* Only processed the indices contained in WeightIndices, if non-empty
+	*/
+	void ComputeWeightsToStore(TConstArrayView<float> InputWeights, TConstArrayView<float> FinalWeights, TConstArrayView<int32> WeightIndices, TArrayView<float> WeightsToStore);
+
+
 };
 
 #undef UE_API
diff --git a/Engine/Plugins/Experimental/Dataflow/Source/DataflowNodes/Public/Dataflow/DataflowVertexAttributeEditableNode.h b/Engine/Plugins/Experimental/Dataflow/Source/DataflowNodes/Public/Dataflow/DataflowVertexAttributeEditableNode.h
--- a/Engine/Plugins/Experimental/Dataflow/Source/DataflowNodes/Public/Dataflow/DataflowVertexAttributeEditableNode.h
+++ b/Engine/Plugins/Experimental/Dataflow/Source/DataflowNodes/Public/Dataflow/DataflowVertexAttributeEditableNode.h
@@ -2,7 +2,7 @@
 
 #pragma once 
 
-#include "Dataflow/DataflowNode.h"
+#include "Dataflow/DataflowToolNode.h"
 
 #include "DataflowVertexAttributeEditableNode.generated.h"
 
@@ -12,7 +12,7 @@
 * Base class for node using the vertex attribute paint tool
 */
 USTRUCT()
-struct FDataflowVertexAttributeEditableNode : public FDataflowNode
+struct FDataflowVertexAttributeEditableNode : public FDataflowToolNode
 {
 	GENERATED_USTRUCT_BODY()
 
@@ -35,10 +35,6 @@ public:
 
 	/** Provide a extra indirection vertex mapping - optional will return empty by default */
 	UE_API virtual void GetExtraVertexMapping(UE::Dataflow::FContext& Context, FName SelectedViewMode, TArray<int32>& OutMappingToWeight, TArray<TArray<int32>>& OutMappingFromWeight) const;
-
-private:
-	UE_API virtual const TArray<float>& GetStoredAttributeValues() const;
-	virtual void SwapStoredAttributeValuesWith(TArray<float>& OtherValues) { ensure(false); }
 };
 
 #undef UE_API
