Dataflow:
- Added DataflowDistanceFromBoxSamplerNode, DataflowDistanceFromSphereSamplerNode

#Synced-CL 52238556

diff --git a/Engine/Plugins/Experimental/Dataflow/Source/DataflowNodes/Private/Dataflow/SamplerNodes/DataflowDistanceFromBoxSamplerNode.cpp b/Engine/Plugins/Experimental/Dataflow/Source/DataflowNodes/Private/Dataflow/SamplerNodes/DataflowDistanceFromBoxSamplerNode.cpp
new file mode 100644
--- /dev/null
+++ b/Engine/Plugins/Experimental/Dataflow/Source/DataflowNodes/Private/Dataflow/SamplerNodes/DataflowDistanceFromBoxSamplerNode.cpp
@@ -0,0 +1,54 @@
+// Copyright Epic Games, Inc. All Rights Reserved.
+
+#include "Dataflow/SamplerNodes/DataflowDistanceFromBoxSamplerNode.h"
+#include "Dataflow/SamplerNodes/DataflowSamplerFunctions.h"
+#include "Dataflow/DataflowNodeFactory.h"
+
+FBox FDataflowDistanceFromBoxFloatSampler::GetRenderBounds() const
+{
+	return FBox(FVector(-50.0), FVector(50.0));
+}
+
+void FDataflowDistanceFromBoxFloatSampler::Sample(TArrayView<const FVector3f> Positions, TArrayView<float> OutValues) const
+{
+	if (Positions.Num() == OutValues.Num())
+	{
+		for (int32 Idx = 0; Idx < Positions.Num(); ++Idx)
+		{
+			OutValues[Idx] = FMath::Sqrt(
+			FMath::Square(FMath::Max3(0.0, Box.Min.X - Positions[Idx].X, Positions[Idx].X - Box.Max.X)) +
+				FMath::Square(FMath::Max3(0.0, Box.Min.Y - Positions[Idx].Y, Positions[Idx].Y - Box.Max.Y)) +
+				FMath::Square(FMath::Max3(0.0, Box.Min.Z - Positions[Idx].Z, Positions[Idx].Z - Box.Max.Z)));
+		}
+	}
+	else
+	{
+		for (int32 Idx = 0; Idx < OutValues.Num(); ++Idx)
+		{
+			OutValues[Idx] = 0.f;
+		}
+	}
+}
+
+FDataflowDistanceFromBoxFloatSamplerNode::FDataflowDistanceFromBoxFloatSamplerNode(const UE::Dataflow::FNodeParameters& InParam, FGuid InGuid)
+	: FDataflowNode(InParam, InGuid)
+{
+	RegisterInputConnection(&DistanceFromBoxSampler.Box, GET_MEMBER_NAME_CHECKED(FDataflowDistanceFromBoxFloatSampler, Box));
+	RegisterOutputConnection(&Sampler);
+}
+
+void FDataflowDistanceFromBoxFloatSamplerNode::Evaluate(UE::Dataflow::FContext& Context, const FDataflowOutput* Out) const
+{
+	if (Out->IsA(&Sampler))
+	{
+		TSharedRef<FDataflowDistanceFromBoxFloatSampler> Impl = MakeShared<FDataflowDistanceFromBoxFloatSampler>();
+
+		Impl->Box = GetValue(Context, &DistanceFromBoxSampler.Box);
+
+		FDataflowFloatSampler OutSampler(Impl.ToSharedPtr());
+
+		SetValue(Context, OutSampler, &Sampler);
+	}
+}
+
+
diff --git a/Engine/Plugins/Experimental/Dataflow/Source/DataflowNodes/Private/Dataflow/SamplerNodes/DataflowDistanceFromBoxSamplerNode.h b/Engine/Plugins/Experimental/Dataflow/Source/DataflowNodes/Private/Dataflow/SamplerNodes/DataflowDistanceFromBoxSamplerNode.h
new file mode 100644
--- /dev/null
+++ b/Engine/Plugins/Experimental/Dataflow/Source/DataflowNodes/Private/Dataflow/SamplerNodes/DataflowDistanceFromBoxSamplerNode.h
@@ -0,0 +1,49 @@
+// Copyright Epic Games, Inc. All Rights Reserved.
+
+#pragma once 
+
+#include "Dataflow/DataflowNode.h"
+#include "Dataflow/SamplerNodes/DataflowSamplerTypes.h"
+
+#include "DataflowDistanceFromBoxSamplerNode.generated.h"
+
+USTRUCT()
+struct FDataflowDistanceFromBoxFloatSampler : public FDataflowFloatSamplerBase
+{
+	GENERATED_BODY();
+
+	virtual ~FDataflowDistanceFromBoxFloatSampler() override = default;
+	virtual void Sample(TArrayView<const FVector3f> Positions, TArrayView<float> OutValues) const override;
+	virtual FBox GetRenderBounds() const override;
+
+public:
+	/** Box to compute the distance from */
+	UPROPERTY(meta = (DataflowInput))
+	FBox Box;
+};
+
+/**
+ * 
+ * Distance from a box float sampler
+ * Input(s) : an AABB
+ * Output(s): for every sampled point the distance from the input box
+ * 
+ */
+USTRUCT()
+struct FDataflowDistanceFromBoxFloatSamplerNode : public FDataflowNode
+{
+	GENERATED_BODY()
+	DATAFLOW_NODE_DEFINE_INTERNAL(FDataflowDistanceFromBoxFloatSamplerNode, "Distance From Box Float Sampler", "Samplers", "")
+
+public:
+	FDataflowDistanceFromBoxFloatSamplerNode(const UE::Dataflow::FNodeParameters& InParam, FGuid InGuid = FGuid::NewGuid());
+
+private:
+	UPROPERTY(meta = (DataflowOutput))
+	FDataflowFloatSampler Sampler;
+
+	UPROPERTY(EditAnywhere, Category = Sampler, meta = (ShowOnlyInnerProperties, SkipInDisplayNameChain))
+	FDataflowDistanceFromBoxFloatSampler DistanceFromBoxSampler;
+
+	void Evaluate(UE::Dataflow::FContext& Context, const FDataflowOutput* Out) const override;
+};
diff --git a/Engine/Plugins/Experimental/Dataflow/Source/DataflowNodes/Private/Dataflow/SamplerNodes/DataflowDistanceFromSphereSamplerNode.cpp b/Engine/Plugins/Experimental/Dataflow/Source/DataflowNodes/Private/Dataflow/SamplerNodes/DataflowDistanceFromSphereSamplerNode.cpp
new file mode 100644
--- /dev/null
+++ b/Engine/Plugins/Experimental/Dataflow/Source/DataflowNodes/Private/Dataflow/SamplerNodes/DataflowDistanceFromSphereSamplerNode.cpp
@@ -0,0 +1,51 @@
+// Copyright Epic Games, Inc. All Rights Reserved.
+
+#include "Dataflow/SamplerNodes/DataflowDistanceFromSphereSamplerNode.h"
+#include "Dataflow/SamplerNodes/DataflowSamplerFunctions.h"
+#include "Dataflow/DataflowNodeFactory.h"
+
+FBox FDataflowDistanceFromSphereFloatSampler::GetRenderBounds() const
+{
+	return FBox(FVector(-50.0), FVector(50.0));
+}
+
+void FDataflowDistanceFromSphereFloatSampler::Sample(TArrayView<const FVector3f> Positions, TArrayView<float> OutValues) const
+{
+	if (Positions.Num() == OutValues.Num())
+	{
+		for (int32 Idx = 0; Idx < Positions.Num(); ++Idx)
+		{
+			OutValues[Idx] = FMath::Abs((FVector(Positions[Idx]) - Sphere.Center).Length() - Sphere.W);
+		}
+	}
+	else
+	{
+		for (int32 Idx = 0; Idx < OutValues.Num(); ++Idx)
+		{
+			OutValues[Idx] = 0.f;
+		}
+	}
+}
+
+FDataflowDistanceFromSphereFloatSamplerNode::FDataflowDistanceFromSphereFloatSamplerNode(const UE::Dataflow::FNodeParameters& InParam, FGuid InGuid)
+	: FDataflowNode(InParam, InGuid)
+{
+	RegisterInputConnection(&DistanceFromSphereSampler.Sphere, GET_MEMBER_NAME_CHECKED(FDataflowDistanceFromSphereFloatSampler, Sphere));
+	RegisterOutputConnection(&Sampler);
+}
+
+void FDataflowDistanceFromSphereFloatSamplerNode::Evaluate(UE::Dataflow::FContext& Context, const FDataflowOutput* Out) const
+{
+	if (Out->IsA(&Sampler))
+	{
+		TSharedRef<FDataflowDistanceFromSphereFloatSampler> Impl = MakeShared<FDataflowDistanceFromSphereFloatSampler>();
+
+		Impl->Sphere = GetValue(Context, &DistanceFromSphereSampler.Sphere);
+
+		FDataflowFloatSampler OutSampler(Impl.ToSharedPtr());
+
+		SetValue(Context, OutSampler, &Sampler);
+	}
+}
+
+
diff --git a/Engine/Plugins/Experimental/Dataflow/Source/DataflowNodes/Private/Dataflow/SamplerNodes/DataflowDistanceFromSphereSamplerNode.h b/Engine/Plugins/Experimental/Dataflow/Source/DataflowNodes/Private/Dataflow/SamplerNodes/DataflowDistanceFromSphereSamplerNode.h
new file mode 100644
--- /dev/null
+++ b/Engine/Plugins/Experimental/Dataflow/Source/DataflowNodes/Private/Dataflow/SamplerNodes/DataflowDistanceFromSphereSamplerNode.h
@@ -0,0 +1,49 @@
+// Copyright Epic Games, Inc. All Rights Reserved.
+
+#pragma once 
+
+#include "Dataflow/DataflowNode.h"
+#include "Dataflow/SamplerNodes/DataflowSamplerTypes.h"
+
+#include "DataflowDistanceFromSphereSamplerNode.generated.h"
+
+USTRUCT()
+struct FDataflowDistanceFromSphereFloatSampler : public FDataflowFloatSamplerBase
+{
+	GENERATED_BODY();
+
+	virtual ~FDataflowDistanceFromSphereFloatSampler() override = default;
+	virtual void Sample(TArrayView<const FVector3f> Positions, TArrayView<float> OutValues) const override;
+	virtual FBox GetRenderBounds() const override;
+
+public:
+	/** Sphere to compute the distance from */
+	UPROPERTY(meta = (DataflowInput))
+	FSphere Sphere;
+};
+
+/**
+ * 
+ * Distance from a sphere float sampler
+ * Input(s) : a sphere
+ * Output(s): for every sampled point the distance from the input sphere
+ * 
+ */
+USTRUCT()
+struct FDataflowDistanceFromSphereFloatSamplerNode : public FDataflowNode
+{
+	GENERATED_BODY()
+	DATAFLOW_NODE_DEFINE_INTERNAL(FDataflowDistanceFromSphereFloatSamplerNode, "Distance From Sphere Float Sampler", "Samplers", "")
+
+public:
+	FDataflowDistanceFromSphereFloatSamplerNode(const UE::Dataflow::FNodeParameters& InParam, FGuid InGuid = FGuid::NewGuid());
+
+private:
+	UPROPERTY(meta = (DataflowOutput))
+	FDataflowFloatSampler Sampler;
+
+	UPROPERTY(EditAnywhere, Category = Sampler, meta = (ShowOnlyInnerProperties, SkipInDisplayNameChain))
+	FDataflowDistanceFromSphereFloatSampler DistanceFromSphereSampler;
+
+	void Evaluate(UE::Dataflow::FContext& Context, const FDataflowOutput* Out) const override;
+};
