Dataflow:
- Added FDataflowDistanceFromPlaneFloatSamplerNode

#Synced-CL 52238556
#Synced-CL 52238556

diff --git a/Engine/Plugins/Experimental/Dataflow/Source/DataflowNodes/Private/Dataflow/SamplerNodes/DataflowDistanceFromPlaneSamplerNode.cpp b/Engine/Plugins/Experimental/Dataflow/Source/DataflowNodes/Private/Dataflow/SamplerNodes/DataflowDistanceFromPlaneSamplerNode.cpp
new file mode 100644
--- /dev/null
+++ b/Engine/Plugins/Experimental/Dataflow/Source/DataflowNodes/Private/Dataflow/SamplerNodes/DataflowDistanceFromPlaneSamplerNode.cpp
@@ -0,0 +1,51 @@
+// Copyright Epic Games, Inc. All Rights Reserved.
+
+#include "Dataflow/SamplerNodes/DataflowDistanceFromPlaneSamplerNode.h"
+#include "Dataflow/SamplerNodes/DataflowSamplerFunctions.h"
+#include "Dataflow/DataflowNodeFactory.h"
+
+FBox FDataflowDistanceFromPlaneFloatSampler::GetRenderBounds() const
+{
+	return FBox(FVector(-50.0), FVector(50.0));
+}
+
+void FDataflowDistanceFromPlaneFloatSampler::Sample(TArrayView<const FVector3f> Positions, TArrayView<float> OutValues) const
+{
+	if (Positions.Num() == OutValues.Num())
+	{
+		for (int32 Idx = 0; Idx < Positions.Num(); ++Idx)
+		{
+			OutValues[Idx] = FMath::Abs(Plane.DistanceFromPoint((FVector)Positions[Idx]));
+		}
+	}
+	else
+	{
+		for (int32 Idx = 0; Idx < OutValues.Num(); ++Idx)
+		{
+			OutValues[Idx] = 0.f;
+		}
+	}
+}
+
+FDataflowDistanceFromPlaneFloatSamplerNode::FDataflowDistanceFromPlaneFloatSamplerNode(const UE::Dataflow::FNodeParameters& InParam, FGuid InGuid)
+	: FDataflowNode(InParam, InGuid)
+{
+	RegisterInputConnection(&DistanceFromPlaneSampler.Plane, GET_MEMBER_NAME_CHECKED(FDataflowDistanceFromPlaneFloatSampler, Plane));
+	RegisterOutputConnection(&Sampler);
+}
+
+void FDataflowDistanceFromPlaneFloatSamplerNode::Evaluate(UE::Dataflow::FContext& Context, const FDataflowOutput* Out) const
+{
+	if (Out->IsA(&Sampler))
+	{
+		TSharedRef<FDataflowDistanceFromPlaneFloatSampler> Impl = MakeShared<FDataflowDistanceFromPlaneFloatSampler>();
+
+		Impl->Plane = GetValue(Context, &DistanceFromPlaneSampler.Plane);
+
+		FDataflowFloatSampler OutSampler(Impl.ToSharedPtr());
+
+		SetValue(Context, OutSampler, &Sampler);
+	}
+}
+
+
diff --git a/Engine/Plugins/Experimental/Dataflow/Source/DataflowNodes/Private/Dataflow/SamplerNodes/DataflowDistanceFromPlaneSamplerNode.h b/Engine/Plugins/Experimental/Dataflow/Source/DataflowNodes/Private/Dataflow/SamplerNodes/DataflowDistanceFromPlaneSamplerNode.h
new file mode 100644
--- /dev/null
+++ b/Engine/Plugins/Experimental/Dataflow/Source/DataflowNodes/Private/Dataflow/SamplerNodes/DataflowDistanceFromPlaneSamplerNode.h
@@ -0,0 +1,50 @@
+// Copyright Epic Games, Inc. All Rights Reserved.
+
+#pragma once 
+
+#include "Dataflow/DataflowNode.h"
+#include "Dataflow/SamplerNodes/DataflowSamplerTypes.h"
+#include "Dataflow/DataflowPlane.h"
+
+#include "DataflowDistanceFromPlaneSamplerNode.generated.h"
+
+USTRUCT()
+struct FDataflowDistanceFromPlaneFloatSampler : public FDataflowFloatSamplerBase
+{
+	GENERATED_BODY();
+
+	virtual ~FDataflowDistanceFromPlaneFloatSampler() override = default;
+	virtual void Sample(TArrayView<const FVector3f> Positions, TArrayView<float> OutValues) const override;
+	virtual FBox GetRenderBounds() const override;
+
+public:
+	/** Plane to compute the distance from */
+	UPROPERTY(meta = (DataflowInput))
+	FDataflowPlane Plane;
+};
+
+/**
+ * 
+ * Distance from a plane float sampler
+ * Input(s) : a plane
+ * Output(s): for every sampled point the distance from the input plane
+ * 
+ */
+USTRUCT()
+struct FDataflowDistanceFromPlaneFloatSamplerNode : public FDataflowNode
+{
+	GENERATED_BODY()
+	DATAFLOW_NODE_DEFINE_INTERNAL(FDataflowDistanceFromPlaneFloatSamplerNode, "Distance From Plane Float Sampler", "Samplers", "")
+
+public:
+	FDataflowDistanceFromPlaneFloatSamplerNode(const UE::Dataflow::FNodeParameters& InParam, FGuid InGuid = FGuid::NewGuid());
+
+private:
+	UPROPERTY(meta = (DataflowOutput))
+	FDataflowFloatSampler Sampler;
+
+	UPROPERTY(EditAnywhere, Category = Sampler, meta = (ShowOnlyInnerProperties, SkipInDisplayNameChain))
+	FDataflowDistanceFromPlaneFloatSampler DistanceFromPlaneSampler;
+
+	void Evaluate(UE::Dataflow::FContext& Context, const FDataflowOutput* Out) const override;
+};
