Dataflow:
- Added FDataflowSplitVectorSamplerNode

#Synced-CL 52238556

diff --git a/Engine/Plugins/Experimental/Dataflow/Source/DataflowNodes/Private/Dataflow/SamplerNodes/DataflowSplitSamplerNode.cpp b/Engine/Plugins/Experimental/Dataflow/Source/DataflowNodes/Private/Dataflow/SamplerNodes/DataflowSplitSamplerNode.cpp
new file mode 100644
--- /dev/null
+++ b/Engine/Plugins/Experimental/Dataflow/Source/DataflowNodes/Private/Dataflow/SamplerNodes/DataflowSplitSamplerNode.cpp
@@ -0,0 +1,78 @@
+// Copyright Epic Games, Inc. All Rights Reserved.
+
+#include "Dataflow/SamplerNodes/DataflowSplitSamplerNode.h"
+#include "Dataflow/SamplerNodes/DataflowSamplerFunctions.h"
+#include "Dataflow/DataflowNodeFactory.h"
+
+FBox FDataflowVectorToFloatSampler::GetRenderBounds() const
+{
+	return FBox(FVector(-50.0), FVector(50.0));
+}
+
+void FDataflowVectorToFloatSampler::Sample(TArrayView<const FVector3f> Positions, TArrayView<float> OutValues) const
+{
+	if (VectorSampler.IsValid() && Positions.Num() == OutValues.Num() && ElementIndex >= 0 && ElementIndex <= 2)
+	{
+		TArray<FVector3f> OutValueFromVectorSampler;
+		OutValueFromVectorSampler.SetNumUninitialized(OutValues.Num());
+
+		VectorSampler->Sample(Positions, OutValueFromVectorSampler);
+
+		for (int32 Idx = 0; Idx < OutValues.Num(); ++Idx)
+		{
+			OutValues[Idx] = OutValueFromVectorSampler[Idx][ElementIndex];
+		}
+	}
+}
+
+FDataflowSplitVectorSamplerNode::FDataflowSplitVectorSamplerNode(const UE::Dataflow::FNodeParameters& InParam, FGuid InGuid)
+	: FDataflowNode(InParam, InGuid)
+{
+	RegisterInputConnection(&Sampler);
+	RegisterOutputConnection(&SamplerX);
+	RegisterOutputConnection(&SamplerY);
+	RegisterOutputConnection(&SamplerZ);
+}
+
+void FDataflowSplitVectorSamplerNode::Evaluate(UE::Dataflow::FContext& Context, const FDataflowOutput* Out) const
+{
+	if (Out->IsA(&SamplerX))
+	{
+		const FDataflowVectorSampler& InSampler = GetValue(Context, &Sampler);
+
+		TSharedRef<FDataflowVectorToFloatSampler> Impl = MakeShared<FDataflowVectorToFloatSampler>();
+
+		Impl->VectorSampler = InSampler.GetImpl();
+		Impl->ElementIndex = 0;
+
+		FDataflowFloatSampler OutSampler(Impl.ToSharedPtr());
+
+		SetValue(Context, OutSampler, &SamplerX);
+	}
+	else if (Out->IsA(&SamplerY))
+	{
+		const FDataflowVectorSampler& InSampler = GetValue(Context, &Sampler);
+
+		TSharedRef<FDataflowVectorToFloatSampler> Impl = MakeShared<FDataflowVectorToFloatSampler>();
+
+		Impl->VectorSampler = InSampler.GetImpl();
+		Impl->ElementIndex = 1;
+
+		FDataflowFloatSampler OutSampler(Impl.ToSharedPtr());
+
+		SetValue(Context, OutSampler, &SamplerY);
+	}
+	else if (Out->IsA(&SamplerZ))
+	{
+		const FDataflowVectorSampler& InSampler = GetValue(Context, &Sampler);
+
+		TSharedRef<FDataflowVectorToFloatSampler> Impl = MakeShared<FDataflowVectorToFloatSampler>();
+
+		Impl->VectorSampler = InSampler.GetImpl();
+		Impl->ElementIndex = 2;
+
+		FDataflowFloatSampler OutSampler(Impl.ToSharedPtr());
+
+		SetValue(Context, OutSampler, &SamplerZ);
+	}
+}
diff --git a/Engine/Plugins/Experimental/Dataflow/Source/DataflowNodes/Private/Dataflow/SamplerNodes/DataflowSplitSamplerNode.h b/Engine/Plugins/Experimental/Dataflow/Source/DataflowNodes/Private/Dataflow/SamplerNodes/DataflowSplitSamplerNode.h
new file mode 100644
--- /dev/null
+++ b/Engine/Plugins/Experimental/Dataflow/Source/DataflowNodes/Private/Dataflow/SamplerNodes/DataflowSplitSamplerNode.h
@@ -0,0 +1,57 @@
+// Copyright Epic Games, Inc. All Rights Reserved.
+
+#pragma once 
+
+#include "Dataflow/DataflowNode.h"
+#include "Dataflow/SamplerNodes/DataflowSamplerTypes.h"
+
+#include "DataflowSplitSamplerNode.generated.h"
+
+USTRUCT()
+struct FDataflowVectorToFloatSampler : public FDataflowFloatSamplerBase
+{
+	GENERATED_BODY();
+
+	virtual ~FDataflowVectorToFloatSampler() override = default;
+	virtual void Sample(TArrayView<const FVector3f> Positions, TArrayView<float> OutValues) const override;
+	virtual FBox GetRenderBounds() const override;
+
+public:
+	/**  */
+	TSharedPtr<const FDataflowVectorSamplerBase> VectorSampler;
+
+	UPROPERTY()
+	int32 ElementIndex = 0;
+};
+
+/**
+ *
+ * Split vector sampler
+ * Input(s): a vector sampler to split it into X, Y, Z float samplers
+ * Output(s): 3 float samplers (X, Y, Z)
+ *
+ */
+USTRUCT()
+struct FDataflowSplitVectorSamplerNode : public FDataflowNode
+{
+	GENERATED_BODY()
+	DATAFLOW_NODE_DEFINE_INTERNAL(FDataflowSplitVectorSamplerNode, "Split Vector Sampler", "Samplers", "")
+
+public:
+	FDataflowSplitVectorSamplerNode(const UE::Dataflow::FNodeParameters& InParam, FGuid InGuid = FGuid::NewGuid());
+
+private:
+	UPROPERTY(meta = (DataflowInput, DataflowIntrinsic))
+	FDataflowVectorSampler Sampler;
+
+	UPROPERTY(meta = (DataflowOutput))
+	FDataflowFloatSampler SamplerX;
+
+	UPROPERTY(meta = (DataflowOutput))
+	FDataflowFloatSampler SamplerY;
+
+	UPROPERTY(meta = (DataflowOutput))
+	FDataflowFloatSampler SamplerZ;
+
+	void Evaluate(UE::Dataflow::FContext& Context, const FDataflowOutput* Out) const override;
+};
