Dataflow:
- Updated DataflowSamplerToAttributeNode to work with anytype samplers

#Synced-CL 52238782

diff --git a/Engine/Plugins/Experimental/Dataflow/Source/DataflowNodes/Private/Dataflow/DataflowSamplerToAttributeNode.cpp b/Engine/Plugins/Experimental/Dataflow/Source/DataflowNodes/Private/Dataflow/DataflowSamplerToAttributeNode.cpp
--- a/Engine/Plugins/Experimental/Dataflow/Source/DataflowNodes/Private/Dataflow/DataflowSamplerToAttributeNode.cpp
+++ b/Engine/Plugins/Experimental/Dataflow/Source/DataflowNodes/Private/Dataflow/DataflowSamplerToAttributeNode.cpp
@@ -14,10 +14,15 @@ namespace UE::Dataflow
 {
 	void RegisterSamplerToAttributeNodes()
 	{
+		DATAFLOW_NODE_REGISTER_CREATION_FACTORY(FDataflowSamplerToAttributeNode);
+
+		// Deprecated
 		DATAFLOW_NODE_REGISTER_CREATION_FACTORY(FDataflowFloatSamplerToAttributeNode);
 	}
 };
 
+// ----------------------------------------------------------------------------------------------------------------------
+
 FDataflowFloatSamplerToAttributeNode::FDataflowFloatSamplerToAttributeNode(const UE::Dataflow::FNodeParameters& InParam, FGuid InGuid)
 	: FDataflowNode(InParam, InGuid)
 {
@@ -74,7 +79,7 @@ void FDataflowFloatSamplerToAttributeNode::Evaluate(UE::Dataflow::FContext& Cont
 			TArrayView<float> Values = MakeArrayView<float>(TargetAttribute.GetData(), TargetAttribute.Num());
 			InSampler.Sample(VertexAttribute->GetConstArray(), Values);
 		}
-		
+
 		SetValue(Context, MoveTemp(InCollection), &Collection);
 	}
 }
@@ -88,5 +93,145 @@ FDataflowNode::FAttributeKey FDataflowFloatSamplerToAttributeNode::GetVertexAttr
 	};
 }
 
+// ----------------------------------------------------------------------------------------------------------------------
+
+FDataflowSamplerToAttributeNode::FDataflowSamplerToAttributeNode(const UE::Dataflow::FNodeParameters& InParam, FGuid InGuid)
+	: FDataflowNode(InParam, InGuid)
+{
+	RegisterInputConnection(&Collection);
+	RegisterInputConnection(&Sampler);
+	RegisterInputConnection(&AttributeName);
+	RegisterOutputConnection(&Collection, &Collection);
+	RegisterOutputConnection(&AttributeName, &AttributeName);
+}
+
+void FDataflowSamplerToAttributeNode::Evaluate(UE::Dataflow::FContext& Context, const FDataflowOutput* Out) const
+{
+	if (Out->IsA(&Collection))
+	{
+		if (IsConnected(&Collection))
+		{
+			const FName InAttributeName = FName(GetValue(Context, &AttributeName, AttributeName));
+			if (InAttributeName.IsNone())
+			{
+				Context.Error(LOCTEXT("NoAttributeName_Msg", "Attribute name is empty, data will not be set"));
+				SafeForwardInput(Context, &Collection, &Collection);
+				return;
+			}
+
+			FManagedArrayCollection InCollection = GetValue(Context, &Collection, Collection);
+
+			const FName InVertexGroupName = VertexGroup.Name;
+			if (!InCollection.HasGroup(InVertexGroupName))
+			{
+				Context.Error(LOCTEXT("InvalidVertexGroup_Msg", "Vertex group does not exists, data will not be set"));
+			
+				SafeForwardInput(Context, &Collection, &Collection);
+				return;
+			}
+
+			const TManagedArray<FVector3f>* VertexAttribute = InCollection.FindAttribute<FVector3f>(FGeometryCollection::VertexPositionAttribute, InVertexGroupName);
+			if (!VertexAttribute)
+			{
+				// cannot really include cloth here so we need to directly search for the vertex position attribute in the render vertices ( eventually we may converge to use only geometry collections )
+				VertexAttribute = InCollection.FindAttribute<FVector3f>("RenderPosition", InVertexGroupName);
+			}
+			if (!VertexAttribute)
+			{
+				// cannot really include cloth here so we need to directly search for the vertex position  attribute in the render vertices ( eventually we may converge to use only geometry collections )
+				VertexAttribute = InCollection.FindAttribute<FVector3f>("SimPosition3D", InVertexGroupName);
+			}
+			if (!VertexAttribute)
+			{
+				Context.Warning(LOCTEXT("NoVertexPosition_Msg", "No vertex position could be found in this collection, no data will be set and the attribute will be inititialized to zeros"));
+			}
+
+			if (const FDataflowFloatSampler* InFloatSampler = GetValue(Context, &Sampler).TryGet <FDataflowFloatSampler>())
+			{
+				if (!bSaveAsColor)
+				{
+					TManagedArray<float>& TargetAttribute = InCollection.AddAttribute<float>(InAttributeName, InVertexGroupName);
+					TargetAttribute.Fill(0);
+
+					if (VertexAttribute)
+					{
+						TArrayView<float> Values = MakeArrayView<float>(TargetAttribute.GetData(), TargetAttribute.Num());
+						InFloatSampler->Sample(VertexAttribute->GetConstArray(), Values);
+					}
+				}
+				else
+				{
+					TManagedArray<FLinearColor>& TargetAttribute = InCollection.AddAttribute<FLinearColor>(InAttributeName, InVertexGroupName);
+					TargetAttribute.Fill(FLinearColor(FVector::ZeroVector));
+
+					if (VertexAttribute)
+					{
+						TArrayView<FLinearColor> Values = MakeArrayView<FLinearColor>(TargetAttribute.GetData(), TargetAttribute.Num());
+
+						TArray<float> SampledValues;
+						SampledValues.SetNum(Values.Num());
+
+						InFloatSampler->Sample(VertexAttribute->GetConstArray(), SampledValues);
+
+						for (int32 Idx = 0; Idx < SampledValues.Num(); ++Idx)
+						{
+							Values[Idx] = FLinearColor(SampledValues[Idx], SampledValues[Idx], SampledValues[Idx]);
+						}
+					}
+				}
+			}
+			else if (const FDataflowVectorSampler* InVectorSampler = GetValue(Context, &Sampler).TryGet<FDataflowVectorSampler>())
+			{
+				if (!bSaveAsColor)
+				{
+					TManagedArray<FVector3f>& TargetAttribute = InCollection.AddAttribute<FVector3f>(InAttributeName, InVertexGroupName);
+					TargetAttribute.Fill(FVector3f::ZeroVector);
+
+					if (VertexAttribute)
+					{
+						TArrayView<FVector3f> Values = MakeArrayView<FVector3f>(TargetAttribute.GetData(), TargetAttribute.Num());
+						InVectorSampler->Sample(VertexAttribute->GetConstArray(), Values);
+					}
+				}
+				else
+				{
+					TManagedArray<FLinearColor>& TargetAttribute = InCollection.AddAttribute<FLinearColor>(InAttributeName, InVertexGroupName);
+					TargetAttribute.Fill(FLinearColor(FVector::ZeroVector));
+
+					if (VertexAttribute)
+					{
+						TArrayView<FLinearColor> Values = MakeArrayView<FLinearColor>(TargetAttribute.GetData(), TargetAttribute.Num());
+
+						TArray<FVector3f> SampledValues;
+						SampledValues.SetNum(Values.Num());
+
+						InVectorSampler->Sample(VertexAttribute->GetConstArray(), SampledValues);
+
+						for (int32 Idx = 0; Idx < SampledValues.Num(); ++Idx)
+						{
+							Values[Idx] = FLinearColor(SampledValues[Idx]);
+						}
+					}
+				}
+			}
+
+			SetValue(Context, MoveTemp(InCollection), &Collection);
+		}
+	}
+	else if (Out->IsA(&AttributeName))
+	{
+		SafeForwardInput(Context, &AttributeName, &AttributeName);
+	}
+}
+
+FDataflowNode::FAttributeKey FDataflowSamplerToAttributeNode::GetVertexAttributeToVisualize(UE::Dataflow::FContext& Context, const FName OutputName, const FName RenderGroup) const
+{
+	return FAttributeKey
+	{
+		.AttributeName = FName(GetValue(Context, &AttributeName, AttributeName)),
+		.GroupName = VertexGroup.Name,
+	};
+}
+
 #undef LOCTEXT_NAMESPACE 
 
diff --git a/Engine/Plugins/Experimental/Dataflow/Source/DataflowNodes/Private/Dataflow/DataflowSamplerToAttributeNode.h b/Engine/Plugins/Experimental/Dataflow/Source/DataflowNodes/Private/Dataflow/DataflowSamplerToAttributeNode.h
--- a/Engine/Plugins/Experimental/Dataflow/Source/DataflowNodes/Private/Dataflow/DataflowSamplerToAttributeNode.h
+++ b/Engine/Plugins/Experimental/Dataflow/Source/DataflowNodes/Private/Dataflow/DataflowSamplerToAttributeNode.h
@@ -11,8 +11,9 @@
 
 /**
 * Use a foat sampler to generate the value of an attribute
+* DEPRECATED 5.8 - use FDataflowSamplerToAttributeNode instead
 */
-USTRUCT()
+USTRUCT(meta = (Deprecated = "5.8"))
 struct FDataflowFloatSamplerToAttributeNode : public FDataflowNode
 {
 	GENERATED_USTRUCT_BODY()
@@ -42,6 +43,48 @@ private:
 	virtual FAttributeKey GetVertexAttributeToVisualize(UE::Dataflow::FContext& Context, const FName OutputName, const FName RenderGroup) const override;
 };
 
+// ----------------------------------------------------------------------------------------------------------------------
+
+/**
+* Use a float or vector sampler to generate the value of an attribute
+*/
+USTRUCT()
+struct FDataflowSamplerToAttributeNode : public FDataflowNode
+{
+	GENERATED_USTRUCT_BODY()
+	DATAFLOW_NODE_DEFINE_INTERNAL(FDataflowSamplerToAttributeNode, "SamplerToAttribute", "Dataflow", "Transfer Generate")
+
+public:
+	FDataflowSamplerToAttributeNode(const UE::Dataflow::FNodeParameters& InParam, FGuid InGuid = FGuid::NewGuid());
+
+private:
+	/** Collection to set the attribute data to */
+	UPROPERTY(meta = (DataflowInput, DataflowOutput, DataflowPassthrough = Collection, DataflowIntrinsic))
+	FManagedArrayCollection Collection;
+
+	/** Sampler to use to set the value of the attribute */
+	UPROPERTY(EditAnywhere, Category = Source, meta = (DataflowInput))
+	FDataflowSamplerTypes Sampler;
+
+	/** Name of the attribute to set the value on */
+	UPROPERTY(EditAnywhere, Category = Attribute, meta = (DataflowInput, DataflowOutput, DataflowPassthrough = AttributeName))
+	FString AttributeName = FString("Dummy");
+
+	/** If Sampler input is a FloatSampler and this flag is on then data will be saved as LinearColor using the 
+	* sampled float values to create grayscale colors,
+	* If Sampler input is a VectorSampler type and this flag is on then data will be saved as LinearColor using 
+	* the sampled vector values to set the colors */
+	UPROPERTY(EditAnywhere, Category = Attribute)
+	bool bSaveAsColor = false;
+
+	/** Vertex group to use */
+	UPROPERTY(EditAnywhere, Category = Target)
+	FScalarVertexPropertyGroup VertexGroup;
+
+	virtual void Evaluate(UE::Dataflow::FContext& Context, const FDataflowOutput* Out) const override;
+	virtual FAttributeKey GetVertexAttributeToVisualize(UE::Dataflow::FContext& Context, const FName OutputName, const FName RenderGroup) const override;
+};
+
 namespace UE::Dataflow
 {
 	void RegisterSamplerToAttributeNodes();
