Dataflow Editor : Fix simulation timeline to match animation asset length

#Synced-CL 52232971

diff --git a/Engine/Plugins/Experimental/Dataflow/Source/DataflowEditor/Private/Widgets/SDataflowSimulationTimeline.cpp b/Engine/Plugins/Experimental/Dataflow/Source/DataflowEditor/Private/Widgets/SDataflowSimulationTimeline.cpp
--- a/Engine/Plugins/Experimental/Dataflow/Source/DataflowEditor/Private/Widgets/SDataflowSimulationTimeline.cpp
+++ b/Engine/Plugins/Experimental/Dataflow/Source/DataflowEditor/Private/Widgets/SDataflowSimulationTimeline.cpp
@@ -64,7 +64,7 @@ void SDataflowSimulationTimeline::Construct(const FArguments& InArgs, const TSha
 		TimeSliderArgs.TickResolution = TickResolution;
 		TimeSliderArgs.OnViewRangeChanged = FOnViewRangeChanged::CreateSP(&InSimulationBinding.Get(), &FDataflowSimulationBinding::HandleViewRangeChanged);
 		TimeSliderArgs.OnClampRangeChanged = FOnTimeRangeChanged::CreateSP(&InSimulationBinding.Get(), &FDataflowSimulationBinding::HandleWorkingRangeChanged);
-		TimeSliderArgs.IsPlaybackRangeLocked = false;
+		TimeSliderArgs.IsPlaybackRangeLocked = true;
 		TimeSliderArgs.PlaybackStatus = EMovieScenePlayerStatus::Stopped;
 		TimeSliderArgs.NumericTypeInterface = NumericTypeInterface;
 		TimeSliderArgs.OnScrubPositionChanged = FOnScrubPositionChanged::CreateSP(this, &SDataflowSimulationTimeline::HandleScrubPositionChanged);
