Dataflow : update timeline so that its view and active range changes base don a the selected animation
when no animation is selected the cache range will be used

#Synced-CL 52232971

diff --git a/Engine/Plugins/Experimental/Dataflow/Source/DataflowEditor/Private/Dataflow/DataflowSimulationControls.cpp b/Engine/Plugins/Experimental/Dataflow/Source/DataflowEditor/Private/Dataflow/DataflowSimulationControls.cpp
--- a/Engine/Plugins/Experimental/Dataflow/Source/DataflowEditor/Private/Dataflow/DataflowSimulationControls.cpp
+++ b/Engine/Plugins/Experimental/Dataflow/Source/DataflowEditor/Private/Dataflow/DataflowSimulationControls.cpp
@@ -132,8 +132,10 @@ namespace UE::Dataflow
 			FInstancedPropertyBag());
 	}
 
-	void SetupSkeletonAnimation(const TObjectPtr<AActor>& PreviewActor, const bool bSkeletalMeshVisibility)
+	bool SetupSkeletonAnimation(const TObjectPtr<AActor>& PreviewActor, const bool bSkeletalMeshVisibility)
 	{
+		bool bAnimationWasSet = false;
+
 		if(PreviewActor)
 		{
 			TInlineComponentArray<UPrimitiveComponent*> PrimComponents;
@@ -178,11 +180,14 @@ namespace UE::Dataflow
 						// Stop the animation 
 						AnimNodeInstance->SetLooping(true);
 						AnimNodeInstance->SetPlaying(false);
-						
+
+						bAnimationWasSet = true;
 					}
 				}
 			}
 		}
+	
+		return bAnimationWasSet;
 	}
 
 	static void FillAnimationDatas(const UAnimSequenceBase* AnimSequence, const float CurrentTime, USkeletalMeshComponent* SkeletalMeshComponent)
diff --git a/Engine/Plugins/Experimental/Dataflow/Source/DataflowEditor/Private/Dataflow/DataflowSimulationScene.cpp b/Engine/Plugins/Experimental/Dataflow/Source/DataflowEditor/Private/Dataflow/DataflowSimulationScene.cpp
--- a/Engine/Plugins/Experimental/Dataflow/Source/DataflowEditor/Private/Dataflow/DataflowSimulationScene.cpp
+++ b/Engine/Plugins/Experimental/Dataflow/Source/DataflowEditor/Private/Dataflow/DataflowSimulationScene.cpp
@@ -297,6 +297,43 @@ void FDataflowSimulationScene::BindSceneSelection()
 	}
 }
 
+void FDataflowSimulationScene::UpdateTimelinedataFromAnimations()
+{
+	float MaxAnimLength = 0.0f;
+	if (PreviewActor)
+	{
+		PreviewActor->ForEachComponent<USkeletalMeshComponent>(true, 
+			[&MaxAnimLength](USkeletalMeshComponent* SkeletalMeshComponent)
+			{
+				if (UAnimSingleNodeInstance* AnimInstance = SkeletalMeshComponent->GetSingleNodeInstance())
+				{
+					MaxAnimLength = FMath::Max(MaxAnimLength, AnimInstance->GetLength());
+				}
+			});
+	}
+
+	if (MaxAnimLength > 0.0f)
+	{
+		// Set the time range to start at 0 and end at the animation length.
+		// Enforce a minimum time range of one frame so the timeline remains usable with very short animations.
+		TimeRange[0] = 0.0f;
+		TimeRange[1] = FMath::Max(MaxAnimLength, DeltaTime);
+	}
+	else
+	{
+		// No animation - reset to the cache params time range
+		TimeRange = SceneDescription->CacheParams.TimeRange;
+	}
+	const float ActualLength = TimeRange[1] - TimeRange[0];
+
+	NumFrames = FMath::Max(1, (ActualLength > 0) ? FMath::Floor((ActualLength + UE_SMALL_NUMBER) * SceneDescription->CacheParams.FrameRate) + 1 : 1);
+
+	if (OnTimeRangeUpdated.IsBound())
+	{
+		OnTimeRangeUpdated.Broadcast(*this);
+	}
+}
+
 void FDataflowSimulationScene::CreateSimulationScene()
 {
 	const int32 PreviewLOD = GetPreviewLOD();
@@ -318,7 +355,11 @@ void FDataflowSimulationScene::CreateSimulationScene()
 			SceneDescription->CacheAsset, false, GetEditorContent(), SceneDescription->BlueprintTransform, SceneDescription->BlueprintVariables);
 		SimulationGenerator->SetDeltaTime(DeltaTime);
 		// Setup all the skelmesh animations
-		UE::Dataflow::SetupSkeletonAnimation(PreviewActor, SceneDescription->bSkeletalMeshVisibility);
+		const bool bAnimationWasSet = UE::Dataflow::SetupSkeletonAnimation(PreviewActor, SceneDescription->bSkeletalMeshVisibility);
+		if (bAnimationWasSet)
+		{
+			UpdateTimelinedataFromAnimations();
+		}
 
 		if(PreviewActor)
 		{
diff --git a/Engine/Plugins/Experimental/Dataflow/Source/DataflowEditor/Private/Widgets/DataflowSimulationBinding.cpp b/Engine/Plugins/Experimental/Dataflow/Source/DataflowEditor/Private/Widgets/DataflowSimulationBinding.cpp
--- a/Engine/Plugins/Experimental/Dataflow/Source/DataflowEditor/Private/Widgets/DataflowSimulationBinding.cpp
+++ b/Engine/Plugins/Experimental/Dataflow/Source/DataflowEditor/Private/Widgets/DataflowSimulationBinding.cpp
@@ -10,14 +10,27 @@
 FDataflowSimulationBinding::FDataflowSimulationBinding(const TWeakPtr<FDataflowSimulationScene>&  InSimulationScene)
 	: SimulationScene(InSimulationScene)
 {
-	if (const TSharedPtr<const FDataflowSimulationScene> PreviewScene = SimulationScene.Pin())
+	if (const TSharedPtr<FDataflowSimulationScene> PreviewScene = SimulationScene.Pin())
 	{
-		// todo use the animation and playback range and find the common part
-		const FAnimatedRange PlaybackRange(PreviewScene->GetTimeRange()[0], PreviewScene->GetTimeRange()[1]);
-		SetViewRange(PlaybackRange);
+		OnSceneTimeRangeUpdated(*PreviewScene);
+		PreviewScene->OnTimeRangeUpdated.AddRaw(this, &FDataflowSimulationBinding::OnSceneTimeRangeUpdated);
 	}
 }
 
+FDataflowSimulationBinding::~FDataflowSimulationBinding()
+{
+	if (const TSharedPtr<FDataflowSimulationScene> PreviewScene = SimulationScene.Pin())
+	{
+		PreviewScene->OnTimeRangeUpdated.RemoveAll(this);
+	}
+}
+
+void FDataflowSimulationBinding::OnSceneTimeRangeUpdated(const FDataflowSimulationScene& Scene)
+{
+	const FAnimatedRange PlaybackRange(Scene.GetTimeRange()[0], Scene.GetTimeRange()[1]);
+	SetViewRange(PlaybackRange);
+}
+
 FFrameRate FDataflowSimulationBinding::GetFrameRate() const
 {
 	if (const TSharedPtr<const FDataflowSimulationScene> PreviewScene = SimulationScene.Pin())
diff --git a/Engine/Plugins/Experimental/Dataflow/Source/DataflowEditor/Private/Widgets/DataflowSimulationBinding.h b/Engine/Plugins/Experimental/Dataflow/Source/DataflowEditor/Private/Widgets/DataflowSimulationBinding.h
--- a/Engine/Plugins/Experimental/Dataflow/Source/DataflowEditor/Private/Widgets/DataflowSimulationBinding.h
+++ b/Engine/Plugins/Experimental/Dataflow/Source/DataflowEditor/Private/Widgets/DataflowSimulationBinding.h
@@ -14,6 +14,7 @@ class FDataflowSimulationBinding : public TSharedFromThis<FDataflowSimulationBin
 {
 public:
 	FDataflowSimulationBinding(const TWeakPtr<FDataflowSimulationScene>&  InSimulationScene);
+	~FDataflowSimulationBinding();
 
 	/** Get the simulation scene from the binding */
 	TWeakPtr<FDataflowSimulationScene> GetSimulationScene() const { return SimulationScene; }
@@ -79,6 +80,9 @@ public:
 	void FillCacheNames(TArray<FString>& CacheNames) const;
 
 private:
+
+	void OnSceneTimeRangeUpdated(const FDataflowSimulationScene& Scene);
+
 	/** Simulation scene to be used for the widget */
 	TWeakPtr<FDataflowSimulationScene> SimulationScene;
 
diff --git a/Engine/Plugins/Experimental/Dataflow/Source/DataflowEditor/Public/Dataflow/DataflowSimulationControls.h b/Engine/Plugins/Experimental/Dataflow/Source/DataflowEditor/Public/Dataflow/DataflowSimulationControls.h
--- a/Engine/Plugins/Experimental/Dataflow/Source/DataflowEditor/Public/Dataflow/DataflowSimulationControls.h
+++ b/Engine/Plugins/Experimental/Dataflow/Source/DataflowEditor/Public/Dataflow/DataflowSimulationControls.h
@@ -28,8 +28,11 @@ namespace UE::Dataflow
 		const TObjectPtr<AChaosCacheManager>& CacheManager, const TObjectPtr<UChaosCacheCollection>& CacheCollection,
 		const bool bIsRecording, const TObjectPtr<UDataflowBaseContent>& DataflowContent, const FTransform& ActorTransform);
 
-	/** Setup the skelmesh animations to be used in the scene/generator */
-	void SetupSkeletonAnimation(const TObjectPtr<AActor>& PreviewActor, const bool bSkeletalMeshVisibility = true);
+	/** 
+	* Setup the skelmesh animations to be used in the scene/generator 
+	* return true if the animatiuon was properly set 
+	*/
+	bool SetupSkeletonAnimation(const TObjectPtr<AActor>& PreviewActor, const bool bSkeletalMeshVisibility = true);
 
 	/** Update the skelmesh animation at some point in time (GT) */
 	void UpdateSkeletonAnimation(const TObjectPtr<AActor>& PreviewActor, const float SimulationTime);
diff --git a/Engine/Plugins/Experimental/Dataflow/Source/DataflowEditor/Public/Dataflow/DataflowSimulationScene.h b/Engine/Plugins/Experimental/Dataflow/Source/DataflowEditor/Public/Dataflow/DataflowSimulationScene.h
--- a/Engine/Plugins/Experimental/Dataflow/Source/DataflowEditor/Public/Dataflow/DataflowSimulationScene.h
+++ b/Engine/Plugins/Experimental/Dataflow/Source/DataflowEditor/Public/Dataflow/DataflowSimulationScene.h
@@ -244,6 +244,9 @@ public:
 	/** Set the locked simulation flag */
 	void SetSimulationLocked(const bool bSimulationLocked, const bool bIsPlaying);
 
+	DECLARE_MULTICAST_DELEGATE_OneParam(FOnTimeRangeUpdated, const FDataflowSimulationScene&);
+	FOnTimeRangeUpdated OnTimeRangeUpdated;
+
 private:
 	friend UDataflowSimulationSceneDescription;
 
@@ -259,6 +262,8 @@ private:
 	/** Sync the blueprint preview properties with the edited asset. */
 	void WritePreviewPropertiesToDataflowOwner() const;
 
+	void UpdateTimelinedataFromAnimations();
+
 	/** Simulation scene description */
 	TObjectPtr<UDataflowSimulationSceneDescription> SceneDescription;
 
