Add SDataflowTemplatePicker Slate dialog widget

Modal Slate dialog for choosing a Dataflow setup when an asset triggers
the corresponding UI. Supports template tiles (name + icon), an embedded
content browser picker for existing Dataflow assets, and a No Dataflow option.

#Synced-CL 52234830

diff --git a/Engine/Plugins/Experimental/Dataflow/Source/DataflowEditor/Private/Dataflow/AssetDefinition_DataflowAsset.cpp b/Engine/Plugins/Experimental/Dataflow/Source/DataflowEditor/Private/Dataflow/AssetDefinition_DataflowAsset.cpp
--- a/Engine/Plugins/Experimental/Dataflow/Source/DataflowEditor/Private/Dataflow/AssetDefinition_DataflowAsset.cpp
+++ b/Engine/Plugins/Experimental/Dataflow/Source/DataflowEditor/Private/Dataflow/AssetDefinition_DataflowAsset.cpp
@@ -9,6 +9,8 @@
 #include "Dataflow/DataflowEditorModule.h"
 #include "Dataflow/DataflowObject.h"
 #include "Dataflow/DataflowInstance.h"
+#include "Dataflow/DataflowTemplateRegistry.h"
+#include "Dataflow/SDataflowTemplatePicker.h"
 #include "Dialog/SMessageDialog.h"
 #include "Framework/Notifications/NotificationManager.h"
 #include "IContentBrowserSingleton.h"
@@ -19,7 +21,6 @@
 #include "ToolMenus.h"
 #include "Widgets/Notifications/SNotificationList.h"
 
-
 #include UE_INLINE_GENERATED_CPP_BY_NAME(AssetDefinition_DataflowAsset)
 
 
@@ -215,6 +216,79 @@ namespace UE::DataflowAssetDefinitionHelpers
 		return true;
 	}
 
+	/**
+	 * Shows the Dataflow template picker dialog and applies the user's choice directly onto Asset:
+	 *   - Template : duplicates the chosen template and embeds it into Asset (owned, no standalone flags).
+	 *   - Existing asset : sets the selected shared UDataflow on Asset without duplication.
+	 *   - None / Cancelled : leaves Asset unchanged.
+	 *
+	 * Asset must implement IDataflowInstanceInterface for the Template/ExistingAsset paths to take effect.
+	 * Always returns true — the picker is a single non-repeating step with no sub-dialogs.
+	 */
+	bool SetDataflowFromTemplatePicker(UObject* Asset)
+	{
+		const TArray<FDataflowTemplateOption> Templates =
+			FDataflowTemplateRegistry::Get().GetTemplateOptions(Asset->GetClass());
+
+		const FDataflowPickerResult PickerResult = SDataflowTemplatePicker::ShowModal(
+			Templates,
+			FSlateApplication::Get().GetActiveTopLevelWindow());
+
+		switch (PickerResult.Type)
+		{
+		case EDataflowPickerResult::Template:
+		{
+			IDataflowInstanceInterface* Interface = Cast<IDataflowInstanceInterface>(Asset);
+			if (!Interface)
+			{
+				return true;
+			}
+
+			// Load the chosen template; nullptr means Blank (TemplateId == NAME_None).
+			const FString TemplatePath = PickerResult.TemplateId.ToString();
+			UDataflow* const Source = (TemplatePath.IsEmpty() || PickerResult.TemplateId.IsNone())
+				? nullptr
+				: LoadObject<UDataflow>(nullptr, *TemplatePath);
+
+			// Embed a duplicate of the template, or a new empty graph for Blank.
+			// Embedded dataflows are owned by the asset — no standalone/public flags.
+			UDataflow* Embedded = Source
+				? DuplicateObject<UDataflow>(Source, Asset, TEXT("EmbeddedDataflow"))
+				: NewObject<UDataflow>(Asset, UDataflow::StaticClass(), TEXT("EmbeddedDataflow"), RF_Transactional);
+
+			if (Embedded)
+			{
+				if (Source)
+				{
+					Embedded->ClearFlags(RF_Public | RF_Standalone);
+					Embedded->SetFlags(RF_Transactional);
+				}
+				Interface->GetDataflowInstance().SetDataflowAsset(Embedded);
+				Asset->Modify();
+			}
+			return true;
+		}
+
+		case EDataflowPickerResult::ExistingAsset:
+		{
+			// Reference the selected asset directly — no duplication.
+			IDataflowInstanceInterface* Interface = Cast<IDataflowInstanceInterface>(Asset);
+			UDataflow* const SharedDataflow = Cast<UDataflow>(PickerResult.SelectedAsset.GetAsset());
+			if (Interface && SharedDataflow)
+			{
+				Interface->GetDataflowInstance().SetDataflowAsset(SharedDataflow);
+				Asset->Modify();
+			}
+			return true;
+		}
+
+		case EDataflowPickerResult::None:
+		case EDataflowPickerResult::Cancelled:
+		default:
+			return true;
+		}
+	}
+
 	// Create a new UDataflow if one doesn't already exist for the Asset
 	UObject* NewOrOpenDataflowAsset(const UObject* Asset)
 	{
diff --git a/Engine/Plugins/Experimental/Dataflow/Source/DataflowEditor/Private/Dataflow/DataflowTemplateRegistry.cpp b/Engine/Plugins/Experimental/Dataflow/Source/DataflowEditor/Private/Dataflow/DataflowTemplateRegistry.cpp
new file mode 100644
--- /dev/null
+++ b/Engine/Plugins/Experimental/Dataflow/Source/DataflowEditor/Private/Dataflow/DataflowTemplateRegistry.cpp
@@ -0,0 +1,140 @@
+// Copyright Epic Games, Inc. All Rights Reserved.
+
+#include "Dataflow/DataflowTemplateRegistry.h"
+
+#include "AssetRegistry/AssetRegistryModule.h"
+#include "AssetRegistry/IAssetRegistry.h"
+#include "Dataflow/DataflowObject.h"
+#include "Styling/SlateIconFinder.h"
+
+#define LOCTEXT_NAMESPACE "DataflowTemplateRegistry"
+
+FDataflowTemplateRegistry& FDataflowTemplateRegistry::Get()
+{
+	static FDataflowTemplateRegistry Instance;
+	return Instance;
+}
+
+void FDataflowTemplateRegistry::RegisterTemplateFolder(
+	const UClass* AssetClass, FDataflowTemplateFolderRegistration Registration)
+{
+	if (!ensureMsgf(AssetClass, TEXT("FDataflowTemplateRegistry::RegisterTemplateFolder: null AssetClass")))
+	{
+		return;
+	}
+	Registrations.FindOrAdd(AssetClass->GetClassPathName()).Add(MoveTemp(Registration));
+}
+
+TArray<FDataflowTemplateOption> FDataflowTemplateRegistry::GetTemplateOptions(
+	const UClass* AssetClass) const
+{
+	TArray<FDataflowTemplateOption> Result;
+
+	const FSlateIcon DataflowDefaultIcon(FName("DataflowEditorStyle"), "ClassIcon.Dataflow");
+
+	// "Blank" is always the first option.
+	{
+		FDataflowTemplateOption& Blank = Result.AddDefaulted_GetRef();
+		Blank.DisplayName = LOCTEXT("Blank_Name",    "Blank");
+		Blank.Tooltip     = LOCTEXT("Blank_Tip",     "Start from an empty Dataflow graph");
+		Blank.Icon = DataflowDefaultIcon.GetIcon();
+		Blank.TemplateId  = NAME_None;   // NAME_None > caller creates a blank UDataflow
+	}
+
+	// Walk the class hierarchy (most-derived first) to collect registered folders,
+	// preserving insertion order and deduplicating by path.
+	TArray<FString>           OrderedFolders;
+	TMap<FString, FSlateIcon> FolderToIcon;
+
+	for (const UClass* Class = AssetClass; Class; Class = Class->GetSuperClass())
+	{
+		const TArray<FDataflowTemplateFolderRegistration>* Regs =
+			Registrations.Find(Class->GetClassPathName());
+		if (!Regs)
+		{
+			continue;
+		}
+		for (const FDataflowTemplateFolderRegistration& Reg : *Regs)
+		{
+			if (!FolderToIcon.Contains(Reg.ContentFolder))
+			{
+				// First registration for this folder wins (most-derived class).
+				FolderToIcon.Add(Reg.ContentFolder) = Reg.Icon;
+				OrderedFolders.Add(Reg.ContentFolder);
+			}
+		}
+	}
+
+	if (OrderedFolders.IsEmpty())
+	{
+		return Result;
+	}
+
+	// Query the asset registry for UDataflow assets in the registered folders.
+	IAssetRegistry& AssetRegistry =
+		FModuleManager::LoadModuleChecked<FAssetRegistryModule>(TEXT("AssetRegistry")).Get();
+
+	FARFilter Filter;
+	Filter.ClassPaths.Add(UDataflow::StaticClass()->GetClassPathName());
+	Filter.bRecursiveClasses = false;
+	Filter.bRecursivePaths   = true;
+	for (const FString& Folder : OrderedFolders)
+	{
+		Filter.PackagePaths.Add(FName(*Folder));
+	}
+
+	TArray<FAssetData> FoundAssets;
+	AssetRegistry.GetAssets(Filter, FoundAssets);
+
+	// Sort by folder order first, then by asset name, so results are stable
+	// regardless of asset registry enumeration order.
+	FoundAssets.Sort([&OrderedFolders](const FAssetData& A, const FAssetData& B)
+	{
+		const FString PathA = A.PackagePath.ToString();
+		const FString PathB = B.PackagePath.ToString();
+		const int32 FolderA = OrderedFolders.IndexOfByPredicate(
+			[&PathA](const FString& F){ return PathA.StartsWith(F); });
+		const int32 FolderB = OrderedFolders.IndexOfByPredicate(
+			[&PathB](const FString& F){ return PathB.StartsWith(F); });
+		if (FolderA != FolderB)
+		{
+			return FolderA < FolderB;
+		}
+		return A.AssetName.LexicalLess(B.AssetName);
+	});
+
+	for (const FAssetData& AssetData : FoundAssets)
+	{
+		// Resolve icon from the folder this asset lives under.
+		const FString AssetPackagePath = AssetData.PackagePath.ToString();
+		const FSlateIcon* FoundIcon = nullptr;
+		for (const FString& Folder : OrderedFolders)
+		{
+			if (AssetPackagePath.StartsWith(Folder))
+			{
+				FoundIcon = FolderToIcon.Find(Folder);
+				break;
+			}
+		}
+
+		const FSlateBrush* Icon = (FoundIcon && FoundIcon->IsSet())
+			? FoundIcon->GetIcon()
+			: DataflowDefaultIcon.GetIcon();
+
+		// Build a human-readable display name: strip leading "DF_", replace underscores.
+		FString DisplayName = AssetData.AssetName.ToString();
+		DisplayName.RemoveFromStart(TEXT("DF_"));
+		DisplayName = DisplayName.Replace(TEXT("_"), TEXT(" "));
+
+		FDataflowTemplateOption& Option = Result.AddDefaulted_GetRef();
+		Option.DisplayName = FText::FromString(DisplayName);
+		Option.Tooltip     = FText::FromString(AssetData.GetObjectPathString());
+		Option.Icon        = Icon;
+		// Full object path stored in TemplateId so callers can LoadObject<UDataflow> directly.
+		Option.TemplateId  = FName(AssetData.GetObjectPathString());
+	}
+
+	return Result;
+}
+
+#undef LOCTEXT_NAMESPACE
diff --git a/Engine/Plugins/Experimental/Dataflow/Source/DataflowEditor/Private/Dataflow/SDataflowTemplatePicker.cpp b/Engine/Plugins/Experimental/Dataflow/Source/DataflowEditor/Private/Dataflow/SDataflowTemplatePicker.cpp
new file mode 100644
--- /dev/null
+++ b/Engine/Plugins/Experimental/Dataflow/Source/DataflowEditor/Private/Dataflow/SDataflowTemplatePicker.cpp
@@ -0,0 +1,419 @@
+// Copyright Epic Games, Inc. All Rights Reserved.
+
+#include "Dataflow/SDataflowTemplatePicker.h"
+
+#include "ContentBrowserModule.h"
+#include "IContentBrowserSingleton.h"
+#include "Styling/AppStyle.h"
+#include "Widgets/Images/SImage.h"
+#include "Widgets/Input/SButton.h"
+#include "Widgets/Input/SCheckBox.h"
+#include "Widgets/Input/SComboButton.h"
+#include "Widgets/Layout/SBox.h"
+#include "Widgets/Layout/SScrollBox.h"
+#include "Widgets/Layout/SSeparator.h"
+#include "Widgets/Layout/SWrapBox.h"
+#include "Widgets/Text/STextBlock.h"
+
+#include "Dataflow/DataflowObject.h"
+
+#define LOCTEXT_NAMESPACE "DataflowTemplatePicker"
+
+namespace DataflowTemplatePickerPrivate
+{
+	static constexpr float TileSize     = 128.f;
+	static constexpr float TileIconSize =  64.f;
+}
+
+FDataflowPickerResult SDataflowTemplatePicker::ShowModal(
+	const TArray<FDataflowTemplateOption>& Templates,
+	TSharedPtr<SWindow>                    ParentWindow)
+{
+	TSharedRef<SWindow> Window = SNew(SWindow)
+		.Title(LOCTEXT("DialogTitle", "Choose Dataflow Setup"))
+		.SizingRule(ESizingRule::UserSized)
+		.ClientSize(FVector2D(640.f, 560.f))
+		.AutoCenter(EAutoCenter::PreferredWorkArea)
+		.SupportsMaximize(false)
+		.SupportsMinimize(false)
+		.HasCloseButton(true)
+		.MinWidth(480.f)
+		.MinHeight(400.f);
+
+	TSharedRef<SDataflowTemplatePicker> Picker =
+		SNew(SDataflowTemplatePicker)
+		.Templates(Templates)
+		.ParentWindow(Window);
+
+	Window->SetContent(Picker);
+
+	FSlateApplication::Get().AddModalWindow(Window, ParentWindow);
+
+	return Picker->GetResult();
+}
+
+void SDataflowTemplatePicker::Construct(const FArguments& InArgs)
+{
+	Templates   = InArgs._Templates;
+	OwnerWindow = InArgs._ParentWindow;
+
+	ChildSlot
+	[
+		SNew(SBox)
+		[
+			SNew(SVerticalBox)
+
+			// Header
+			+ SVerticalBox::Slot()
+			.AutoHeight()
+			.Padding(16.f, 14.f, 16.f, 10.f)
+			[
+				SNew(STextBlock)
+				.Text(LOCTEXT("HeaderLabel", "Choose Dataflow Setup"))
+				.TextStyle(FAppStyle::Get(), "LargeText")
+			]
+
+			+ SVerticalBox::Slot()
+			.AutoHeight()
+			[
+				SNew(SSeparator)
+			]
+
+			// Template tiles — fills all space not claimed by the fixed sections below.
+			+ SVerticalBox::Slot()
+			.FillHeight(1.f)
+			.Padding(16.f, 12.f, 16.f, 0.f)
+			[
+				SNew(SScrollBox)
+				.Orientation(Orient_Vertical)
+				+ SScrollBox::Slot()
+				[
+					BuildTemplateGrid()
+				]
+			]
+
+			// "or pick existing" divider
+			+ SVerticalBox::Slot()
+			.AutoHeight()
+			.Padding(16.f, 14.f)
+			[
+				SNew(STextBlock)
+				.Text(LOCTEXT("OrExisting", "or pick an existing asset"))
+				.TextStyle(FAppStyle::Get(), "LargeText")
+			]
+
+			// Embedded asset picker
+			+ SVerticalBox::Slot()
+			.AutoHeight()
+			.Padding(16.f, 0.f, 16.f, 4.f)
+			[
+				BuildAssetPickerSection()
+			]
+
+			+ SVerticalBox::Slot()
+			.AutoHeight()
+			.Padding(FMargin(0.f, 10.f, 0.f, 0.f))
+			[
+				SNew(SSeparator)
+			]
+
+			// Bottom bar
+			+ SVerticalBox::Slot()
+			.AutoHeight()
+			[
+				BuildBottomBar()
+			]
+		]
+	];
+}
+
+///////////////////////////////////////////////////////////////////////////////////////////////////
+
+TSharedRef<SWidget> SDataflowTemplatePicker::BuildTemplateGrid()
+{
+	// SWrapBox automatically wraps tiles to a new row when the dialog is resized.
+	TSharedRef<SWrapBox> Grid = SNew(SWrapBox)
+		.UseAllottedSize(true)
+		.InnerSlotPadding(FVector2D(8.f, 8.f));
+
+	for (int32 i = 0; i < Templates.Num(); ++i)
+	{
+		Grid->AddSlot()
+		[
+			MakeTemplateTile(Templates[i], i)
+		];
+	}
+
+	return Grid;
+}
+
+TSharedRef<SWidget> SDataflowTemplatePicker::MakeTemplateTile(
+	const FDataflowTemplateOption& Option, int32 Index)
+{
+	using namespace DataflowTemplatePickerPrivate;
+
+	const FSlateBrush* IconBrush = Option.Icon
+		? Option.Icon
+		: FAppStyle::GetBrush("ClassIcon.DataTable");   // generic fallback
+
+	// SCheckBox with ToggleButtonCheckbox style gives us free selected-state visuals
+	// and radio-button behaviour when IsChecked returns Checked for only one index.
+	return SNew(SBox)
+		.WidthOverride(TileSize)
+		.HeightOverride(TileSize)
+		.ToolTipText(Option.Tooltip)
+		[
+			SNew(SCheckBox)
+			.Style(FAppStyle::Get(), "ToggleButtonCheckbox")
+			.IsChecked_Lambda([this, Index]()
+			{
+				return GetTileCheckState(Index);
+			})
+			.OnCheckStateChanged(this, &SDataflowTemplatePicker::OnTemplateTileChecked, Index)
+			.HAlign(HAlign_Center)
+			.Padding(6.f)
+			[
+				SNew(SVerticalBox)
+
+				+ SVerticalBox::Slot()
+				.FillHeight(1.f)
+				.HAlign(HAlign_Center)
+				.VAlign(VAlign_Center)
+				[
+					SNew(SImage)
+					.Image(IconBrush)
+					.DesiredSizeOverride(FVector2D(TileIconSize, TileIconSize))
+				]
+
+				+ SVerticalBox::Slot()
+				.AutoHeight()
+				.HAlign(HAlign_Center)
+				.Padding(0.f, 4.f, 0.f, 0.f)
+				[
+					SNew(STextBlock)
+					.Text(Option.DisplayName)
+					.Justification(ETextJustify::Center)
+					.WrapTextAt(TileSize - 12.f)
+				]
+			]
+		];
+}
+
+TSharedRef<SWidget> SDataflowTemplatePicker::BuildAssetPickerSection()
+{
+	return SNew(SHorizontalBox)
+
+		// Combo button — shows current selection, opens the inline picker as a dropdown.
+		+ SHorizontalBox::Slot()
+		.FillWidth(1.f)
+		.VAlign(VAlign_Center)
+		[
+			SAssignNew(AssetComboButton, SComboButton)
+			.ButtonStyle(FAppStyle::Get(), "PropertyEditor.AssetComboStyle")
+			.ForegroundColor(FAppStyle::GetColor("PropertyEditor.AssetName.ColorAndOpacity"))
+			.ContentPadding(FMargin(4.f, 2.f))
+			.CollapseMenuOnParentFocus(true)
+			.OnGetMenuContent(this, &SDataflowTemplatePicker::MakeAssetPickerMenu)
+			.ButtonContent()
+			[
+				SNew(STextBlock)
+				.Text(this, &SDataflowTemplatePicker::GetSelectedAssetLabel)
+				.Font(FAppStyle::GetFontStyle("PropertyWindow.NormalFont"))
+			]
+		]
+
+		// Clear button — visible only when an asset is selected.
+		+ SHorizontalBox::Slot()
+		.AutoWidth()
+		.VAlign(VAlign_Center)
+		.Padding(2.f, 0.f, 0.f, 0.f)
+		[
+			SNew(SButton)
+			.ButtonStyle(FAppStyle::Get(), "NoBorder")
+			.Visibility_Lambda([this]() { return HasAssetSelected() ? EVisibility::Visible : EVisibility::Collapsed; })
+			.OnClicked_Lambda([this]() { OnClearAssetSelection(); return FReply::Handled(); })
+			.ToolTipText(LOCTEXT("ClearAssetTip", "Clear selected asset"))
+			[
+				SNew(SImage)
+				.Image(FAppStyle::GetBrush("Icons.X"))
+				.DesiredSizeOverride(FVector2D(12.f, 12.f))
+			]
+		];
+}
+
+TSharedRef<SWidget> SDataflowTemplatePicker::MakeAssetPickerMenu()
+{
+	FAssetPickerConfig Config;
+	Config.Filter.ClassPaths.Add(UDataflow::StaticClass()->GetClassPathName());
+	Config.bAllowNullSelection  = false;
+	Config.bShowBottomToolbar   = false;
+	Config.bAutohideSearchBar   = false;
+	Config.InitialAssetViewType = EAssetViewType::List;
+	Config.OnAssetSelected      = FOnAssetSelected::CreateSP(
+	                                  this, &SDataflowTemplatePicker::OnAssetSelected);
+	Config.OnAssetDoubleClicked = FOnAssetDoubleClicked::CreateSP(
+	                                  this, &SDataflowTemplatePicker::OnAssetDoubleClicked);
+
+	IContentBrowserSingleton& CB =
+		FModuleManager::LoadModuleChecked<FContentBrowserModule>("ContentBrowser").Get();
+
+	return SNew(SBox)
+		.WidthOverride(350.f)
+		.HeightOverride(300.f)
+		[
+			CB.CreateAssetPicker(Config)
+		];
+}
+
+TSharedRef<SWidget> SDataflowTemplatePicker::BuildBottomBar()
+{
+	return SNew(SHorizontalBox)
+		.Clipping(EWidgetClipping::ClipToBounds)
+
+		// Spacer
+		+ SHorizontalBox::Slot()
+		.FillWidth(1.f)
+
+		// Right: No Dataflow + OK buttons
+		+ SHorizontalBox::Slot()
+		.HAlign(HAlign_Right)
+		.VAlign(VAlign_Center)
+		.AutoWidth()
+		.Padding(8.f, 10.f, 16.f, 10.f)
+		[
+			SNew(SHorizontalBox)
+
+			+ SHorizontalBox::Slot()
+			.AutoWidth()
+			.Padding(0.f, 0.f, 8.f, 0.f)
+			[
+				SNew(SButton)
+				.Text(LOCTEXT("NoneBtn", "No Dataflow"))
+				.ToolTipText(LOCTEXT("NoneTip", "Continue without attaching a Dataflow graph"))
+				.OnClicked(this, &SDataflowTemplatePicker::OnNoDataflowClicked)
+			]
+
+			+ SHorizontalBox::Slot()
+			.AutoWidth()
+			[
+				SNew(SButton)
+				.Text(LOCTEXT("OKBtn", "OK"))
+				.ButtonStyle(FAppStyle::Get(), "PrimaryButton")
+				.IsEnabled(this, &SDataflowTemplatePicker::IsOKEnabled)
+				.OnClicked(this, &SDataflowTemplatePicker::OnOKClicked)
+			]
+		];
+}
+
+void SDataflowTemplatePicker::OnTemplateTileChecked(ECheckBoxState NewState, int32 Index)
+{
+	if (NewState == ECheckBoxState::Checked)
+	{
+		CandidateResult.Type          = EDataflowPickerResult::Template;
+		CandidateResult.TemplateIndex = Index;
+		CandidateResult.TemplateId    = Templates[Index].TemplateId;
+		CandidateResult.SelectedAsset = FAssetData();
+	}
+	else
+	{
+		// Unchecking the active tile clears the template selection.
+		if (CandidateResult.TemplateIndex == Index)
+		{
+			CandidateResult = FDataflowPickerResult{};
+		}
+	}
+}
+
+void SDataflowTemplatePicker::OnAssetSelected(const FAssetData& AssetData)
+{
+	CandidateResult.Type          = EDataflowPickerResult::ExistingAsset;
+	CandidateResult.SelectedAsset = AssetData;
+	CandidateResult.TemplateIndex = INDEX_NONE;
+	CandidateResult.TemplateId    = NAME_None;
+
+	if (TSharedPtr<SComboButton> Combo = AssetComboButton.Pin())
+	{
+		Combo->SetIsOpen(false);
+	}
+}
+
+void SDataflowTemplatePicker::OnAssetDoubleClicked(const FAssetData& AssetData)
+{
+	OnAssetSelected(AssetData);
+	OnOKClicked();
+}
+
+void SDataflowTemplatePicker::OnClearAssetSelection()
+{
+	if (CandidateResult.Type == EDataflowPickerResult::ExistingAsset)
+	{
+		CandidateResult = FDataflowPickerResult{};
+	}
+}
+
+FText SDataflowTemplatePicker::GetSelectedAssetLabel() const
+{
+	if (CandidateResult.Type == EDataflowPickerResult::ExistingAsset && CandidateResult.SelectedAsset.IsValid())
+	{
+		return FText::FromName(CandidateResult.SelectedAsset.AssetName);
+	}
+	return LOCTEXT("NoAssetSelected", "None");
+}
+
+bool SDataflowTemplatePicker::HasAssetSelected() const
+{
+	return CandidateResult.Type == EDataflowPickerResult::ExistingAsset && CandidateResult.SelectedAsset.IsValid();
+}
+
+FReply SDataflowTemplatePicker::OnNoDataflowClicked()
+{
+	// Ignore CandidateResult
+	Result = { EDataflowPickerResult::None, INDEX_NONE, NAME_None, FAssetData() };
+	CloseOwnerWindow();
+	return FReply::Handled();
+}
+
+FReply SDataflowTemplatePicker::OnOKClicked()
+{
+	// Result already populated by tile / asset selection callbacks.
+	Result = CandidateResult;
+	CloseOwnerWindow();
+	return FReply::Handled();
+}
+
+FReply SDataflowTemplatePicker::OnCancelClicked()
+{
+	Result = FDataflowPickerResult{};   // default: Cancelled
+	CloseOwnerWindow();
+	return FReply::Handled();
+}
+
+void SDataflowTemplatePicker::CloseOwnerWindow()
+{
+	if (TSharedPtr<SWindow> Window = OwnerWindow.Pin())
+	{
+		Window->RequestDestroyWindow();
+	}
+}
+
+ECheckBoxState SDataflowTemplatePicker::GetTileCheckState(int32 Index) const
+{
+	return (CandidateResult.Type == EDataflowPickerResult::Template && CandidateResult.TemplateIndex == Index)
+		? ECheckBoxState::Checked
+		: ECheckBoxState::Unchecked;
+}
+
+bool SDataflowTemplatePicker::IsOKEnabled() const
+{
+	switch (CandidateResult.Type)
+	{
+	case EDataflowPickerResult::Template:
+		return Templates.IsValidIndex(CandidateResult.TemplateIndex);
+	case EDataflowPickerResult::ExistingAsset:
+		return CandidateResult.SelectedAsset.IsValid();
+	default:
+		return false;
+	}
+}
+
+#undef LOCTEXT_NAMESPACE
diff --git a/Engine/Plugins/Experimental/Dataflow/Source/DataflowEditor/Public/Dataflow/AssetDefinition_DataflowAsset.h b/Engine/Plugins/Experimental/Dataflow/Source/DataflowEditor/Public/Dataflow/AssetDefinition_DataflowAsset.h
--- a/Engine/Plugins/Experimental/Dataflow/Source/DataflowEditor/Public/Dataflow/AssetDefinition_DataflowAsset.h
+++ b/Engine/Plugins/Experimental/Dataflow/Source/DataflowEditor/Public/Dataflow/AssetDefinition_DataflowAsset.h
@@ -35,6 +35,9 @@ namespace UE::DataflowAssetDefinitionHelpers
 
 	// Create a new UDataflow if one doesn't already exist for the Cloth Asset
 	DATAFLOWEDITOR_API UObject* NewOrOpenDataflowAsset(const UObject* Asset);
+	
+	// Set dataflow from a template picked by the user
+	DATAFLOWEDITOR_API bool SetDataflowFromTemplatePicker(UObject* Asset);
 
 	// Determine whether or not a UDataflow can be opened directly in the Dataflow Editor
 	DATAFLOWEDITOR_API bool CanOpenDataflowAssetInEditor(const UObject* DataflowAsset);
diff --git a/Engine/Plugins/Experimental/Dataflow/Source/DataflowEditor/Public/Dataflow/DataflowTemplateRegistry.h b/Engine/Plugins/Experimental/Dataflow/Source/DataflowEditor/Public/Dataflow/DataflowTemplateRegistry.h
new file mode 100644
--- /dev/null
+++ b/Engine/Plugins/Experimental/Dataflow/Source/DataflowEditor/Public/Dataflow/DataflowTemplateRegistry.h
@@ -0,0 +1,75 @@
+// Copyright Epic Games, Inc. All Rights Reserved.
+
+#pragma once
+
+#include "CoreMinimal.h"
+#include "AssetRegistry/AssetData.h"
+#include "Textures/SlateIcon.h"
+
+/** One template entry */
+struct FDataflowTemplateOption
+{
+	FText              DisplayName;
+	FText              Tooltip;
+	/** Icon shown in the tile. Pass nullptr to fall back to a generic dataflow icon. */
+	const FSlateBrush* Icon = nullptr;
+	/** Opaque identifier returned to the caller in the result. */
+	FName              TemplateId;
+};
+
+/**
+ * Describes a content folder containing UDataflow template assets for a specific asset class.
+ * Passed to FDataflowTemplateRegistry::RegisterTemplateFolder at module startup.
+ */
+struct FDataflowTemplateFolderRegistration
+{
+	/** Content path to scan recursively, e.g. "/Game/Dataflow/Templates/Cloth". */
+	FString ContentFolder;
+
+	/**
+	 * Icon for every template found in this folder. Using FSlateIcon allows referencing brushes
+	 * from any style set, not just FAppStyle — required for class icons that live in plugin style
+	 * sets (e.g. FSlateIcon("GroomEditorStyle", "ClassIcon.GroomAsset")).
+	 * A default Dataflow icon is used when the icon is not set.
+	 */
+	FSlateIcon Icon;
+};
+
+/**
+ * Registry mapping asset classes to content folders containing UDataflow template assets.
+ *
+ * Modules register their template folders in StartupModule. When the Dataflow picker dialog
+ * opens for an asset, it queries this registry to populate the template tile list.
+ * Registrations for parent classes are automatically included (class hierarchy walk,
+ * most-derived first).
+ *
+ * Example (in a plugin's editor module StartupModule):
+ *
+ *   FDataflowTemplateRegistry::Get().RegisterTemplateFolder(
+ *       UChaosClothAsset::StaticClass(),
+ *       { TEXT("/ChaosClothEditor/Dataflow/Templates"), FSlateIcon("ChaosClothEditorStyle", "ClassIcon.ChaosClothAsset") });
+ */
+class DATAFLOWEDITOR_API FDataflowTemplateRegistry
+{
+public:
+	static FDataflowTemplateRegistry& Get();
+
+	/**
+	 * Register a content folder to scan for UDataflow template assets for a given asset class.
+	 * Registrations are also visible when the picker is opened for any subclass of AssetClass.
+	 * Safe to call from any module's StartupModule (game thread only).
+	 */
+	void RegisterTemplateFolder(const UClass* AssetClass, FDataflowTemplateFolderRegistration Registration);
+
+	/**
+	 * Build template options for the given asset class by scanning all registered folders
+	 * (including those registered for parent classes, most-derived first).
+	 * A "Blank" option is always prepended.
+	 * Queries the asset registry synchronously — call from the game thread only.
+	 */
+	TArray<FDataflowTemplateOption> GetTemplateOptions(const UClass* AssetClass) const;
+
+private:
+	// Keyed by class path so we don't hold raw UClass* pointers across module reloads.
+	TMap<FTopLevelAssetPath, TArray<FDataflowTemplateFolderRegistration>> Registrations;
+};
diff --git a/Engine/Plugins/Experimental/Dataflow/Source/DataflowEditor/Public/Dataflow/SDataflowTemplatePicker.h b/Engine/Plugins/Experimental/Dataflow/Source/DataflowEditor/Public/Dataflow/SDataflowTemplatePicker.h
new file mode 100644
--- /dev/null
+++ b/Engine/Plugins/Experimental/Dataflow/Source/DataflowEditor/Public/Dataflow/SDataflowTemplatePicker.h
@@ -0,0 +1,99 @@
+// Copyright Epic Games, Inc. All Rights Reserved.
+
+#pragma once
+
+#include "CoreMinimal.h"
+#include "AssetRegistry/AssetData.h"
+#include "Widgets/SCompoundWidget.h"
+#include "Widgets/Input/SComboButton.h"
+#include "Dataflow/DataflowTemplateRegistry.h"
+
+///////////////////////////////////////////////////////////////////////////////////////////////////
+
+/** Discriminator for FDataflowPickerResult::Type. */
+enum class EDataflowPickerResult : uint8
+{
+	/** User selected a template tile. TemplateIndex and TemplateId are valid. */
+	Template,
+	/** User picked an existing asset from the content browser. SelectedAsset is valid. */
+	ExistingAsset,
+	/** User clicked "No Dataflow". */
+	None,
+	/** User cancelled or closed the window. */
+	Cancelled,
+};
+
+///////////////////////////////////////////////////////////////////////////////////////////////////
+
+struct FDataflowPickerResult
+{
+	EDataflowPickerResult Type          = EDataflowPickerResult::Cancelled;
+	int32                 TemplateIndex = INDEX_NONE;
+	FName                 TemplateId;
+	FAssetData            SelectedAsset;
+};
+
+///////////////////////////////////////////////////////////////////////////////////////////////////
+
+/**
+ * Modal dialog that lets the user:
+ *   • choose one of several Dataflow templates (name + icon tiles),
+ *   • pick an existing Dataflow asset from the content browser, or
+ *   • opt out ("No Dataflow").
+ *
+ * Typical usage:
+ *   FDataflowPickerResult Result = SDataflowTemplatePicker::ShowModal(Templates, ParentWindow);
+ */
+class DATAFLOWEDITOR_API SDataflowTemplatePicker : public SCompoundWidget
+{
+public:
+	SLATE_BEGIN_ARGS(SDataflowTemplatePicker)
+		: _ParentWindow()
+	{}
+		SLATE_ARGUMENT(TArray<FDataflowTemplateOption>, Templates)
+		SLATE_ARGUMENT(TWeakPtr<SWindow>,               ParentWindow)
+	SLATE_END_ARGS()
+
+	void Construct(const FArguments& InArgs);
+
+	/**
+	 * Creates a modal window, blocks until dismissed, then returns the result.
+	 * Pass a valid ParentWindow to centre the dialog; nullptr uses the main window.
+	 */
+	static FDataflowPickerResult ShowModal(
+		const TArray<FDataflowTemplateOption>& Templates,
+		TSharedPtr<SWindow>                    ParentWindow = nullptr);
+
+	const FDataflowPickerResult& GetResult() const { return Result; }
+
+private:
+	// Build helpers
+	TSharedRef<SWidget> BuildTemplateGrid();
+	TSharedRef<SWidget> MakeTemplateTile(const FDataflowTemplateOption& Option, int32 Index);
+	TSharedRef<SWidget> BuildAssetPickerSection();
+	TSharedRef<SWidget> MakeAssetPickerMenu();
+	TSharedRef<SWidget> BuildBottomBar();
+
+	// Interaction
+	void    OnTemplateTileChecked(ECheckBoxState NewState, int32 Index);
+	void    OnAssetSelected(const FAssetData& AssetData);
+	void    OnAssetDoubleClicked(const FAssetData& AssetData);
+	void    OnClearAssetSelection();
+	FReply  OnNoDataflowClicked();
+	FReply  OnOKClicked();
+	FReply  OnCancelClicked();
+	void    CloseOwnerWindow();
+
+	// Attribute getters (used by Slate attribute bindings)
+	ECheckBoxState GetTileCheckState(int32 Index) const;
+	bool           IsOKEnabled() const;
+	FText          GetSelectedAssetLabel() const;
+	bool           HasAssetSelected() const;
+
+	// State
+	TArray<FDataflowTemplateOption> Templates;
+	TWeakPtr<SWindow>               OwnerWindow;
+	TWeakPtr<SComboButton>          AssetComboButton;
+	FDataflowPickerResult           CandidateResult;
+	FDataflowPickerResult           Result;
+};
