Dataflow : Added an asset right click menu entry to set or change the Dataflow graph associated with the selected asset

#Synced-CL 52234830

diff --git a/Engine/Plugins/Experimental/Dataflow/Source/DataflowEditor/Private/Dataflow/AssetDefinition_DataflowAsset.cpp b/Engine/Plugins/Experimental/Dataflow/Source/DataflowEditor/Private/Dataflow/AssetDefinition_DataflowAsset.cpp
--- a/Engine/Plugins/Experimental/Dataflow/Source/DataflowEditor/Private/Dataflow/AssetDefinition_DataflowAsset.cpp
+++ b/Engine/Plugins/Experimental/Dataflow/Source/DataflowEditor/Private/Dataflow/AssetDefinition_DataflowAsset.cpp
@@ -380,6 +380,58 @@ namespace UE::DataflowAssetDefinitionHelpers
 
 	//////////////////////////////////////////////////////////////////////////////////////////////////////////////////////////////////////
 
+	static void ExecuteSetDataflow(const FToolMenuContext& InContext)
+	{
+		if (const UContentBrowserAssetContextMenuContext* Context = UContentBrowserAssetContextMenuContext::FindContextWithAssets(InContext))
+		{
+			TArray<UObject*> SelectedAssets = Context->LoadSelectedObjects<UObject>();
+			if (!SelectedAssets.IsEmpty())
+			{
+				for (UObject* SelectedAsset : SelectedAssets)
+				{
+					if (SelectedAsset)
+					{
+						UE::DataflowAssetDefinitionHelpers::SetDataflowFromTemplatePicker(SelectedAsset);
+					}
+				}
+			}
+		}
+	}
+
+	static bool CanExecuteSetDataflow(const FToolMenuContext& InContext)
+	{
+		if (const UContentBrowserAssetContextMenuContext* Context = UContentBrowserAssetContextMenuContext::FindContextWithAssets(InContext))
+		{
+			TArray<UObject*> SelectedAssets = Context->LoadSelectedObjects<UObject>();
+			// only enable when one asset is selected because this will open a dialog box for each of them 
+			return (SelectedAssets.Num() == 1);
+		}
+		return false;
+	}
+
+	void RegisterSetDataflowAssetMenu(UClass* AssetClass)
+	{
+		UToolMenus::RegisterStartupCallback(
+			FSimpleMulticastDelegate::FDelegate::CreateLambda([AssetClass]()
+				{
+					FToolMenuOwnerScoped OwnerScoped(UE_MODULE_NAME);
+					UToolMenu* Menu = UE::ContentBrowser::ExtendToolMenu_AssetContextMenu(AssetClass);
+					FToolMenuSection& Section = Menu->FindOrAddSection("GetAssetActions");
+					{
+						const TAttribute<FText> Label = LOCTEXT("SetDataflow_Label", "Set Dataflow...");
+						const TAttribute<FText> ToolTip = LOCTEXT("SetDataflow_Tooltip", "Select and set a dataflow graph for this asset. This will open the dataflow template selection dialog.");
+						const TAttribute<FSlateIcon> Icon = FSlateIcon(FAppStyle::GetAppStyleSetName(), "GraphEditor.EventGraph_16x");
+
+						FToolUIAction UIAction;
+						UIAction.ExecuteAction = FToolMenuExecuteAction::CreateStatic(&ExecuteSetDataflow);
+						UIAction.CanExecuteAction = FToolMenuCanExecuteAction::CreateStatic(&CanExecuteSetDataflow);
+						Section.AddMenuEntry("Dataflow_SetDataflow", Label, ToolTip, Icon, UIAction);
+					}
+				}));
+	}
+
+	//////////////////////////////////////////////////////////////////////////////////////////////////////////////////////////////////////
+
 	static void ExecuteReEvaluateDataflow(const FToolMenuContext& InContext)
 	{
 		if (const UContentBrowserAssetContextMenuContext* Context = UContentBrowserAssetContextMenuContext::FindContextWithAssets(InContext))
diff --git a/Engine/Plugins/Experimental/Dataflow/Source/DataflowEditor/Public/Dataflow/AssetDefinition_DataflowAsset.h b/Engine/Plugins/Experimental/Dataflow/Source/DataflowEditor/Public/Dataflow/AssetDefinition_DataflowAsset.h
--- a/Engine/Plugins/Experimental/Dataflow/Source/DataflowEditor/Public/Dataflow/AssetDefinition_DataflowAsset.h
+++ b/Engine/Plugins/Experimental/Dataflow/Source/DataflowEditor/Public/Dataflow/AssetDefinition_DataflowAsset.h
@@ -45,6 +45,9 @@ namespace UE::DataflowAssetDefinitionHelpers
 	// Register a entry on the asset menu that will embed an external dataflow in the selected asset
 	DATAFLOWEDITOR_API void RegisterEmbedDataflowAssetMenu(UClass* AssetClass);
 
+	// Register a entry on the asset menu that will set a dataflow asset for the selected asset
+	DATAFLOWEDITOR_API void RegisterSetDataflowAssetMenu(UClass* AssetClass);
+
 	// Register a entry on the asset menu that will Re-Evaluate the assoiated dataflow for the selected asset
 	DATAFLOWEDITOR_API void RegisterReEvaluateDataflowAssetMenu(UClass* AssetClass);
 
@@ -66,6 +69,7 @@ namespace UE::DataflowAssetDefinitionHelpers
 	{ \
 		static FDelayedAutoRegisterHelper DelayedAutoRegister(EDelayedRegisterRunPhase::EndOfEngineInit, []() \
 		{ \
+			UE::DataflowAssetDefinitionHelpers::RegisterSetDataflowAssetMenu(AssetClassType::StaticClass()); \
 			UE::DataflowAssetDefinitionHelpers::RegisterEmbedDataflowAssetMenu(AssetClassType::StaticClass()); \
 			UE::DataflowAssetDefinitionHelpers::RegisterReEvaluateDataflowAssetMenu(AssetClassType::StaticClass()); \
 		}); \
