Dataflow template registry : added the ability to register single assets

#Synced-CL 52234830

diff --git a/Engine/Plugins/Experimental/Dataflow/Source/DataflowEditor/Private/Dataflow/DataflowTemplateRegistry.cpp b/Engine/Plugins/Experimental/Dataflow/Source/DataflowEditor/Private/Dataflow/DataflowTemplateRegistry.cpp
--- a/Engine/Plugins/Experimental/Dataflow/Source/DataflowEditor/Private/Dataflow/DataflowTemplateRegistry.cpp
+++ b/Engine/Plugins/Experimental/Dataflow/Source/DataflowEditor/Private/Dataflow/DataflowTemplateRegistry.cpp
@@ -15,6 +15,17 @@ FDataflowTemplateRegistry& FDataflowTemplateRegistry::Get()
 	return Instance;
 }
 
+void FDataflowTemplateRegistry::RegisterTemplateAsset(
+	const UClass* AssetClass, FDataflowTemplateAssetRegistration Registration)
+{
+	if (!ensureMsgf(AssetClass, TEXT("FDataflowTemplateRegistry::RegisterTemplateAsset: null AssetClass")) ||
+	    !ensureMsgf(Registration.AssetPath.IsValid(), TEXT("FDataflowTemplateRegistry::RegisterTemplateAsset: invalid AssetPath")))
+	{
+		return;
+	}
+	AssetRegistrations.FindOrAdd(AssetClass->GetClassPathName()).Add(MoveTemp(Registration));
+}
+
 void FDataflowTemplateRegistry::RegisterTemplateFolder(
 	const UClass* AssetClass, FDataflowTemplateFolderRegistration Registration)
 {
@@ -41,6 +52,55 @@ TArray<FDataflowTemplateOption> FDataflowTemplateRegistry::GetTemplateOptions(
 		Blank.TemplateId  = NAME_None;   // NAME_None > caller creates a blank UDataflow
 	}
 
+	// Collect individually registered assets (most-derived first), deduplicating by object path.
+	TSet<FName> AddedTemplateIds;  // tracks ids already emitted; also used below for folder dedup
+	for (const UClass* Class = AssetClass; Class; Class = Class->GetSuperClass())
+	{
+		const TArray<FDataflowTemplateAssetRegistration>* AssetRegs =
+			AssetRegistrations.Find(Class->GetClassPathName());
+		if (!AssetRegs)
+		{
+			continue;
+		}
+		for (const FDataflowTemplateAssetRegistration& Reg : *AssetRegs)
+		{
+			const FName TemplateId = FName(Reg.AssetPath.ToString());
+			if (AddedTemplateIds.Contains(TemplateId))
+			{
+				continue;
+			}
+			AddedTemplateIds.Add(TemplateId);
+
+			// DisplayName: explicit or derived from asset name.
+			FText DisplayName = Reg.DisplayName;
+			if (DisplayName.IsEmpty())
+			{
+				FString Derived = Reg.AssetPath.GetAssetName();
+				Derived.RemoveFromStart(TEXT("DF_"));
+				Derived = Derived.Replace(TEXT("_"), TEXT(" "));
+				DisplayName = FText::FromString(Derived);
+			}
+
+			// Tooltip: explicit or asset object path.
+			FText Tooltip = Reg.Tooltip;
+			if (Tooltip.IsEmpty())
+			{
+				Tooltip = FText::FromString(Reg.AssetPath.ToString());
+			}
+
+			// Icon: explicit registration or fallback.
+			const FSlateBrush* Icon = (Reg.Icon.IsSet())
+				? Reg.Icon.GetIcon()
+				: DataflowDefaultIcon.GetIcon();
+
+			FDataflowTemplateOption& Option = Result.AddDefaulted_GetRef();
+			Option.DisplayName = DisplayName;
+			Option.Tooltip     = Tooltip;
+			Option.Icon        = Icon;
+			Option.TemplateId  = TemplateId;
+		}
+	}
+
 	// Walk the class hierarchy (most-derived first) to collect registered folders,
 	// preserving insertion order and deduplicating by path.
 	TArray<FString>           OrderedFolders;
@@ -105,6 +165,14 @@ TArray<FDataflowTemplateOption> FDataflowTemplateRegistry::GetTemplateOptions(
 
 	for (const FAssetData& AssetData : FoundAssets)
 	{
+		// Skip assets already emitted as individually registered templates.
+		const FName TemplateId = FName(AssetData.GetObjectPathString());
+		if (AddedTemplateIds.Contains(TemplateId))
+		{
+			continue;
+		}
+		AddedTemplateIds.Add(TemplateId);
+
 		// Resolve icon from the folder this asset lives under.
 		const FString AssetPackagePath = AssetData.PackagePath.ToString();
 		const FSlateIcon* FoundIcon = nullptr;
@@ -131,7 +199,7 @@ TArray<FDataflowTemplateOption> FDataflowTemplateRegistry::GetTemplateOptions(
 		Option.Tooltip     = FText::FromString(AssetData.GetObjectPathString());
 		Option.Icon        = Icon;
 		// Full object path stored in TemplateId so callers can LoadObject<UDataflow> directly.
-		Option.TemplateId  = FName(AssetData.GetObjectPathString());
+		Option.TemplateId  = TemplateId;
 	}
 
 	return Result;
diff --git a/Engine/Plugins/Experimental/Dataflow/Source/DataflowEditor/Public/Dataflow/DataflowTemplateRegistry.h b/Engine/Plugins/Experimental/Dataflow/Source/DataflowEditor/Public/Dataflow/DataflowTemplateRegistry.h
--- a/Engine/Plugins/Experimental/Dataflow/Source/DataflowEditor/Public/Dataflow/DataflowTemplateRegistry.h
+++ b/Engine/Plugins/Experimental/Dataflow/Source/DataflowEditor/Public/Dataflow/DataflowTemplateRegistry.h
@@ -17,6 +17,34 @@ struct FDataflowTemplateOption
 	FName              TemplateId;
 };
 
+/**
+ * Describes a single UDataflow asset to expose as a named template for a specific asset class.
+ * Passed to FDataflowTemplateRegistry::RegisterTemplateAsset at module startup.
+ *
+ * Example:
+ *   FDataflowTemplateRegistry::Get().RegisterTemplateAsset(
+ *       UGroomAsset::StaticClass(),
+ *       { FSoftObjectPath(TEXT("/GroomEditor/Dataflow/Templates/DF_BasicHair")),
+ *         LOCTEXT("BasicHair", "Basic Hair"), {}, FSlateIcon("GroomEditorStyle", "ClassThumbnail.GroomAsset") });
+ */
+struct FDataflowTemplateAssetRegistration
+{
+	/** Soft path to the specific UDataflow asset. */
+	FSoftObjectPath AssetPath;
+
+	/** Tile label. If empty, derived from the asset name (strips "DF_", replaces underscores with spaces). */
+	FText DisplayName;
+
+	/** Tooltip shown on hover. If empty, defaults to the asset object path. */
+	FText Tooltip;
+
+	/**
+	 * Tile icon. Falls back to the generic Dataflow icon if not set.
+	 * e.g. FSlateIcon("GroomEditorStyle", "ClassThumbnail.GroomAsset").
+	 */
+	FSlateIcon Icon;
+};
+
 /**
  * Describes a content folder containing UDataflow template assets for a specific asset class.
  * Passed to FDataflowTemplateRegistry::RegisterTemplateFolder at module startup.
@@ -54,6 +82,13 @@ class DATAFLOWEDITOR_API FDataflowTemplateRegistry
 public:
 	static FDataflowTemplateRegistry& Get();
 
+	/**
+	 * Register a single UDataflow asset as a template for the given asset class.
+	 * Registrations are visible for subclasses of AssetClass.
+	 * Safe to call from any module's StartupModule (game thread only).
+	 */
+	void RegisterTemplateAsset(const UClass* AssetClass, FDataflowTemplateAssetRegistration Registration);
+
 	/**
 	 * Register a content folder to scan for UDataflow template assets for a given asset class.
 	 * Registrations are also visible when the picker is opened for any subclass of AssetClass.
@@ -71,5 +106,6 @@ public:
 
 private:
 	// Keyed by class path so we don't hold raw UClass* pointers across module reloads.
+	TMap<FTopLevelAssetPath, TArray<FDataflowTemplateAssetRegistration>>  AssetRegistrations;
 	TMap<FTopLevelAssetPath, TArray<FDataflowTemplateFolderRegistration>> Registrations;
 };
