Dataflow editor : revise the way we compute if the graph is dirty or not
- support multi terminal node by only looking at the specific terminal targeted by the settings
- moved the timestamp setting to be only in the post evaluation callback to make sure it is updated at the right time
- fix the issue where evaluating a single node would report wrong information about the graph dirty state

#Synced-CL 52324855

diff --git a/Engine/Plugins/Experimental/Dataflow/Source/DataflowEditor/Private/Dataflow/DataflowEditorToolkit.cpp b/Engine/Plugins/Experimental/Dataflow/Source/DataflowEditor/Private/Dataflow/DataflowEditorToolkit.cpp
--- a/Engine/Plugins/Experimental/Dataflow/Source/DataflowEditor/Private/Dataflow/DataflowEditorToolkit.cpp
+++ b/Engine/Plugins/Experimental/Dataflow/Source/DataflowEditor/Private/Dataflow/DataflowEditorToolkit.cpp
@@ -1788,11 +1788,6 @@ void FDataflowEditorToolkit::OnPropertyValueChanged(const FPropertyChangedEvent&
 				FDataflowEditorCommands::OnCommentPropertyValueChanged(PropertyChangedEvent, SelectedComments);
 			}
 
-			// For manual evaluation no need to invalidate the construction scene if we change any parameters
-			// since it will trigger the rendering callbacks with the same cached collection
-			const bool bMakesDirty = (EvaluationMode != EDataflowEditorEvaluationMode::Manual);
-			EditorContent->SetLastModifiedTimestamp(LastNodeTimestamp, bMakesDirty);
-
 			// Refresh graph display to display nodes with warning/error
 			GraphEditor->NotifyGraphChanged();
 		}
@@ -2014,29 +2009,30 @@ bool FDataflowEditorToolkit::IsGraphDirty() const
 {
 	using namespace UE::Dataflow;
 
+	// we need to check the specific terminal for this asset and see if it's up to date 
 	if (const TObjectPtr<UDataflowBaseContent>& EditorContent = GetEditorContent())
 	{
-		if (TSharedPtr<UE::Dataflow::FEngineContext> DataflowContext = EditorContent->GetDataflowContext())
+		const FTimestamp LastModifiedTimestamp = EditorContent->GetLastModifiedTimestamp();
+		const FName TerminalNodeName(EditorContent->GetDataflowTerminal());
+
+		if (UDataflow* DataflowAsset = EditorContent->GetDataflowAsset())
 		{
-			if (UDataflow* DataflowAsset = EditorContent->GetDataflowAsset())
+			if (TSharedPtr<FGraph> DataflowGraph = DataflowAsset->GetDataflow())
 			{
-				if (TSharedPtr<UE::Dataflow::FGraph> DataflowGraph = DataflowAsset->GetDataflow())
+				if (FGraph::FConstNodeSharedPtr Node = DataflowGraph->FindBaseNode(TerminalNodeName))
 				{
-					const FTimestamp LastModifiedTimestamp = EditorContent->GetLastModifiedTimestamp();
-
-					for (TSharedPtr<FDataflowNode> Node : DataflowGraph->GetNodes())
+					if (const FDataflowTerminalNode* TerminalNode = Node->AsType<FDataflowTerminalNode>())
 					{
-						if (Node)
+						if (LastModifiedTimestamp < TerminalNode->GetTimestamp())
 						{
-							if (FDataflowTerminalNode* TerminalNode = Node->AsType<FDataflowTerminalNode>())
-							{
-								if (LastModifiedTimestamp < TerminalNode->GetTimestamp())
-								{
-									return true;
-								}
-							}
+							return true;
 						}
 					}
+					else
+					{
+						// no main terminal means always dirty
+						return true;
+					}
 				}
 			}
 		}
@@ -2441,10 +2437,7 @@ void FDataflowEditorToolkit::EvaluateTerminalNode(const FDataflowTerminalNode& T
 	if (const TObjectPtr<UDataflowBaseContent>& EditorContent = GetEditorContent())
 	{
 		FTimestamp TerminalNodeTimeStamp = FTimestamp::Invalid;
-
 		EvaluateNode(&TerminalNode, nullptr, TerminalNodeTimeStamp);
-
-		EditorContent->SetLastModifiedTimestamp(TerminalNodeTimeStamp, /*bMakeDirty*/true);
 	}
 }
 
@@ -2465,10 +2458,13 @@ void FDataflowEditorToolkit::EvaluateNode(const FDataflowNode* Node, const FData
 					const FDateTime StartEvaluationTime = FDateTime::Now();
 
 					// when evaluation completes, refresh the views
+					TWeakObjectPtr<UDataflowBaseContent> WeakEditorContent = EditorContent;
 					TWeakPtr<FDataflowEditorToolkit> WeakToolkit = SharedThis(this);
 					TWeakPtr<const FDataflowNode> WeakNode = Node->AsWeak();
-					auto OnEvaluationCompleted = [WeakToolkit, WeakNode, StartEvaluationTime](UE::Dataflow::FContext& Context)
+					auto OnEvaluationCompleted = [WeakEditorContent, WeakToolkit, WeakNode, StartEvaluationTime](UE::Dataflow::FContext& Context)
 						{
+							check(IsInGameThread());
+
 							const FDateTime EndEvaluationTime = FDateTime::Now();
 							// refresh views 
 							if (TSharedPtr<FDataflowEditorToolkit> ToolKit = WeakToolkit.Pin())
@@ -2481,7 +2477,15 @@ void FDataflowEditorToolkit::EvaluateNode(const FDataflowNode* Node, const FData
 								const FText ElapsedTimeText = FText::FromString(UE::Dataflow::Private::GetElapsedTimeString(EndEvaluationTime, StartEvaluationTime));
 								const FText NodeNameText = FText::FromName(Node->GetName());
 								const FText MessageText = FText::Format(LOCTEXT("EvalElapsedLog", "Graph evaluation from node [{0}] time : {1}"), NodeNameText, ElapsedTimeText);
-								Context.Info(MessageText);
+								Context.Info(MessageText, Node.Get());
+
+								if (TStrongObjectPtr<UDataflowBaseContent> EditorContent = WeakEditorContent.Pin())
+								{
+									if (Node->GetName().ToString() == EditorContent->GetDataflowTerminal())
+									{
+										EditorContent->SetLastModifiedTimestamp(Node->GetTimestamp());
+									}
+								}
 							}
 						};
 					bool bUseAsyncEvaluation = false;
@@ -2567,11 +2571,8 @@ TSharedRef<SDataflowGraphEditor> FDataflowEditorToolkit::CreateGraphEditorWidget
 			{
 				if (const TObjectPtr<UDataflowBaseContent>& EditorContent = GetEditorContent())
 				{
-					UE::Dataflow::FTimestamp LastNodeTimestamp = EditorContent->GetLastModifiedTimestamp();
-
-					EvaluateNode(Node, Output, LastNodeTimestamp);
-
-					EditorContent->SetLastModifiedTimestamp(LastNodeTimestamp);
+					FTimestamp NodeTimeStamp = FTimestamp::Invalid;
+					EvaluateNode(Node, Output, NodeTimeStamp);
 				}
 			}
 		};
