Dataflow:
- Added missing initialization

#Synced-CL 52406454

diff --git a/Engine/Plugins/Experimental/Dataflow/Source/DataflowNodes/Private/Dataflow/SamplerNodes/DataflowDistanceFromSphereSamplerNode.h b/Engine/Plugins/Experimental/Dataflow/Source/DataflowNodes/Private/Dataflow/SamplerNodes/DataflowDistanceFromSphereSamplerNode.h
--- a/Engine/Plugins/Experimental/Dataflow/Source/DataflowNodes/Private/Dataflow/SamplerNodes/DataflowDistanceFromSphereSamplerNode.h
+++ b/Engine/Plugins/Experimental/Dataflow/Source/DataflowNodes/Private/Dataflow/SamplerNodes/DataflowDistanceFromSphereSamplerNode.h
@@ -19,7 +19,7 @@ struct FDataflowDistanceFromSphereFloatSampler : public FDataflowFloatSamplerBas
 public:
 	/** Sphere to compute the distance from */
 	UPROPERTY(meta = (DataflowInput))
-	FSphere Sphere;
+	FSphere Sphere = FSphere(0.0);
 };
 
 /**
