dataflow template registry : add option to skip the default empty dataflow option

#Synced-CL 52234830

#Synced-CL 52234830

diff --git a/Engine/Plugins/Experimental/Dataflow/Source/DataflowEditor/Private/Dataflow/AssetDefinition_DataflowAsset.cpp b/Engine/Plugins/Experimental/Dataflow/Source/DataflowEditor/Private/Dataflow/AssetDefinition_DataflowAsset.cpp
--- a/Engine/Plugins/Experimental/Dataflow/Source/DataflowEditor/Private/Dataflow/AssetDefinition_DataflowAsset.cpp
+++ b/Engine/Plugins/Experimental/Dataflow/Source/DataflowEditor/Private/Dataflow/AssetDefinition_DataflowAsset.cpp
@@ -225,10 +225,10 @@ namespace UE::DataflowAssetDefinitionHelpers
 	 * Asset must implement IDataflowInstanceInterface for the Template/ExistingAsset paths to take effect.
 	 * Always returns true — the picker is a single non-repeating step with no sub-dialogs.
 	 */
-	bool SetDataflowFromTemplatePicker(UObject* Asset)
+	bool SetDataflowFromTemplatePicker(UObject* Asset, bool bShowDefaultBlankOption)
 	{
 		const TArray<FDataflowTemplateOption> Templates =
-			FDataflowTemplateRegistry::Get().GetTemplateOptions(Asset->GetClass());
+			FDataflowTemplateRegistry::Get().GetTemplateOptions(Asset->GetClass(), bShowDefaultBlankOption);
 
 		const FDataflowPickerResult PickerResult = SDataflowTemplatePicker::ShowModal(
 			Templates,
diff --git a/Engine/Plugins/Experimental/Dataflow/Source/DataflowEditor/Private/Dataflow/DataflowTemplateRegistry.cpp b/Engine/Plugins/Experimental/Dataflow/Source/DataflowEditor/Private/Dataflow/DataflowTemplateRegistry.cpp
--- a/Engine/Plugins/Experimental/Dataflow/Source/DataflowEditor/Private/Dataflow/DataflowTemplateRegistry.cpp
+++ b/Engine/Plugins/Experimental/Dataflow/Source/DataflowEditor/Private/Dataflow/DataflowTemplateRegistry.cpp
@@ -37,13 +37,14 @@ void FDataflowTemplateRegistry::RegisterTemplateFolder(
 }
 
 TArray<FDataflowTemplateOption> FDataflowTemplateRegistry::GetTemplateOptions(
-	const UClass* AssetClass) const
+	const UClass* AssetClass, bool bAddDefaultBlankOption) const
 {
 	TArray<FDataflowTemplateOption> Result;
 
 	const FSlateIcon DataflowDefaultIcon(FName("DataflowEditorStyle"), "ClassIcon.Dataflow");
 
 	// "Blank" is always the first option.
+	if (bAddDefaultBlankOption)
 	{
 		FDataflowTemplateOption& Blank = Result.AddDefaulted_GetRef();
 		Blank.DisplayName = LOCTEXT("Blank_Name",    "Blank");
diff --git a/Engine/Plugins/Experimental/Dataflow/Source/DataflowEditor/Private/Dataflow/SDataflowTemplatePicker.cpp b/Engine/Plugins/Experimental/Dataflow/Source/DataflowEditor/Private/Dataflow/SDataflowTemplatePicker.cpp
--- a/Engine/Plugins/Experimental/Dataflow/Source/DataflowEditor/Private/Dataflow/SDataflowTemplatePicker.cpp
+++ b/Engine/Plugins/Experimental/Dataflow/Source/DataflowEditor/Private/Dataflow/SDataflowTemplatePicker.cpp
@@ -69,7 +69,7 @@ void SDataflowTemplatePicker::Construct(const FArguments& InArgs)
 			.Padding(16.f, 14.f, 16.f, 10.f)
 			[
 				SNew(STextBlock)
-				.Text(LOCTEXT("HeaderLabel", "Choose Dataflow Setup"))
+				.Text(LOCTEXT("HeaderLabel", "Choose a Dataflow Graph template"))
 				.TextStyle(FAppStyle::Get(), "LargeText")
 			]
 
@@ -98,7 +98,7 @@ void SDataflowTemplatePicker::Construct(const FArguments& InArgs)
 			.Padding(16.f, 14.f)
 			[
 				SNew(STextBlock)
-				.Text(LOCTEXT("OrExisting", "or pick an existing asset"))
+				.Text(LOCTEXT("OrExisting", "Or pick an existing Dataflow Graph"))
 				.TextStyle(FAppStyle::Get(), "LargeText")
 			]
 
diff --git a/Engine/Plugins/Experimental/Dataflow/Source/DataflowEditor/Public/Dataflow/AssetDefinition_DataflowAsset.h b/Engine/Plugins/Experimental/Dataflow/Source/DataflowEditor/Public/Dataflow/AssetDefinition_DataflowAsset.h
--- a/Engine/Plugins/Experimental/Dataflow/Source/DataflowEditor/Public/Dataflow/AssetDefinition_DataflowAsset.h
+++ b/Engine/Plugins/Experimental/Dataflow/Source/DataflowEditor/Public/Dataflow/AssetDefinition_DataflowAsset.h
@@ -37,7 +37,7 @@ namespace UE::DataflowAssetDefinitionHelpers
 	DATAFLOWEDITOR_API UObject* NewOrOpenDataflowAsset(const UObject* Asset);
 	
 	// Set dataflow from a template picked by the user
-	DATAFLOWEDITOR_API bool SetDataflowFromTemplatePicker(UObject* Asset);
+	DATAFLOWEDITOR_API bool SetDataflowFromTemplatePicker(UObject* Asset, bool bShowDefaultBlankOption = true);
 
 	// Determine whether or not a UDataflow can be opened directly in the Dataflow Editor
 	DATAFLOWEDITOR_API bool CanOpenDataflowAssetInEditor(const UObject* DataflowAsset);
diff --git a/Engine/Plugins/Experimental/Dataflow/Source/DataflowEditor/Public/Dataflow/DataflowTemplateRegistry.h b/Engine/Plugins/Experimental/Dataflow/Source/DataflowEditor/Public/Dataflow/DataflowTemplateRegistry.h
--- a/Engine/Plugins/Experimental/Dataflow/Source/DataflowEditor/Public/Dataflow/DataflowTemplateRegistry.h
+++ b/Engine/Plugins/Experimental/Dataflow/Source/DataflowEditor/Public/Dataflow/DataflowTemplateRegistry.h
@@ -102,7 +102,7 @@ public:
 	 * A "Blank" option is always prepended.
 	 * Queries the asset registry synchronously — call from the game thread only.
 	 */
-	TArray<FDataflowTemplateOption> GetTemplateOptions(const UClass* AssetClass) const;
+	TArray<FDataflowTemplateOption> GetTemplateOptions(const UClass* AssetClass, bool bAddDefaultBlankOption = true) const;
 
 private:
 	// Keyed by class path so we don't hold raw UClass* pointers across module reloads.
