Dataflow : Improved render settings detail view
- Added search bar and favorite option
- Added section

#Synced-CL 52463862

diff --git a/Engine/Plugins/Experimental/Dataflow/Source/DataflowEditor/Private/Dataflow/DataflowEditorModeToolkit.cpp b/Engine/Plugins/Experimental/Dataflow/Source/DataflowEditor/Private/Dataflow/DataflowEditorModeToolkit.cpp
--- a/Engine/Plugins/Experimental/Dataflow/Source/DataflowEditor/Private/Dataflow/DataflowEditorModeToolkit.cpp
+++ b/Engine/Plugins/Experimental/Dataflow/Source/DataflowEditor/Private/Dataflow/DataflowEditorModeToolkit.cpp
@@ -346,6 +346,7 @@ UEditorInteractiveToolsContext* FDataflowEditorModeToolkit::GetCurrentToolsConte
 	// this should not happen, but don't crash if it does
 	return GetEditorModeManager().GetInteractiveToolsContext();
 }
+
 void FDataflowEditorModeToolkit::OnToolStarted(UInteractiveToolManager* Manager, UInteractiveTool* Tool)
 {
 	ensure(Tool == Manager->GetActiveTool(EToolSide::Left));
@@ -396,6 +397,14 @@ void FDataflowEditorModeToolkit::OnToolEnded(UInteractiveToolManager* Manager, U
 	RefreshConstructionScene();
 }
 
+void FDataflowEditorModeToolkit::CustomizeModeDetailsViewArgs(FDetailsViewArgs& ArgsInOut)
+{
+	FBaseCharacterFXEditorModeToolkit::CustomizeModeDetailsViewArgs(ArgsInOut);
+	ArgsInOut.bAllowSearch = true;
+	ArgsInOut.bAllowFavoriteSystem = true;
+	ArgsInOut.bShowSectionSelector = true;
+}
+
 void FDataflowEditorModeToolkit::RefreshConstructionScene()
 {
 	auto GetConstructionScene = [this]() -> TWeakPtr<FDataflowConstructionScene>
diff --git a/Engine/Plugins/Experimental/Dataflow/Source/DataflowEditor/Private/DataflowRendering/DataflowGeometryCollectionRenderableType.cpp b/Engine/Plugins/Experimental/Dataflow/Source/DataflowEditor/Private/DataflowRendering/DataflowGeometryCollectionRenderableType.cpp
--- a/Engine/Plugins/Experimental/Dataflow/Source/DataflowEditor/Private/DataflowRendering/DataflowGeometryCollectionRenderableType.cpp
+++ b/Engine/Plugins/Experimental/Dataflow/Source/DataflowEditor/Private/DataflowRendering/DataflowGeometryCollectionRenderableType.cpp
@@ -34,6 +34,8 @@
 #include "PlanarCut.h" // for ConvertGeometryCollectionToDynamicMesh
 #include "Dataflow/DataflowEditorStyle.h"
 
+#define LOCTEXT_NAMESPACE "DataflowGeometryCollectionRenderableType"
+
 namespace GeometryCollectionSurfaceRenderableType::Private
 {
 	bool bFastDynamicMeshGeneration = false;
@@ -817,5 +819,15 @@ namespace UE::Dataflow::Private
 		UE_DATAFLOW_REGISTER_RENDERABLE_TYPE(FGeometryCollectionSurfaceRenderableType);
 		UE_DATAFLOW_REGISTER_RENDERABLE_TYPE(FGeometryCollectionUvsRenderableType);
 		UE_DATAFLOW_REGISTER_RENDERABLE_TYPE(FGeometryCollectionAssetRenderableType);
+
+		static const FName SurfaceCategoryName(TEXT("Surface"));
+
+		UDataflowRenderableTypeSettings::RegisterSection(
+			UDataflowGeometryCollectionSurfaceRenderSettings::StaticClass(),
+			"Collection", LOCTEXT("Collection", "Collection"),
+			{ SurfaceCategoryName }
+		);
 	}
 }
+
+#undef LOCTEXT_NAMESPACE 
diff --git a/Engine/Plugins/Experimental/Dataflow/Source/DataflowEditor/Private/DataflowRendering/DataflowRenderableTypeSettings.cpp b/Engine/Plugins/Experimental/Dataflow/Source/DataflowEditor/Private/DataflowRendering/DataflowRenderableTypeSettings.cpp
new file mode 100644
--- /dev/null
+++ b/Engine/Plugins/Experimental/Dataflow/Source/DataflowEditor/Private/DataflowRendering/DataflowRenderableTypeSettings.cpp
@@ -0,0 +1,20 @@
+// Copyright Epic Games, Inc. All Rights Reserved.
+
+#include "DataflowRendering/DataflowRenderableTypeSettings.h"
+#include "PropertyEditorModule.h"
+#include "Modules/ModuleManager.h"
+
+/*static*/ void UDataflowRenderableTypeSettings::RegisterSection(UClass* Class, FName Name, const FText& DisplayName, TConstArrayView<FName> Categories)
+{
+	if (Class)
+	{
+		FPropertyEditorModule& PropertyModule = FModuleManager::LoadModuleChecked<FPropertyEditorModule>("PropertyEditor");
+		TSharedRef<FPropertySection> Section = PropertyModule.FindOrCreateSection(Class->GetFName(), Name, DisplayName);
+
+		// Assign categories to the section
+		for (const FName& Category : Categories)
+		{
+			Section->AddCategory(Category);
+		}
+	}
+}
diff --git a/Engine/Plugins/Experimental/Dataflow/Source/DataflowEditor/Public/Dataflow/DataflowEditorModeToolkit.h b/Engine/Plugins/Experimental/Dataflow/Source/DataflowEditor/Public/Dataflow/DataflowEditorModeToolkit.h
--- a/Engine/Plugins/Experimental/Dataflow/Source/DataflowEditor/Public/Dataflow/DataflowEditorModeToolkit.h
+++ b/Engine/Plugins/Experimental/Dataflow/Source/DataflowEditor/Public/Dataflow/DataflowEditorModeToolkit.h
@@ -47,6 +47,7 @@ public:
 	//~ Begin FModeToolkit interface
 	UE_API virtual void OnToolStarted(UInteractiveToolManager* Manager, UInteractiveTool* Tool) override;
 	UE_API virtual void OnToolEnded(UInteractiveToolManager* Manager, UInteractiveTool* Tool) override;
+	UE_API virtual void CustomizeModeDetailsViewArgs(FDetailsViewArgs& ArgsInOut) override;
 	//~ End FModeToolkit interface
 
 	//~ Begin IToolkit interface
diff --git a/Engine/Plugins/Experimental/Dataflow/Source/DataflowEditor/Public/DataflowRendering/DataflowRenderableTypeSettings.h b/Engine/Plugins/Experimental/Dataflow/Source/DataflowEditor/Public/DataflowRendering/DataflowRenderableTypeSettings.h
--- a/Engine/Plugins/Experimental/Dataflow/Source/DataflowEditor/Public/DataflowRendering/DataflowRenderableTypeSettings.h
+++ b/Engine/Plugins/Experimental/Dataflow/Source/DataflowEditor/Public/DataflowRendering/DataflowRenderableTypeSettings.h
@@ -10,6 +10,8 @@ UCLASS(MinimalAPI)
 class UDataflowRenderableTypeSettings : public UObject
 {
 	GENERATED_BODY()
+public:
+	static void RegisterSection(UClass* Class, FName Name, const FText& DisplayName, TConstArrayView<FName> Categories);
 };
 
 // Common types
