Dataflow : update new rendering  to only read actual value from input instead of pulling them from the graph when the evaluation is in manual mode

#Synced-CL 52497132

diff --git a/Engine/Plugins/Experimental/Dataflow/Source/DataflowEditor/Private/Dataflow/DataflowConstructionScene.cpp b/Engine/Plugins/Experimental/Dataflow/Source/DataflowEditor/Private/Dataflow/DataflowConstructionScene.cpp
--- a/Engine/Plugins/Experimental/Dataflow/Source/DataflowEditor/Private/Dataflow/DataflowConstructionScene.cpp
+++ b/Engine/Plugins/Experimental/Dataflow/Source/DataflowEditor/Private/Dataflow/DataflowConstructionScene.cpp
@@ -678,7 +678,7 @@ void FDataflowConstructionScene::UpdateComponentsFromNode(const UDataflowEdNode*
 				{
 					if (TSharedPtr<const FDataflowNode> DataflowNode = EdNode->GetDataflowNode())
 					{
-						UE::Dataflow::FRenderableTypeInstance::GetRenderableInstancesForNode(*DataflowContext, *ConstructionViewMode, *DataflowNode, NodeRenderableTypeInstances);
+						UE::Dataflow::FRenderableTypeInstance::GetRenderableInstancesForNode(*DataflowContext, *ConstructionViewMode, *DataflowNode, bEvaluateOutputs, NodeRenderableTypeInstances);
 					}
 				}
 				RenderableTypeInstances.Append(NodeRenderableTypeInstances);
diff --git a/Engine/Plugins/Experimental/Dataflow/Source/DataflowEditor/Private/Dataflow/DataflowEditorUtil.cpp b/Engine/Plugins/Experimental/Dataflow/Source/DataflowEditor/Private/Dataflow/DataflowEditorUtil.cpp
--- a/Engine/Plugins/Experimental/Dataflow/Source/DataflowEditor/Private/Dataflow/DataflowEditorUtil.cpp
+++ b/Engine/Plugins/Experimental/Dataflow/Source/DataflowEditor/Private/Dataflow/DataflowEditorUtil.cpp
@@ -149,7 +149,7 @@ namespace UE::Dataflow
 
 				// check the component based rendering system (new)
 				TArray<FRenderableTypeInstance> RenderableInstances;
-				FRenderableTypeInstance::GetRenderableInstancesForNode(*Context, ViewMode, *DataflowNode, RenderableInstances);
+				FRenderableTypeInstance::GetRenderableInstancesForNode(*Context, ViewMode, *DataflowNode, /*bEvaluateOutputs*/false, RenderableInstances);
 				if (RenderableInstances.Num() > 0)
 				{
 					return true;
diff --git a/Engine/Plugins/Experimental/Dataflow/Source/DataflowEditor/Private/DataflowRendering/DataflowRenderableTypeInstance.cpp b/Engine/Plugins/Experimental/Dataflow/Source/DataflowEditor/Private/DataflowRendering/DataflowRenderableTypeInstance.cpp
--- a/Engine/Plugins/Experimental/Dataflow/Source/DataflowEditor/Private/DataflowRendering/DataflowRenderableTypeInstance.cpp
+++ b/Engine/Plugins/Experimental/Dataflow/Source/DataflowEditor/Private/DataflowRendering/DataflowRenderableTypeInstance.cpp
@@ -37,7 +37,7 @@ namespace UE::Dataflow
 		}
 	}
 
-	/* static */ void FRenderableTypeInstance::GetRenderableInstancesForNode(FContext& Context, const IDataflowConstructionViewMode& ViewMode, const FDataflowNode& Node, TArray<FRenderableTypeInstance>& OutRenderableTypeInstances)
+	/* static */ void FRenderableTypeInstance::GetRenderableInstancesForNode(FContext& Context, const IDataflowConstructionViewMode& ViewMode, const FDataflowNode& Node, bool bEvaluateOutputs, TArray<FRenderableTypeInstance>& OutRenderableTypeInstances)
 	{
 		OutRenderableTypeInstances.Reset();
 
@@ -57,7 +57,7 @@ namespace UE::Dataflow
 				{
 					if (RenderableType  && RenderableType->IsViewModeSupported(ViewMode))
 					{
-						FRenderableTypeInstance RenderableTypeInstance(Context, Node.AsShared(), Output->GetName(), RenderableType);
+						FRenderableTypeInstance RenderableTypeInstance(Context, Node.AsShared(), Output->GetName(), RenderableType, bEvaluateOutputs);
 						if (RenderableTypeInstance.CanRender())
 						{
 							if (Private::IsVisibleRenderGroup(RenderableType->GetRenderGroup(), *Output))
diff --git a/Engine/Plugins/Experimental/Dataflow/Source/DataflowEditor/Public/DataflowRendering/DataflowRenderableTypeInstance.h b/Engine/Plugins/Experimental/Dataflow/Source/DataflowEditor/Public/DataflowRendering/DataflowRenderableTypeInstance.h
--- a/Engine/Plugins/Experimental/Dataflow/Source/DataflowEditor/Public/DataflowRendering/DataflowRenderableTypeInstance.h
+++ b/Engine/Plugins/Experimental/Dataflow/Source/DataflowEditor/Public/DataflowRendering/DataflowRenderableTypeInstance.h
@@ -29,7 +29,7 @@ namespace UE::Dataflow
 	{
 	public:
 		/** find all renderable instances for a specific node and a specific view mode */
-		static void GetRenderableInstancesForNode(FContext& Context, const IDataflowConstructionViewMode& ViewMode, const FDataflowNode& Node, TArray<FRenderableTypeInstance>& OutRenderableInstances);
+		static void GetRenderableInstancesForNode(FContext& Context, const IDataflowConstructionViewMode& ViewMode, const FDataflowNode& Node, bool bEvaluateOutputs, TArray<FRenderableTypeInstance>& OutRenderableInstances);
 
 		/** Whether this instance can render */
 		bool CanRender() const;
@@ -103,7 +103,7 @@ namespace UE::Dataflow
 			{
 				if (const FDataflowOutput* Output = Node->FindOutput(Name))
 				{
-					if (EvaluationContext.IsThreaded())
+					if (EvaluationContext.IsThreaded() || !bEvaluateOutputs)
 					{
 						return Output->ReadValue<T>(EvaluationContext, Default);
 					}
@@ -157,11 +157,12 @@ namespace UE::Dataflow
 
 		UE_API static bool HasMetadata(const FDataflowOutput& Output, FName Metadata);
 
-		FRenderableTypeInstance(FContext& InContext, const TSharedPtr<const FDataflowNode>& InNode, const FName InOutputName, const IRenderableType* InRenderableType)
+		FRenderableTypeInstance(FContext& InContext, const TSharedPtr<const FDataflowNode>& InNode, const FName InOutputName, const IRenderableType* InRenderableType, bool bEvaluateOutputsIn)
 			: EvaluationContext(InContext)
 			, Node(InNode)
 			, OutputName(InOutputName)
 			, RenderableType(InRenderableType)
+			, bEvaluateOutputs(bEvaluateOutputsIn)
 		{}
 
 	private:
@@ -169,6 +170,7 @@ namespace UE::Dataflow
 		TSharedPtr<const FDataflowNode> Node;
 		const FName OutputName;
 		const IRenderableType* RenderableType;
+		bool bEvaluateOutputs;
 	};
 }
 
