Dataflow : do not launch tools when nodes are not up to date and display a clear message about it in the construction viewport
- this is needed when in manual mode where the construuction view will not pull the data from the graph

#Synced-CL 52497132

diff --git a/Engine/Plugins/Experimental/Dataflow/Source/DataflowEditor/Private/Dataflow/DataflowEditorMode.cpp b/Engine/Plugins/Experimental/Dataflow/Source/DataflowEditor/Private/Dataflow/DataflowEditorMode.cpp
--- a/Engine/Plugins/Experimental/Dataflow/Source/DataflowEditor/Private/Dataflow/DataflowEditorMode.cpp
+++ b/Engine/Plugins/Experimental/Dataflow/Source/DataflowEditor/Private/Dataflow/DataflowEditorMode.cpp
@@ -275,6 +275,25 @@ bool UDataflowEditorMode::CanSelectedNodeLaunchTool() const
 	return false;
 }
 
+bool UDataflowEditorMode::IsSelectedNodeUpToDate(UE::Dataflow::FContext& Context) const
+{
+	if (TObjectPtr<UDataflowEdNode> SelectedNode = GetSelectedNode())
+	{
+		if (TSharedPtr<FDataflowNode> DataflowNode = SelectedNode->GetDataflowNode())
+		{
+			TArray<FDataflowOutput*> Outputs = DataflowNode->GetOutputs();
+			for (const FDataflowOutput* Output : Outputs)
+			{
+				if (Output && Context.HasData(Output->CacheKey(), DataflowNode->GetTimestamp()))
+				{
+					return true;
+				}
+			}
+		}
+	}
+	return false;
+}
+
 bool UDataflowEditorMode::IsDataflowEditable() const
 {
 	if (TSharedPtr<SDataflowGraphEditor> GraphEditor = DataflowGraphEditor.Pin())
@@ -750,37 +769,54 @@ void UDataflowEditorMode::ModeTick(float DeltaTime)
 		FirstTimeFocusConstructionViewport();
 	}
 
+	TSharedPtr<UE::Dataflow::FContext> DataflowContext;
+	if (const UEditorInteractiveToolsContext* const ConstructionToolsContext = GetInteractiveToolsContext())
+	{
+		if (UDataflowContextObject* ContextObject = ConstructionToolsContext->ContextObjectStore->FindContext<UDataflowContextObject>())
+		{
+			DataflowContext = ContextObject->GetDataflowContext();
+		}
+	}
+
 	TObjectPtr<UDataflowEdNode> SelectedNode = GetSelectedNode();
 
+
 	if (CanSelectedNodeLaunchTool())
 	{
 		if (TSharedPtr<const FUICommandInfo> Command = UDataflowEditorMode::GetSelectedNodeToolCommand())
 		{
-			const bool bLaunchedTool = Toolkit->GetToolkitCommands()->TryExecuteAction(Command.ToSharedRef());
-			if (!bLaunchedTool)
+			if (DataflowContext && !IsSelectedNodeUpToDate(*DataflowContext))
+			{
+				UserMessage = FText::Format(LOCTEXT("OutOfDateNode", "{0} tool cannot launch because the node is out of date and need evaluation (Evaluation mode may be set to Manual)"), Command->GetLabel());
+			}
+			else
 			{
-				const USelection* Selection = ConstructionScenePtr ? ConstructionScenePtr->GetSelectedComponents(): nullptr;
-				if (Selection && Selection->Num() == 0)
+				const bool bLaunchedTool = Toolkit->GetToolkitCommands()->TryExecuteAction(Command.ToSharedRef());
+				if (!bLaunchedTool)
 				{
-					// TODO : find a more accurate message ?
-					if (!ConstructionScenePtr->HasRenderableGeometry())
+					const USelection* Selection = ConstructionScenePtr ? ConstructionScenePtr->GetSelectedComponents() : nullptr;
+					if (Selection && Selection->Num() == 0)
 					{
-						UserMessage = FText::Format(LOCTEXT("NoComponentAvailable", "{0} tool cannot launch because there's no component in the construction view - check the node settings"), Command->GetLabel());
+						// TODO : find a more accurate message ?
+						if (!ConstructionScenePtr->HasRenderableGeometry())
+						{
+							UserMessage = FText::Format(LOCTEXT("NoComponentAvailable", "{0} tool cannot launch because there's no component in the construction view - check the node settings"), Command->GetLabel());
+						}
+						else
+						{
+							UserMessage = FText::Format(LOCTEXT("ToolNeedOneComponent", "You must select a component to launch the {0} tool"), Command->GetLabel());
+						}
 					}
 					else
 					{
-						UserMessage = FText::Format(LOCTEXT("ToolNeedOneComponent", "You must select a component to launch the {0} tool"), Command->GetLabel());
+						UserMessage = FText::Format(LOCTEXT("ToolFailedToLaunch", "Failed to launch {0} tool - Check Editor logs"), Command->GetLabel());
 					}
 				}
 				else
 				{
-					UserMessage = FText::Format(LOCTEXT("ToolFailedToLaunch", "Failed to launch {0} tool - Check Editor logs"), Command->GetLabel());
+					UserMessage = FText::GetEmpty();
 				}
 			}
-			else
-			{
-				UserMessage = FText::GetEmpty();
-			}
 		}
 	}
 	else
diff --git a/Engine/Plugins/Experimental/Dataflow/Source/DataflowEditor/Public/Dataflow/DataflowEditorMode.h b/Engine/Plugins/Experimental/Dataflow/Source/DataflowEditor/Public/Dataflow/DataflowEditorMode.h
--- a/Engine/Plugins/Experimental/Dataflow/Source/DataflowEditor/Public/Dataflow/DataflowEditorMode.h
+++ b/Engine/Plugins/Experimental/Dataflow/Source/DataflowEditor/Public/Dataflow/DataflowEditorMode.h
@@ -20,6 +20,7 @@ class UDataflowEditor;
 namespace UE::Dataflow
 {
 	class IDataflowConstructionViewMode;
+	class FContext;
 }
 enum class EToolShutdownType : uint8;
 
@@ -196,6 +197,7 @@ private:
 
 	bool CanUseTools() const;
 	bool CanSelectedNodeLaunchTool() const;
+	bool IsSelectedNodeUpToDate(UE::Dataflow::FContext& Context) const;
 	TSharedPtr<const FUICommandInfo> GetSelectedNodeToolCommand() const;
 	bool IsDataflowEditable() const;
 
