Dataflow : Move all assets to use the new way to bind dataflow menu commends in the asset context menu

#Synced-CL 52497132

diff --git a/Engine/Plugins/ChaosClothAssetEditorCore/Source/ChaosClothAssetEditor/Private/ChaosClothAsset/AssetDefinition_ClothAsset.cpp b/Engine/Plugins/ChaosClothAssetEditorCore/Source/ChaosClothAssetEditor/Private/ChaosClothAsset/AssetDefinition_ClothAsset.cpp
--- a/Engine/Plugins/ChaosClothAssetEditorCore/Source/ChaosClothAssetEditor/Private/ChaosClothAsset/AssetDefinition_ClothAsset.cpp
+++ b/Engine/Plugins/ChaosClothAssetEditorCore/Source/ChaosClothAssetEditor/Private/ChaosClothAsset/AssetDefinition_ClothAsset.cpp
@@ -152,7 +152,4 @@ bool UAssetDefinition_ClothAsset::LaunchClothDataflowAssetEditor(UChaosClothAsse
 	return false;
 }
 
-// Register dataflow related right-click menu options
-DATAFLOW_REGISTER_ASSET_MENUS(UChaosClothAsset)
-
 #undef LOCTEXT_NAMESPACE
diff --git a/Engine/Plugins/ChaosClothAssetEditorCore/Source/ChaosClothAssetEditor/Private/ChaosClothAsset/ChaosClothAssetEditorModule.cpp b/Engine/Plugins/ChaosClothAssetEditorCore/Source/ChaosClothAssetEditor/Private/ChaosClothAsset/ChaosClothAssetEditorModule.cpp
--- a/Engine/Plugins/ChaosClothAssetEditorCore/Source/ChaosClothAssetEditor/Private/ChaosClothAsset/ChaosClothAssetEditorModule.cpp
+++ b/Engine/Plugins/ChaosClothAssetEditorCore/Source/ChaosClothAssetEditor/Private/ChaosClothAsset/ChaosClothAssetEditorModule.cpp
@@ -18,6 +18,7 @@
 #include "EditorModeRegistry.h"
 #include "Algo/AllOf.h"
 #include "ContentBrowserMenuContexts.h"
+#include "Dataflow/AssetDefinition_DataflowAsset.h"
 #include "Dataflow/DataflowEditorModeUILayer.h"
 #include "Dataflow/DataflowEditor.h"
 #include "ComponentAssetBroker.h"
@@ -162,6 +163,9 @@ namespace UE::Chaos::ClothAsset
 
 		ClothAssetComponentBroker = MakeShareable(new FClothAssetComponentBroker);
 		FComponentAssetBrokerage::RegisterBroker(ClothAssetComponentBroker, UChaosClothComponent::StaticClass(), true, true);
+
+		DataflowMenusHandle = UE::DataflowAssetDefinitionHelpers::RegisterDataflowAssetMenus(UChaosClothAsset::StaticClass());
+
 	}
 
 	void FChaosClothAssetEditorModule::ShutdownModule()
@@ -176,6 +180,8 @@ namespace UE::Chaos::ClothAsset
 			Var->OnChangedDelegate().Remove(OnCVarChangedDelegateHandle);
 		}
 
+		UE::DataflowAssetDefinitionHelpers::UnregisterDataflowAssetMenus(DataflowMenusHandle);
+
 		Dataflow::FDataflowNodeDetailExtensionRegistry::GetInstance().DeregisterExtension(FClothSimulationNodeDetailExtender::Name);
 		Dataflow::FDataflowConstructionVisualizationRegistry::GetInstance().DeregisterVisualization(FClothDataflowConstructionVisualization::Name);
 		Dataflow::FDataflowSimulationVisualizationRegistry::GetInstance().DeregisterVisualization(FClothDataflowSimulationVisualization::Name);
diff --git a/Engine/Plugins/ChaosClothAssetEditorCore/Source/ChaosClothAssetEditor/Private/ChaosClothAsset/ChaosClothAssetEditorModule.h b/Engine/Plugins/ChaosClothAssetEditorCore/Source/ChaosClothAssetEditor/Private/ChaosClothAsset/ChaosClothAssetEditorModule.h
--- a/Engine/Plugins/ChaosClothAssetEditorCore/Source/ChaosClothAssetEditor/Private/ChaosClothAsset/ChaosClothAssetEditorModule.h
+++ b/Engine/Plugins/ChaosClothAssetEditorCore/Source/ChaosClothAssetEditor/Private/ChaosClothAsset/ChaosClothAssetEditorModule.h
@@ -19,6 +19,7 @@ namespace UE::Chaos::ClothAsset
 
 		FDelegateHandle StartupCallbackDelegateHandle;
 		FDelegateHandle OnCVarChangedDelegateHandle;
+		FDelegateHandle DataflowMenusHandle;
 
 		void RegisterMenus();
 		
diff --git a/Engine/Plugins/Experimental/ChaosFlesh/Source/ChaosFleshEditor/Private/ChaosFlesh/Asset/AssetDefinition_FleshAsset.cpp b/Engine/Plugins/Experimental/ChaosFlesh/Source/ChaosFleshEditor/Private/ChaosFlesh/Asset/AssetDefinition_FleshAsset.cpp
--- a/Engine/Plugins/Experimental/ChaosFlesh/Source/ChaosFleshEditor/Private/ChaosFlesh/Asset/AssetDefinition_FleshAsset.cpp
+++ b/Engine/Plugins/Experimental/ChaosFlesh/Source/ChaosFleshEditor/Private/ChaosFlesh/Asset/AssetDefinition_FleshAsset.cpp
@@ -100,6 +100,4 @@ EAssetCommandResult UAssetDefinition_FleshAsset::OpenAssets(const FAssetOpenArgs
 	return EAssetCommandResult::Unhandled;
 }
 
-DATAFLOW_REGISTER_ASSET_MENUS(UFleshAsset);
-
 #undef LOCTEXT_NAMESPACE
diff --git a/Engine/Plugins/Experimental/ChaosFlesh/Source/ChaosFleshEditor/Private/ChaosFlesh/ChaosFleshEditorPlugin.cpp b/Engine/Plugins/Experimental/ChaosFlesh/Source/ChaosFleshEditor/Private/ChaosFlesh/ChaosFleshEditorPlugin.cpp
--- a/Engine/Plugins/Experimental/ChaosFlesh/Source/ChaosFleshEditor/Private/ChaosFlesh/ChaosFleshEditorPlugin.cpp
+++ b/Engine/Plugins/Experimental/ChaosFlesh/Source/ChaosFleshEditor/Private/ChaosFlesh/ChaosFleshEditorPlugin.cpp
@@ -13,6 +13,7 @@
 #include "ChaosFlesh/FleshActor.h"
 #include "ChaosFlesh/FleshAsset.h"
 #include "ChaosFlesh/FleshComponent.h"
+#include "Dataflow/AssetDefinition_DataflowAsset.h"
 #include "Dataflow/DataflowTemplateRegistry.h"
 #include "Editor/FleshEditorStyle.h"
 
@@ -94,10 +95,13 @@ void IChaosFleshEditorPlugin::StartupModule()
 	FDataflowTemplateRegistry::Get().RegisterTemplateFolder(UFleshAsset::StaticClass(),
 		{ TEXT("/ChaosFlesh/Templates"), FSlateIcon("ChaosFleshEditorStyle", "ClassThumbnail.FleshAsset") });
 	
+	DataflowMenusHandle = UE::DataflowAssetDefinitionHelpers::RegisterDataflowAssetMenus(UFleshAsset::StaticClass());
+
 }
 
 void IChaosFleshEditorPlugin::ShutdownModule()
 {
+	UE::DataflowAssetDefinitionHelpers::UnregisterDataflowAssetMenus(DataflowMenusHandle);
 }
 
 IMPLEMENT_MODULE(IChaosFleshEditorPlugin, ChaosFleshEditor)
diff --git a/Engine/Plugins/Experimental/ChaosFlesh/Source/ChaosFleshEditor/Public/ChaosFlesh/ChaosFleshEditorPlugin.h b/Engine/Plugins/Experimental/ChaosFlesh/Source/ChaosFleshEditor/Public/ChaosFlesh/ChaosFleshEditorPlugin.h
--- a/Engine/Plugins/Experimental/ChaosFlesh/Source/ChaosFleshEditor/Public/ChaosFlesh/ChaosFleshEditorPlugin.h
+++ b/Engine/Plugins/Experimental/ChaosFlesh/Source/ChaosFleshEditor/Public/ChaosFlesh/ChaosFleshEditorPlugin.h
@@ -60,5 +60,7 @@ private:
 
 	// Styleset for flesh tool brushes/fonts etc.
 	TSharedPtr<FSlateStyleSet> StyleSet;
+
+	FDelegateHandle DataflowMenusHandle;
 };
 
diff --git a/Engine/Plugins/Experimental/ChaosOutfitAsset/Source/ChaosOutfitAssetEditor/Private/ChaosOutfitAsset/AssetDefinition_OutfitAsset.cpp b/Engine/Plugins/Experimental/ChaosOutfitAsset/Source/ChaosOutfitAssetEditor/Private/ChaosOutfitAsset/AssetDefinition_OutfitAsset.cpp
--- a/Engine/Plugins/Experimental/ChaosOutfitAsset/Source/ChaosOutfitAssetEditor/Private/ChaosOutfitAsset/AssetDefinition_OutfitAsset.cpp
+++ b/Engine/Plugins/Experimental/ChaosOutfitAsset/Source/ChaosOutfitAssetEditor/Private/ChaosOutfitAsset/AssetDefinition_OutfitAsset.cpp
@@ -84,7 +84,4 @@ EAssetCommandResult UAssetDefinition_OutfitAsset::OpenAssets(const FAssetOpenArg
 	return EAssetCommandResult::Unhandled;
 }
 
-// Register dataflow related right-click menu options
-DATAFLOW_REGISTER_ASSET_MENUS(UChaosOutfitAsset)
-
 #undef LOCTEXT_NAMESPACE
diff --git a/Engine/Plugins/Experimental/ChaosOutfitAsset/Source/ChaosOutfitAssetEditor/Private/ChaosOutfitAsset/ChaosOutfitAssetEditorModule.cpp b/Engine/Plugins/Experimental/ChaosOutfitAsset/Source/ChaosOutfitAssetEditor/Private/ChaosOutfitAsset/ChaosOutfitAssetEditorModule.cpp
--- a/Engine/Plugins/Experimental/ChaosOutfitAsset/Source/ChaosOutfitAssetEditor/Private/ChaosOutfitAsset/ChaosOutfitAssetEditorModule.cpp
+++ b/Engine/Plugins/Experimental/ChaosOutfitAsset/Source/ChaosOutfitAssetEditor/Private/ChaosOutfitAsset/ChaosOutfitAssetEditorModule.cpp
@@ -7,6 +7,7 @@
 #include "ChaosOutfitAsset/OutfitAsset.h"
 #include "ChaosOutfitAsset/OutfitAssetEditorStyle.h"
 #include "ChaosOutfitAsset/OutfitEditorCommands.h"
+#include "Dataflow/AssetDefinition_DataflowAsset.h"
 #include "Dataflow/DataflowTemplateRegistry.h"
 #include "Engine/SkeletalMesh.h"
 #include "Animation/Skeleton.h"
@@ -79,6 +80,8 @@ namespace UE::Chaos::OutfitAsset
 					 .Icon = ChaosOutfitAssetIcon
 				});
 
+			DataflowMenusHandle = UE::DataflowAssetDefinitionHelpers::RegisterDataflowAssetMenus(UChaosOutfitAsset::StaticClass());
+
 			// TODO Register visualizations (see FChaosClothAssetEditorModule::StartupModule() in ChaosClothAssetEditorModule.cpp)
 		}
 
@@ -88,6 +91,8 @@ namespace UE::Chaos::OutfitAsset
 			{
 				// TODO Deregister visualizations (see FChaosClothAssetEditorModule::StartupModule() in ChaosClothAssetEditorModule.cpp)
 
+				UE::DataflowAssetDefinitionHelpers::UnregisterDataflowAssetMenus(DataflowMenusHandle);
+
 				UToolMenus::UnRegisterStartupCallback(StartupCallbackDelegateHandle);
 				FOutfitEditorCommands::Unregister();
 
@@ -205,6 +210,7 @@ private:
 		}
 
 		FDelegateHandle StartupCallbackDelegateHandle;
+		FDelegateHandle DataflowMenusHandle;
 	};
 } // namespace UE::Chaos::OutfitAsset
 
diff --git a/Engine/Plugins/Experimental/Dataflow/Source/DataflowEditor/Private/Dataflow/AssetDefinition_DataflowAsset.cpp b/Engine/Plugins/Experimental/Dataflow/Source/DataflowEditor/Private/Dataflow/AssetDefinition_DataflowAsset.cpp
--- a/Engine/Plugins/Experimental/Dataflow/Source/DataflowEditor/Private/Dataflow/AssetDefinition_DataflowAsset.cpp
+++ b/Engine/Plugins/Experimental/Dataflow/Source/DataflowEditor/Private/Dataflow/AssetDefinition_DataflowAsset.cpp
@@ -5,12 +5,14 @@
 #include "AssetRegistry/AssetRegistryModule.h"
 #include "ContentBrowserMenuContexts.h"
 #include "ContentBrowserModule.h"
+#include "Dataflow/DataflowAttachment.h"
 #include "Dataflow/DataflowEditor.h"
 #include "Dataflow/DataflowEditorModule.h"
 #include "Dataflow/DataflowObject.h"
 #include "Dataflow/DataflowInstance.h"
 #include "Dataflow/DataflowTemplateRegistry.h"
 #include "Dataflow/SDataflowTemplatePicker.h"
+#include "Dataflow/DataflowSimulationScene.h"
 #include "Dialog/SMessageDialog.h"
 #include "Framework/Notifications/NotificationManager.h"
 #include "IContentBrowserSingleton.h"
@@ -306,45 +308,34 @@ namespace UE::DataflowAssetDefinitionHelpers
 
 	//////////////////////////////////////////////////////////////////////////////////////////////////////////////////////////////////////
 
-	static void ExecuteEmbedDataflow(const FToolMenuContext& InContext)
+	static void ExecuteEmbedDataflow(const TArray<FAssetData> Assets)
 	{
-		if (const UContentBrowserAssetContextMenuContext* Context = UContentBrowserAssetContextMenuContext::FindContextWithAssets(InContext))
+		for (const FAssetData& AssetData: Assets)
 		{
-			TArray<UObject*> SelectedAssets = Context->LoadSelectedObjects<UObject>();
-			if (!SelectedAssets.IsEmpty())
+			if (UObject* SelectedAsset = AssetData.GetAsset())
 			{
-				for (UObject* SelectedAsset : SelectedAssets)
+				UDataflow* NewDataflow = nullptr;
+				if (UE::DataflowAssetDefinitionHelpers::EmbedDataflowAsset(SelectedAsset, NewDataflow))
 				{
-					if (SelectedAsset)
-					{
-						UDataflow* NewDataflow = nullptr;
-						if (UE::DataflowAssetDefinitionHelpers::EmbedDataflowAsset(SelectedAsset, NewDataflow))
-						{
-							SelectedAsset->Modify();
-						}
-					}
+					SelectedAsset->Modify();
 				}
 			}
 		}
 	}
 
-	static bool CanExecuteEmbedDataflow(const FToolMenuContext& InContext)
+	static bool CanExecuteEmbedDataflow(const TArray<FAssetData> Assets)
 	{
-		if (const UContentBrowserAssetContextMenuContext* Context = UContentBrowserAssetContextMenuContext::FindContextWithAssets(InContext))
+		for (const FAssetData& AssetData : Assets)
 		{
-			TArray<UObject*> SelectedAssets = Context->LoadSelectedObjects<UObject>();
-			if (!SelectedAssets.IsEmpty())
+			if (UObject* SelectedAsset = AssetData.GetAsset())
 			{
-				for (UObject* SelectedAsset : SelectedAssets)
+				if (IDataflowInstanceInterface* Interface = UE::Dataflow::InstanceUtils::GetDataflowInterfaceFromObject(SelectedAsset))
 				{
-					if (IDataflowInstanceInterface* Interface = UE::Dataflow::InstanceUtils::GetDataflowInterfaceFromObject(SelectedAsset))
+					if (UDataflow* DataflowObject = Interface->GetDataflowInstance().GetDataflowAsset())
 					{
-						if (UDataflow* DataflowObject = Interface->GetDataflowInstance().GetDataflowAsset())
+						if (DataflowObject->IsAsset())
 						{
-							if (DataflowObject->IsAsset())
-							{
-								return true;
-							}
+							return true;
 						}
 					}
 				}
@@ -353,117 +344,102 @@ namespace UE::DataflowAssetDefinitionHelpers
 		return false;
 	}
 
-	void RegisterEmbedDataflowAssetMenu(UClass* AssetClass)
+	//////////////////////////////////////////////////////////////////////////////////////////////////////////////////////////////////////
+
+	static void ExecuteSetDataflow(const TArray<FAssetData> Assets)
 	{
-		UToolMenus::RegisterStartupCallback(
-			FSimpleMulticastDelegate::FDelegate::CreateLambda([AssetClass]()
-				{
-					FToolMenuOwnerScoped OwnerScoped(UE_MODULE_NAME);
-					UToolMenu* Menu = UE::ContentBrowser::ExtendToolMenu_AssetContextMenu(AssetClass);
-					FToolMenuSection& Section = Menu->FindOrAddSection("GetAssetActions");
-					{
-						const TAttribute<FText> Label = LOCTEXT("EmbedDataflow_Label", "Embed Dataflow");
-						const TAttribute<FText> ToolTip = LOCTEXT("EmbedDataflow_Tooltip", "Embed the external Dataflow graph in the current asset, the external asset will remain but will be disconnected");
-						const TAttribute<FSlateIcon> Icon = FSlateIcon(FAppStyle::GetAppStyleSetName(), "GraphEditor.EventGraph_16x");
-
-						FToolUIAction UIAction;
-						UIAction.ExecuteAction = FToolMenuExecuteAction::CreateStatic(&ExecuteEmbedDataflow);
-						UIAction.CanExecuteAction = FToolMenuCanExecuteAction::CreateStatic(&CanExecuteEmbedDataflow);
-						Section.AddMenuEntry("Dataflow_EmbedDataflow", Label, ToolTip, Icon, UIAction);
-					}
-				}));
+		for (const FAssetData& AssetData: Assets)
+		{
+			if (UObject* AssetBase = AssetData.GetAsset())
+			{
+				UE::DataflowAssetDefinitionHelpers::SetDataflowFromTemplatePicker(AssetBase);
+			}
+		}
+	}
+	
+	static bool CanExecuteSetDataflow(const TArray<FAssetData> Assets)
+	{
+		// only enable when one asset is selected because this will open a dialog box for each of them 
+		return (Assets.Num() == 1);
 	}
 
 	//////////////////////////////////////////////////////////////////////////////////////////////////////////////////////////////////////
 
-	static void ExecuteSetDataflow(const FToolMenuContext& InContext)
+	static void ExecuteReEvaluateDataflow(const TArray<FAssetData> Assets)
 	{
-		if (const UContentBrowserAssetContextMenuContext* Context = UContentBrowserAssetContextMenuContext::FindContextWithAssets(InContext))
+		for (const FAssetData& AssetData : Assets)
 		{
-			TArray<UObject*> SelectedAssets = Context->LoadSelectedObjects<UObject>();
-			if (!SelectedAssets.IsEmpty())
+			if (UObject* SelectedAsset = AssetData.GetAsset())
 			{
-				for (UObject* SelectedAsset : SelectedAssets)
+				if (IDataflowInstanceInterface* Interface = UE::Dataflow::InstanceUtils::GetDataflowInterfaceFromObject(SelectedAsset))
 				{
-					if (SelectedAsset)
+					if (Interface->GetDataflowInstance().UpdateOwnerAsset(/*bUpdateDependentAsset*/true))
 					{
-						UE::DataflowAssetDefinitionHelpers::SetDataflowFromTemplatePicker(SelectedAsset);
+						SelectedAsset->Modify();
 					}
 				}
 			}
 		}
 	}
 
-	static bool CanExecuteSetDataflow(const FToolMenuContext& InContext)
+	static bool CanReEvaluateDataflow(const TArray<FAssetData> Assets)
 	{
-		if (const UContentBrowserAssetContextMenuContext* Context = UContentBrowserAssetContextMenuContext::FindContextWithAssets(InContext))
+		for (const FAssetData& AssetData : Assets)
 		{
-			TArray<UObject*> SelectedAssets = Context->LoadSelectedObjects<UObject>();
-			// only enable when one asset is selected because this will open a dialog box for each of them 
-			return (SelectedAssets.Num() == 1);
-		}
-		return false;
-	}
-
-	void RegisterSetDataflowAssetMenu(UClass* AssetClass)
-	{
-		UToolMenus::RegisterStartupCallback(
-			FSimpleMulticastDelegate::FDelegate::CreateLambda([AssetClass]()
+			if (UObject* SelectedAsset = AssetData.GetAsset())
+			{
+				if (IDataflowInstanceInterface* Interface = UE::Dataflow::InstanceUtils::GetDataflowInterfaceFromObject(SelectedAsset))
 				{
-					FToolMenuOwnerScoped OwnerScoped(UE_MODULE_NAME);
-					UToolMenu* Menu = UE::ContentBrowser::ExtendToolMenu_AssetContextMenu(AssetClass);
-					FToolMenuSection& Section = Menu->FindOrAddSection("GetAssetActions");
+					if (UDataflow* DataflowObject = Interface->GetDataflowInstance().GetDataflowAsset())
 					{
-						const TAttribute<FText> Label = LOCTEXT("SetDataflow_Label", "Set Dataflow...");
-						const TAttribute<FText> ToolTip = LOCTEXT("SetDataflow_Tooltip", "Select and set a dataflow graph for this asset. This will open the dataflow template selection dialog.");
-						const TAttribute<FSlateIcon> Icon = FSlateIcon(FAppStyle::GetAppStyleSetName(), "GraphEditor.EventGraph_16x");
-
-						FToolUIAction UIAction;
-						UIAction.ExecuteAction = FToolMenuExecuteAction::CreateStatic(&ExecuteSetDataflow);
-						UIAction.CanExecuteAction = FToolMenuCanExecuteAction::CreateStatic(&CanExecuteSetDataflow);
-						Section.AddMenuEntry("Dataflow_SetDataflow", Label, ToolTip, Icon, UIAction);
+						return true;
 					}
-				}));
+				}
+			}
+		}
+		return false;
 	}
 
 	//////////////////////////////////////////////////////////////////////////////////////////////////////////////////////////////////////
 
-	static void ExecuteReEvaluateDataflow(const FToolMenuContext& InContext)
+	static void ExecuteOpenDataflowEditor(const TArray<FAssetData> Assets)
 	{
-		if (const UContentBrowserAssetContextMenuContext* Context = UContentBrowserAssetContextMenuContext::FindContextWithAssets(InContext))
+		if (UAssetEditorSubsystem* const AssetEditorSubsystem = GEditor->GetEditorSubsystem<UAssetEditorSubsystem>())
 		{
-			TArray<UObject*> SelectedAssets = Context->LoadSelectedObjects<UObject>();
-			if (!SelectedAssets.IsEmpty())
+			for (const FAssetData& AssetData : Assets)
 			{
-				for (UObject* SelectedAsset : SelectedAssets)
+				if (UObject* SelectedAsset = AssetData.GetAsset())
 				{
-					if (IDataflowInstanceInterface* Interface = UE::Dataflow::InstanceUtils::GetDataflowInterfaceFromObject(SelectedAsset))
+					if (UDataflowAttachment* Attachment = UE::Dataflow::InstanceUtils::GetDataflowAttachmentFromObject(SelectedAsset))
 					{
-						if (Interface->GetDataflowInstance().UpdateOwnerAsset(/*bUpdateDependentAsset*/true))
+						if (UDataflow* DataflowObject = Attachment->GetDataflowInstance().GetDataflowAsset())
 						{
-							SelectedAsset->Modify();
+							if (UDataflowEditor* const AssetEditor = NewObject<UDataflowEditor>(AssetEditorSubsystem, NAME_None, RF_Transient))
+							{
+								// TODO(dataflow) : should probably read the setting from the attachment if available
+								const TSubclassOf<AActor> ActorClass; // = StaticLoadClass(AActor::StaticClass(), nullptr,
+									//TEXT("/ChaosClothAssetEditor/BP_ClothPreview.BP_ClothPreview_C"), nullptr, LOAD_None, nullptr);
+								AssetEditor->Initialize({ Attachment }, ActorClass);
+							}
 						}
 					}
 				}
+
 			}
 		}
 	}
 
-	static bool CanReEvaluateDataflow(const FToolMenuContext& InContext)
+	static bool CanOpenDataflowEditor(const TArray<FAssetData> Assets)
 	{
-		if (const UContentBrowserAssetContextMenuContext* Context = UContentBrowserAssetContextMenuContext::FindContextWithAssets(InContext))
+		for (const FAssetData& AssetData : Assets)
 		{
-			TArray<UObject*> SelectedAssets = Context->LoadSelectedObjects<UObject>();
-			if (!SelectedAssets.IsEmpty())
+			if (UObject* SelectedAsset = AssetData.GetAsset())
 			{
-				for (UObject* SelectedAsset : SelectedAssets)
+				if (UDataflowAttachment* Attachment = UE::Dataflow::InstanceUtils::GetDataflowAttachmentFromObject(SelectedAsset))
 				{
-					if (IDataflowInstanceInterface* Interface = UE::Dataflow::InstanceUtils::GetDataflowInterfaceFromObject(SelectedAsset))
+					if (UDataflow* DataflowObject = Attachment->GetDataflowInstance().GetDataflowAsset())
 					{
-						if (UDataflow* DataflowObject = Interface->GetDataflowInstance().GetDataflowAsset())
-						{
-							return true;
-						}
+						return true;
 					}
 				}
 			}
@@ -471,25 +447,94 @@ namespace UE::DataflowAssetDefinitionHelpers
 		return false;
 	}
 
-	void RegisterReEvaluateDataflowAssetMenu(UClass* AssetClass)
+	//////////////////////////////////////////////////////////////////////////////////////////////////////////////////////////////////////
+
+	namespace Private
 	{
-		UToolMenus::RegisterStartupCallback(
-			FSimpleMulticastDelegate::FDelegate::CreateLambda([AssetClass]()
+		static void FillMenu(FMenuBuilder& MenuBuilder, const TArray<FAssetData> Assets, UClass* AssetClass)
+		{
+			if (Assets.Num() == 1 && Assets[0].IsInstanceOf(AssetClass))
+			{
+				MenuBuilder.BeginSection("DataflowActionsMenu_Dataflow", LOCTEXT("MenuSection", "DataFlow Actions"));
+
+				// Set Dataflow
 				{
-					FToolMenuOwnerScoped OwnerScoped(UE_MODULE_NAME);
-					UToolMenu* Menu = UE::ContentBrowser::ExtendToolMenu_AssetContextMenu(AssetClass);
-					FToolMenuSection& Section = Menu->FindOrAddSection("GetAssetActions");
-					{
-						const TAttribute<FText> Label = LOCTEXT("ReEvaluateDataflow_Label", "Re-Evaluate Dataflow");
-						const TAttribute<FText> ToolTip = LOCTEXT("ReEvaluateDataflow_Tooltip", "Re-evaluate the associated Dataflow to regenerate the selected asset(s)");
-						const TAttribute<FSlateIcon> Icon = FSlateIcon(FAppStyle::GetAppStyleSetName(), "GraphEditor.EventGraph_16x");
-
-						FToolUIAction UIAction;
-						UIAction.ExecuteAction = FToolMenuExecuteAction::CreateStatic(&ExecuteReEvaluateDataflow);
-						UIAction.CanExecuteAction = FToolMenuCanExecuteAction::CreateStatic(&CanReEvaluateDataflow);
-						Section.AddMenuEntry("Dataflow_ReEvaluateDataflow", Label, ToolTip, Icon, UIAction);
-					}
-				}));
+					FMenuEntryParams Params;
+					Params.DirectActions.ExecuteAction = FExecuteAction::CreateStatic(&ExecuteSetDataflow, Assets);
+					Params.DirectActions.CanExecuteAction = FCanExecuteAction::CreateStatic(&CanExecuteSetDataflow, Assets);
+					Params.LabelOverride = LOCTEXT("SetDataflow_Label", "Set Dataflow...");
+					Params.ToolTipOverride = LOCTEXT("SetDataflow_Tooltip", "Select and set a Dataflow graph for this asset. This will open the Dataflow template selection dialog.");
+					Params.IconOverride = FSlateIcon(FAppStyle::GetAppStyleSetName(), "GraphEditor.EventGraph_16x");
+					MenuBuilder.AddMenuEntry(Params);
+				}
+
+				// Embed Dataflow
+				{
+					FMenuEntryParams Params;
+					Params.DirectActions.ExecuteAction = FExecuteAction::CreateStatic(&ExecuteEmbedDataflow, Assets);
+					Params.DirectActions.CanExecuteAction = FCanExecuteAction::CreateStatic(&CanExecuteEmbedDataflow, Assets);
+					Params.LabelOverride = LOCTEXT("EmbedDataflow_Label", "Embed Dataflow");
+					Params.ToolTipOverride = LOCTEXT("EmbedDataflow_Tooltip", "Embed the external Dataflow graph in the current asset, the external asset will remain but will be disconnected");
+					Params.IconOverride = FSlateIcon(FAppStyle::GetAppStyleSetName(), "GraphEditor.EventGraph_16x");
+					MenuBuilder.AddMenuEntry(Params);
+				}
+
+				// Re-evaluate dataflow Dataflow
+				{
+					FMenuEntryParams Params;
+					Params.DirectActions.ExecuteAction = FExecuteAction::CreateStatic(&ExecuteReEvaluateDataflow, Assets);
+					Params.DirectActions.CanExecuteAction = FCanExecuteAction::CreateStatic(&CanReEvaluateDataflow, Assets);
+					Params.LabelOverride = LOCTEXT("ReEvaluateDataflow_Label", "Re-Evaluate Dataflow");
+					Params.ToolTipOverride = LOCTEXT("ReEvaluateDataflow_Tooltip", "Re-evaluate the associated Dataflow to regenerate the selected asset(s)");
+					Params.IconOverride = FSlateIcon(FAppStyle::GetAppStyleSetName(), "GraphEditor.EventGraph_16x");
+					MenuBuilder.AddMenuEntry(Params);
+				}
+
+				// Open Dataflow Editor ( only for asset using the Dataflow Attachment mechanism)
+				{
+					FMenuEntryParams Params;
+					Params.DirectActions.ExecuteAction = FExecuteAction::CreateStatic(&ExecuteOpenDataflowEditor, Assets);
+					Params.DirectActions.IsActionVisibleDelegate = FCanExecuteAction::CreateStatic(&CanOpenDataflowEditor, Assets);
+					Params.LabelOverride = LOCTEXT("OpenDataflowEditor_Label", "Edit in Dataflow Editor");
+					Params.ToolTipOverride = LOCTEXT("OpenDataflowEditor_Tooltip", "Open the dataflow editor to edit the corresponding Dataflow graph");
+					Params.IconOverride = FSlateIcon(FAppStyle::GetAppStyleSetName(), "GraphEditor.EventGraph_16x");
+					MenuBuilder.AddMenuEntry(Params);
+				}
+
+				MenuBuilder.EndSection();
+			}
+		}
+
+		static TSharedRef<FExtender> ExtendAssetContextMenu(const TArray<FAssetData>& Assets, UClass* AssetClass)
+		{
+			TSharedRef<FExtender> Extender = MakeShared<FExtender>();
+			Extender->AddMenuExtension("GetAssetActions", EExtensionHook::After, nullptr, FMenuExtensionDelegate::CreateStatic(&FillMenu, Assets, AssetClass));
+			return Extender;
+		}
+	}
+
+	FDelegateHandle RegisterDataflowAssetMenus(UClass* AssetClass)
+	{
+		FContentBrowserModule& ContentBrowser = FModuleManager::LoadModuleChecked<FContentBrowserModule>(TEXT("ContentBrowser"));
+		TArray<FContentBrowserMenuExtender_SelectedAssets>& AssetMenuExtenders = ContentBrowser.GetAllAssetViewContextMenuExtenders();
+		AssetMenuExtenders.Add(FContentBrowserMenuExtender_SelectedAssets::CreateStatic(&Private::ExtendAssetContextMenu, AssetClass));
+		return AssetMenuExtenders.Last().GetHandle();
+	}
+
+	void UnregisterDataflowAssetMenus(FDelegateHandle Handle)
+	{
+		if (Handle.IsValid() && FModuleManager::Get().IsModuleLoaded(TEXT("ContentBrowser")))
+		{
+			FContentBrowserModule& ContentBrowserModule = FModuleManager::LoadModuleChecked<FContentBrowserModule>(TEXT("ContentBrowser"));
+			TArray<FContentBrowserMenuExtender_SelectedAssets>& AssetMenuExtenders = ContentBrowserModule.GetAllAssetViewContextMenuExtenders();
+
+			AssetMenuExtenders.RemoveAll([Handle](const FContentBrowserMenuExtender_SelectedAssets& InDelegate)
+				{
+					return InDelegate.GetHandle() == Handle;
+				});
+
+			Handle.Reset();
+		}
 	}
 
 	UObject* FactoryCreateNew(
diff --git a/Engine/Plugins/Experimental/Dataflow/Source/DataflowEditor/Public/Dataflow/AssetDefinition_DataflowAsset.h b/Engine/Plugins/Experimental/Dataflow/Source/DataflowEditor/Public/Dataflow/AssetDefinition_DataflowAsset.h
--- a/Engine/Plugins/Experimental/Dataflow/Source/DataflowEditor/Public/Dataflow/AssetDefinition_DataflowAsset.h
+++ b/Engine/Plugins/Experimental/Dataflow/Source/DataflowEditor/Public/Dataflow/AssetDefinition_DataflowAsset.h
@@ -42,14 +42,10 @@ namespace UE::DataflowAssetDefinitionHelpers
 	// Determine whether or not a UDataflow can be opened directly in the Dataflow Editor
 	DATAFLOWEDITOR_API bool CanOpenDataflowAssetInEditor(const UObject* DataflowAsset);
 
-	// Register a entry on the asset menu that will embed an external dataflow in the selected asset
-	DATAFLOWEDITOR_API void RegisterEmbedDataflowAssetMenu(UClass* AssetClass);
+	// Register Dataflow related menus
+	DATAFLOWEDITOR_API FDelegateHandle RegisterDataflowAssetMenus(UClass* AssetClass);
 
-	// Register a entry on the asset menu that will set a dataflow asset for the selected asset
-	DATAFLOWEDITOR_API void RegisterSetDataflowAssetMenu(UClass* AssetClass);
-
-	// Register a entry on the asset menu that will Re-Evaluate the assoiated dataflow for the selected asset
-	DATAFLOWEDITOR_API void RegisterReEvaluateDataflowAssetMenu(UClass* AssetClass);
+	DATAFLOWEDITOR_API void UnregisterDataflowAssetMenus(FDelegateHandle Handle);
 
 	// Create a new Dataflow based asset using the specified Dataflow template, an asset with an empty Dataflow if an empty path is provided, or an unbound asset if the path is nullptr.
 	// @note The Dataflow template is always duplicated to use with the new asset, even when not embedded.
@@ -64,17 +60,6 @@ namespace UE::DataflowAssetDefinitionHelpers
 		const TCHAR* AssetPrefix);
 }
 
-#define DATAFLOW_REGISTER_ASSET_MENUS(AssetClassType) \
-	namespace MenuExtension_##AssetClassType \
-	{ \
-		static FDelayedAutoRegisterHelper DelayedAutoRegister(EDelayedRegisterRunPhase::EndOfEngineInit, []() \
-		{ \
-			UE::DataflowAssetDefinitionHelpers::RegisterSetDataflowAssetMenu(AssetClassType::StaticClass()); \
-			UE::DataflowAssetDefinitionHelpers::RegisterEmbedDataflowAssetMenu(AssetClassType::StaticClass()); \
-			UE::DataflowAssetDefinitionHelpers::RegisterReEvaluateDataflowAssetMenu(AssetClassType::StaticClass()); \
-		}); \
-	} 
-
 struct FDataflowInput;
 struct FDataflowOutput;
 
diff --git a/Engine/Plugins/Experimental/GeometryCollectionPlugin/Source/GeometryCollectionEditor/Private/GeometryCollection/AssetDefinition_GeometryCollection.cpp b/Engine/Plugins/Experimental/GeometryCollectionPlugin/Source/GeometryCollectionEditor/Private/GeometryCollection/AssetDefinition_GeometryCollection.cpp
--- a/Engine/Plugins/Experimental/GeometryCollectionPlugin/Source/GeometryCollectionEditor/Private/GeometryCollection/AssetDefinition_GeometryCollection.cpp
+++ b/Engine/Plugins/Experimental/GeometryCollectionPlugin/Source/GeometryCollectionEditor/Private/GeometryCollection/AssetDefinition_GeometryCollection.cpp
@@ -219,12 +219,6 @@ namespace MenuExtension_GeometryCollection
 						Section.AddMenuEntry("GeometryCollection_ConvertToSkeletalMesh", Label, ToolTip, Icon, UIAction);
 					}
 				}));
-
-			UE::DataflowAssetDefinitionHelpers::RegisterSetDataflowAssetMenu(UGeometryCollection::StaticClass());
-
-			UE::DataflowAssetDefinitionHelpers::RegisterEmbedDataflowAssetMenu(UGeometryCollection::StaticClass());
-
-			UE::DataflowAssetDefinitionHelpers::RegisterReEvaluateDataflowAssetMenu(UGeometryCollection::StaticClass());
 		});
 }
 
diff --git a/Engine/Plugins/Experimental/GeometryCollectionPlugin/Source/GeometryCollectionEditor/Private/GeometryCollection/GeometryCollectionEditorPlugin.cpp b/Engine/Plugins/Experimental/GeometryCollectionPlugin/Source/GeometryCollectionEditor/Private/GeometryCollection/GeometryCollectionEditorPlugin.cpp
--- a/Engine/Plugins/Experimental/GeometryCollectionPlugin/Source/GeometryCollectionEditor/Private/GeometryCollection/GeometryCollectionEditorPlugin.cpp
+++ b/Engine/Plugins/Experimental/GeometryCollectionPlugin/Source/GeometryCollectionEditor/Private/GeometryCollection/GeometryCollectionEditorPlugin.cpp
@@ -8,6 +8,7 @@
 #include "LevelEditor.h"
 #include "Modules/ModuleManager.h"
 #include "PropertyEditorModule.h"
+#include "Dataflow/AssetDefinition_DataflowAsset.h"
 #include "Dataflow/DataflowTemplateRegistry.h"
 #include "GeometryCollection/GeometryCollectionComponent.h"
 #include "GeometryCollection/GeometryCollectionConversion.h"
@@ -264,6 +265,8 @@ void IGeometryCollectionEditorPlugin::StartupModule()
 
 	FDataflowTemplateRegistry::Get().RegisterTemplateFolder(UGeometryCollection::StaticClass(),
 		{ TEXT("/GeometryCollectionPlugin/DataflowTemplates"), FSlateIcon("GeometryCollectionEditorStyle", "ClassThumbnail.GeometryCollection") });
+
+	DataflowMenusHandle = UE::DataflowAssetDefinitionHelpers::RegisterDataflowAssetMenus(UGeometryCollection::StaticClass());
 }
 
 
@@ -271,6 +274,8 @@ void IGeometryCollectionEditorPlugin::ShutdownModule()
 {
 	if (UObjectInitialized())
 	{
+		UE::DataflowAssetDefinitionHelpers::UnregisterDataflowAssetMenus(DataflowMenusHandle);
+
 		UThumbnailManager::Get().UnregisterCustomRenderer(UGeometryCollection::StaticClass());
 
 		FAssetToolsModule& AssetToolsModule = FAssetToolsModule::GetModule();
diff --git a/Engine/Plugins/Experimental/GeometryCollectionPlugin/Source/GeometryCollectionEditor/Public/GeometryCollection/GeometryCollectionEditorPlugin.h b/Engine/Plugins/Experimental/GeometryCollectionPlugin/Source/GeometryCollectionEditor/Public/GeometryCollection/GeometryCollectionEditorPlugin.h
--- a/Engine/Plugins/Experimental/GeometryCollectionPlugin/Source/GeometryCollectionEditor/Public/GeometryCollection/GeometryCollectionEditorPlugin.h
+++ b/Engine/Plugins/Experimental/GeometryCollectionPlugin/Source/GeometryCollectionEditor/Public/GeometryCollection/GeometryCollectionEditorPlugin.h
@@ -77,5 +77,6 @@ private:
 
 	FLevelEditorModule::FLevelEditorMenuExtender ViewMenuExtender;
 	FDelegateHandle ViewMenuExtenderHandle;
+	FDelegateHandle DataflowMenusHandle;
 };
 
diff --git a/Engine/Plugins/Runtime/HairStrands/Source/HairStrandsEditor/Private/GroomActions.cpp b/Engine/Plugins/Runtime/HairStrands/Source/HairStrandsEditor/Private/GroomActions.cpp
--- a/Engine/Plugins/Runtime/HairStrands/Source/HairStrandsEditor/Private/GroomActions.cpp
+++ b/Engine/Plugins/Runtime/HairStrands/Source/HairStrandsEditor/Private/GroomActions.cpp
@@ -727,8 +727,5 @@ static FDelayedAutoRegisterHelper DelayedAutoRegister(EDelayedRegisterRunPhase::
 	}));
 });
 
-// Register dataflow related right-click menu options
-DATAFLOW_REGISTER_ASSET_MENUS(UGroomAsset)
-
 } // namespace MenuExtension_GroomAsset
 #undef LOCTEXT_NAMESPACE
diff --git a/Engine/Plugins/Runtime/HairStrands/Source/HairStrandsEditor/Private/HairStrandsEditor.cpp b/Engine/Plugins/Runtime/HairStrands/Source/HairStrandsEditor/Private/HairStrandsEditor.cpp
--- a/Engine/Plugins/Runtime/HairStrands/Source/HairStrandsEditor/Private/HairStrandsEditor.cpp
+++ b/Engine/Plugins/Runtime/HairStrands/Source/HairStrandsEditor/Private/HairStrandsEditor.cpp
@@ -38,6 +38,8 @@
 #include "UObject/UObjectIterator.h"
 #include "HairStrandsFactory.h"
 
+#include "Dataflow/AssetDefinition_DataflowAsset.h"
+
 IMPLEMENT_MODULE(FGroomEditor, HairStrandsEditor);
 
 #define LOCTEXT_NAMESPACE "GroomEditor"
@@ -158,10 +160,14 @@ void FGroomEditor::StartupModule()
 	UThumbnailManager::Get().RegisterCustomRenderer(UGroomBindingAsset::StaticClass(), UGroomBindingAssetThumbnailRenderer::StaticClass());
 
 	UE::Groom::RegisterGroomDataflowTemplatePath(TEXT("/HairStrands/Dataflow/Templates"));
+
+	DataflowMenusHandle = UE::DataflowAssetDefinitionHelpers::RegisterDataflowAssetMenus(UGroomAsset::StaticClass());
 }
 
 void FGroomEditor::ShutdownModule()
 {
+	UE::DataflowAssetDefinitionHelpers::UnregisterDataflowAssetMenus(DataflowMenusHandle);
+
 	UToolMenus::UnRegisterStartupCallback(this);
 	UToolMenus::UnregisterOwner(UE_MODULE_NAME);
 
diff --git a/Engine/Plugins/Runtime/HairStrands/Source/HairStrandsEditor/Public/HairStrandsEditor.h b/Engine/Plugins/Runtime/HairStrands/Source/HairStrandsEditor/Public/HairStrandsEditor.h
--- a/Engine/Plugins/Runtime/HairStrands/Source/HairStrandsEditor/Public/HairStrandsEditor.h
+++ b/Engine/Plugins/Runtime/HairStrands/Source/HairStrandsEditor/Public/HairStrandsEditor.h
@@ -62,6 +62,7 @@ private:
 	FDelegateHandle TrackEditorBindingHandle;
 	FDelegateHandle PreviewPlatformChangedHandle;
 	FDelegateHandle PreviewFeatureLevelChangedHandle;
+	FDelegateHandle DataflowMenusHandle;
 
 public:
 
