Dataflow : Add skeletal mesh dataflow attachment and corresponding menu commands to the Skeletal asset context menu

#Synced-CL 52497132

diff --git a/Engine/Plugins/Experimental/Dataflow/Source/DataflowEditor/Private/Dataflow/DataflowEditorModule.cpp b/Engine/Plugins/Experimental/Dataflow/Source/DataflowEditor/Private/Dataflow/DataflowEditorModule.cpp
--- a/Engine/Plugins/Experimental/Dataflow/Source/DataflowEditor/Private/Dataflow/DataflowEditorModule.cpp
+++ b/Engine/Plugins/Experimental/Dataflow/Source/DataflowEditor/Private/Dataflow/DataflowEditorModule.cpp
@@ -13,6 +13,7 @@
 #include "Dataflow/DataflowFunctionPropertyCustomization.h"
 #include "Dataflow/DataflowInstance.h"
 #include "Dataflow/DataflowSNodeFactories.h"
+#include "Dataflow/DataflowSkeletalMeshAttachment.h"
 #include "Dataflow/DebugDrawText3dVisualization.h"
 #include "Dataflow/MeshStatsConstructionVisualization.h"
 #include "Dataflow/ScalarVertexPropertyGroupCustomization.h"
@@ -51,6 +52,7 @@
 #include "DataflowRendering/DataflowPlaneRenderableType.h"
 #include "DataflowRendering/DataflowPhysicsAssetRenderableType.h"
 
+#include "Engine/SkeletalMesh.h"
 #include "ContentBrowserModule.h"
 #include "Modules/ModuleManager.h"
 #include "Styling/SlateIconFinder.h"
@@ -293,7 +295,7 @@ void FDataflowEditorModule::StartupModule()
 
 #if WITH_EDITOR
 	// register settings
-	if(ISettingsModule* SettingsModule = FModuleManager::GetModulePtr<ISettingsModule>("Settings"))
+	if (ISettingsModule* SettingsModule = FModuleManager::GetModulePtr<ISettingsModule>("Settings"))
 	{
 		SettingsModule->RegisterSettings("Editor", "ContentEditors", "DataflowEditor",
 			LOCTEXT("DataflowEditorLabel", "Dataflow Editor"),
@@ -305,7 +307,7 @@ void FDataflowEditorModule::StartupModule()
 
 
 	FDataflowEditorStyle::Get();
-	
+
 	// Register type customizations
 	if (FPropertyEditorModule* const PropertyModule = FModuleManager::GetModulePtr<FPropertyEditorModule>("PropertyEditor"))
 	{
@@ -324,31 +326,31 @@ void FDataflowEditorModule::StartupModule()
 
 	UE::Dataflow::FDataflowToolRegistry& ToolRegistry = UE::Dataflow::FDataflowToolRegistry::Get();
 	TSharedRef<const FDataflowToolActionCommandBindings> Actions = MakeShared<FDataflowToolActionCommandBindings>();
-	
+
 	ToolRegistry.AddNodeToToolMapping(FDataflowCollectionAddScalarVertexPropertyNode::StaticType(),
 		NewObject<UDataflowEditorVertexAttributePaintToolBuilder>(), Actions,
 		FSlateIcon(FName("DataflowEditorStyle"), FName("Dataflow.PaintWeightMap")), LOCTEXT("AddWeightMapNodeButtonText", "Paint Weight Map"));
 
-	ToolRegistry.AddNodeToToolMapping(FDataflowCollectionEditSkinWeightsNode::StaticType(), 
+	ToolRegistry.AddNodeToToolMapping(FDataflowCollectionEditSkinWeightsNode::StaticType(),
 		NewObject<UDataflowEditorSkinWeightsPaintToolBuilder>(), Actions,
 		FSlateIcon(FName("DataflowEditorStyle"), FName("Dataflow.EditSkinWeights")), LOCTEXT("AddSkinWeightNodeButtonText", "Edit Skin Weights"));
 
-	ToolRegistry.AddNodeToToolMapping(FDataflowCollectionEditSkeletonBonesNode::StaticType(), 
+	ToolRegistry.AddNodeToToolMapping(FDataflowCollectionEditSkeletonBonesNode::StaticType(),
 		NewObject<UDataflowEditorEditSkeletonBonesToolBuilder>(), Actions,
 		FSlateIcon(FName("DataflowEditorStyle"), FName("Dataflow.EditSkeletonBones")), LOCTEXT("AddSkeletonEditNodeButtonText", "Edit Skeleton Bones"));
 
 	ToolRegistry.AddNodeToToolMapping(FDataflowSelectionToolNode::StaticType(),
 		NewObject<UDataflowMeshSelectionToolBuilder>(), Actions,
 		FSlateIcon(FAppStyle::GetAppStyleSetName(), "LevelEditor.StrictBoxSelect"), LOCTEXT("AddSelectionNodeButtonText", "Selection"));
-		
+
 	DATAFLOW_NODE_REGISTER_CREATION_FACTORY(FDataflowCorrectSkinWeightsNode);
 	DATAFLOW_NODE_REGISTER_CREATION_FACTORY(FDataflowGetSkinningSelectionNode);
 	DATAFLOW_NODE_REGISTER_CREATION_FACTORY(FDataflowSetSkinningSelectionNode);
-	
+
 	UE::Dataflow::FDataflowConstructionVisualizationRegistry& ConstructionVisualizationRegistry = UE::Dataflow::FDataflowConstructionVisualizationRegistry::GetInstance();
 	ConstructionVisualizationRegistry.RegisterVisualization(MakeUnique<UE::Dataflow::FMeshStatsConstructionVisualization>());
 	ConstructionVisualizationRegistry.RegisterVisualization(MakeUnique<UE::Dataflow::FDebugDrawText3dVisualization>());
-	
+
 	FDataflowEditorCommands::Register();
 
 	// register rendering types
@@ -375,12 +377,17 @@ void FDataflowEditorModule::StartupModule()
 			FOnGenerateAssetViewExtraStateIndicators::CreateStatic(&UE::Dataflow::Private::OnGenerateDataflowAssetBindingTooltip)
 		));
 	}
+
+	FDataflowAttachmentFactory::Get().Register(USkeletalMesh::StaticClass()->GetFName(), [](UObject* Owner) { return NewObject<UDataflowSkeletalMeshAttachment>(Owner); });
+	SkeletalMeshDataflowMenusHandle = UE::DataflowAssetDefinitionHelpers::RegisterDataflowAssetMenus(USkeletalMesh::StaticClass());
 }
 
 void FDataflowEditorModule::ShutdownModule()
 {
 	using namespace UE::Dataflow::Private;
 
+	UE::DataflowAssetDefinitionHelpers::UnregisterDataflowAssetMenus(SkeletalMeshDataflowMenusHandle);
+
 	if (FContentBrowserModule* ContentBrowserModule = FModuleManager::Get().GetModulePtr<FContentBrowserModule>(TEXT("ContentBrowser")))
 	{
 		if (AssetViewExtraStateGeneratorHandle.IsValid())
@@ -424,5 +431,4 @@ void FDataflowEditorModule::ShutdownModule()
 
 IMPLEMENT_MODULE(FDataflowEditorModule, DataflowEditor)
 
-
 #undef LOCTEXT_NAMESPACE
diff --git a/Engine/Plugins/Experimental/Dataflow/Source/DataflowEditor/Public/Dataflow/DataflowEditorModule.h b/Engine/Plugins/Experimental/Dataflow/Source/DataflowEditor/Public/Dataflow/DataflowEditorModule.h
--- a/Engine/Plugins/Experimental/Dataflow/Source/DataflowEditor/Public/Dataflow/DataflowEditorModule.h
+++ b/Engine/Plugins/Experimental/Dataflow/Source/DataflowEditor/Public/Dataflow/DataflowEditorModule.h
@@ -27,6 +27,7 @@ private:
 
 	TSharedPtr<FDataflowSNodeFactory> DataflowSNodeFactory;
 	FDelegateHandle AssetViewExtraStateGeneratorHandle;
+	FDelegateHandle SkeletalMeshDataflowMenusHandle;
 
 };
 
diff --git a/Engine/Source/Runtime/Experimental/Dataflow/Engine/Private/Dataflow/DataflowSkeletalMeshAttachment.cpp b/Engine/Source/Runtime/Experimental/Dataflow/Engine/Private/Dataflow/DataflowSkeletalMeshAttachment.cpp
new file mode 100644
--- /dev/null
+++ b/Engine/Source/Runtime/Experimental/Dataflow/Engine/Private/Dataflow/DataflowSkeletalMeshAttachment.cpp
@@ -0,0 +1,37 @@
+// Copyright Epic Games, Inc. All Rights Reserved.
+
+#include "Dataflow/DataflowSkeletalMeshAttachment.h"
+#include "Dataflow/DataflowContent.h"
+#include "PhysicsEngine/PhysicsAsset.h"
+#include "Engine/SkeletalMesh.h"
+
+void UDataflowSkeletalMeshAttachment::WriteDataflowContent(const TObjectPtr<UDataflowBaseContent>& DataflowContent) const
+{
+	if (const TObjectPtr<UDataflowSkeletalContent> Content = Cast<UDataflowSkeletalContent>(DataflowContent))
+	{
+		Content->SetDataflowAsset(GetDataflowInstance().GetDataflowAsset());
+		Content->SetDataflowTerminal(GetDataflowInstance().GetDataflowTerminal().ToString());
+
+		if (USkeletalMesh* SkmAsset = GetTypedOuter<USkeletalMesh>())
+		{
+			Content->SetSkeletalMesh(SkmAsset);
+		}
+	}
+}
+
+void UDataflowSkeletalMeshAttachment::ReadDataflowContent(const TObjectPtr<UDataflowBaseContent>& DataflowContent)
+{
+
+}
+
+TObjectPtr<UDataflowBaseContent> UDataflowSkeletalMeshAttachment::CreateDataflowContent()
+{
+	TObjectPtr<UDataflowBaseContent> BaseContent = NewObject<UDataflowSkeletalContent>();
+
+	BaseContent->SetDataflowOwner(this);
+	BaseContent->SetTerminalAsset(this);
+
+	WriteDataflowContent(BaseContent);
+
+	return BaseContent;
+}
\ No newline at end of file
diff --git a/Engine/Source/Runtime/Experimental/Dataflow/Engine/Public/Dataflow/DataflowSkeletalMeshAttachment.h b/Engine/Source/Runtime/Experimental/Dataflow/Engine/Public/Dataflow/DataflowSkeletalMeshAttachment.h
new file mode 100644
--- /dev/null
+++ b/Engine/Source/Runtime/Experimental/Dataflow/Engine/Public/Dataflow/DataflowSkeletalMeshAttachment.h
@@ -0,0 +1,28 @@
+// Copyright Epic Games, Inc. All Rights Reserved.
+
+#pragma once
+
+#include "Dataflow/DataflowAttachment.h"
+#include "DataflowSkeletalMeshAttachment.generated.h"
+
+UCLASS()
+class UDataflowSkeletalMeshAttachment : public UDataflowAttachment
+{
+	GENERATED_BODY()
+
+public:
+
+	// IDataflowContentOwner
+	void WriteDataflowContent(const TObjectPtr<UDataflowBaseContent>& DataflowContent) const override;
+	void ReadDataflowContent(const TObjectPtr<UDataflowBaseContent>& DataflowContent) override;
+	////////////////////////////////////////////////////////////////////////////////////////////////////////////////
+
+protected:
+
+	// IDataflowContentOwner
+	TObjectPtr<UDataflowBaseContent> CreateDataflowContent() override;
+	////////////////////////////////////////////////////////////////////////////////////////////////////////////////
+
+private:
+};
+
