Dataflow : added an API to dataflow attachment to get the preview actor path for the Dataflow Editor
- Updated existing attachments implemention
- Also added a skeletal mesh preview actor blueprint

#virtualized
#Synced-CL 52497132

diff --git a/Engine/Plugins/Experimental/ChaosRigidAsset/Source/ChaosRigidAssetEngine/Private/PhysicsAssetDataflowAttachment.cpp b/Engine/Plugins/Experimental/ChaosRigidAsset/Source/ChaosRigidAssetEngine/Private/PhysicsAssetDataflowAttachment.cpp
--- a/Engine/Plugins/Experimental/ChaosRigidAsset/Source/ChaosRigidAssetEngine/Private/PhysicsAssetDataflowAttachment.cpp
+++ b/Engine/Plugins/Experimental/ChaosRigidAsset/Source/ChaosRigidAssetEngine/Private/PhysicsAssetDataflowAttachment.cpp
@@ -37,4 +37,8 @@ TObjectPtr<UDataflowBaseContent> UPhysicsAssetDataflowAttachment::CreateDataflow
 	return BaseContent;
 }
 
+FString UPhysicsAssetDataflowAttachment::GetPreviewActorPath() const
+{
+	return TEXT("/ChaosRigidAsset/BP_PhysicsAssetPreview.BP_PhysicsAssetPreview_C");
+}
 
diff --git a/Engine/Plugins/Experimental/ChaosRigidAsset/Source/ChaosRigidAssetEngine/Public/PhysicsAssetDataflowAttachment.h b/Engine/Plugins/Experimental/ChaosRigidAsset/Source/ChaosRigidAssetEngine/Public/PhysicsAssetDataflowAttachment.h
--- a/Engine/Plugins/Experimental/ChaosRigidAsset/Source/ChaosRigidAssetEngine/Public/PhysicsAssetDataflowAttachment.h
+++ b/Engine/Plugins/Experimental/ChaosRigidAsset/Source/ChaosRigidAssetEngine/Public/PhysicsAssetDataflowAttachment.h
@@ -23,5 +23,9 @@ protected:
 	TObjectPtr<UDataflowBaseContent> CreateDataflowContent() override;
 	////////////////////////////////////////////////////////////////////////////////////////////////////////////////
 
+	// UDataflowAttachment
+	FString GetPreviewActorPath() const override;
+	////////////////////////////////////////////////////////////////////////////////////////////////////////////////
+
 private:
 };
diff --git a/Engine/Plugins/Experimental/Dataflow/Source/DataflowEditor/Private/Dataflow/AssetDefinition_DataflowAsset.cpp b/Engine/Plugins/Experimental/Dataflow/Source/DataflowEditor/Private/Dataflow/AssetDefinition_DataflowAsset.cpp
--- a/Engine/Plugins/Experimental/Dataflow/Source/DataflowEditor/Private/Dataflow/AssetDefinition_DataflowAsset.cpp
+++ b/Engine/Plugins/Experimental/Dataflow/Source/DataflowEditor/Private/Dataflow/AssetDefinition_DataflowAsset.cpp
@@ -417,14 +417,13 @@ namespace UE::DataflowAssetDefinitionHelpers
 							if (UDataflowEditor* const AssetEditor = NewObject<UDataflowEditor>(AssetEditorSubsystem, NAME_None, RF_Transient))
 							{
 								// TODO(dataflow) : should probably read the setting from the attachment if available
-								const TSubclassOf<AActor> ActorClass; // = StaticLoadClass(AActor::StaticClass(), nullptr,
-									//TEXT("/ChaosClothAssetEditor/BP_ClothPreview.BP_ClothPreview_C"), nullptr, LOAD_None, nullptr);
-								AssetEditor->Initialize({ Attachment }, ActorClass);
+								const FString PreviewActorClassPath = Attachment->GetPreviewActorPath();
+								const TSubclassOf<AActor> PreviewActorClass = StaticLoadClass(AActor::StaticClass(), nullptr, PreviewActorClassPath, nullptr, LOAD_None, nullptr);
+								AssetEditor->Initialize({ Attachment }, PreviewActorClass);
 							}
 						}
 					}
 				}
-
 			}
 		}
 	}
diff --git a/Engine/Source/Runtime/Experimental/Dataflow/Engine/Private/Dataflow/DataflowAttachment.cpp b/Engine/Source/Runtime/Experimental/Dataflow/Engine/Private/Dataflow/DataflowAttachment.cpp
--- a/Engine/Source/Runtime/Experimental/Dataflow/Engine/Private/Dataflow/DataflowAttachment.cpp
+++ b/Engine/Source/Runtime/Experimental/Dataflow/Engine/Private/Dataflow/DataflowAttachment.cpp
@@ -19,6 +19,13 @@ void UDataflowAttachment::PostLoad()
 	Instance.PostLoad();
 }
 
+// Provides trhe path to the preview actor to be used in the corresponding Dataflow editor
+FString UDataflowAttachment::GetPreviewActorPath() const
+{
+	return FString();
+}
+
+
 //==============================================================================================
 
 /*static*/ FDataflowAttachmentFactory& FDataflowAttachmentFactory::Get()
@@ -44,3 +51,4 @@ UDataflowAttachment* FDataflowAttachmentFactory::FDataflowAttachmentFactory::Cre
 	}
 	return nullptr;
 }
+
diff --git a/Engine/Source/Runtime/Experimental/Dataflow/Engine/Private/Dataflow/DataflowSkeletalMeshAttachment.cpp b/Engine/Source/Runtime/Experimental/Dataflow/Engine/Private/Dataflow/DataflowSkeletalMeshAttachment.cpp
--- a/Engine/Source/Runtime/Experimental/Dataflow/Engine/Private/Dataflow/DataflowSkeletalMeshAttachment.cpp
+++ b/Engine/Source/Runtime/Experimental/Dataflow/Engine/Private/Dataflow/DataflowSkeletalMeshAttachment.cpp
@@ -34,4 +34,9 @@ TObjectPtr<UDataflowBaseContent> UDataflowSkeletalMeshAttachment::CreateDataflow
 	WriteDataflowContent(BaseContent);
 
 	return BaseContent;
+}
+
+FString UDataflowSkeletalMeshAttachment::GetPreviewActorPath() const
+{
+	return TEXT("/Dataflow/BP_SkeletalMeshPreview.BP_SkeletalMeshPreview_C");
 }
\ No newline at end of file
diff --git a/Engine/Source/Runtime/Experimental/Dataflow/Engine/Public/Dataflow/DataflowAttachment.h b/Engine/Source/Runtime/Experimental/Dataflow/Engine/Public/Dataflow/DataflowAttachment.h
--- a/Engine/Source/Runtime/Experimental/Dataflow/Engine/Public/Dataflow/DataflowAttachment.h
+++ b/Engine/Source/Runtime/Experimental/Dataflow/Engine/Public/Dataflow/DataflowAttachment.h
@@ -39,6 +39,10 @@ public:
 	// UObject
 	UE_API virtual void PostLoad() override;
 	////////////////////////////////////////////////////////////////////////////////////////////////////////////////
+
+	/** Provides trhe path to the preview actor to be used in the corresponding Dataflow editor */
+	UE_API virtual FString GetPreviewActorPath() const;
+
 private:
 
 	// Dataflow instance for the attachment when editing.
diff --git a/Engine/Source/Runtime/Experimental/Dataflow/Engine/Public/Dataflow/DataflowSkeletalMeshAttachment.h b/Engine/Source/Runtime/Experimental/Dataflow/Engine/Public/Dataflow/DataflowSkeletalMeshAttachment.h
--- a/Engine/Source/Runtime/Experimental/Dataflow/Engine/Public/Dataflow/DataflowSkeletalMeshAttachment.h
+++ b/Engine/Source/Runtime/Experimental/Dataflow/Engine/Public/Dataflow/DataflowSkeletalMeshAttachment.h
@@ -23,6 +23,10 @@ protected:
 	TObjectPtr<UDataflowBaseContent> CreateDataflowContent() override;
 	////////////////////////////////////////////////////////////////////////////////////////////////////////////////
 
+	// UDataflowAttachment
+	FString GetPreviewActorPath() const override;
+	////////////////////////////////////////////////////////////////////////////////////////////////////////////////
+
 private:
 };
 
