Dataflow : Expose render methods for aggregate geometry to the outside modules

#Synced-CL 52601091

diff --git a/Engine/Plugins/Experimental/Dataflow/Source/DataflowEditor/Private/DataflowRendering/DataflowPhysicsAssetRenderableType.cpp b/Engine/Plugins/Experimental/Dataflow/Source/DataflowEditor/Private/DataflowRendering/DataflowPhysicsAssetRenderableType.cpp
--- a/Engine/Plugins/Experimental/Dataflow/Source/DataflowEditor/Private/DataflowRendering/DataflowPhysicsAssetRenderableType.cpp
+++ b/Engine/Plugins/Experimental/Dataflow/Source/DataflowEditor/Private/DataflowRendering/DataflowPhysicsAssetRenderableType.cpp
@@ -3,7 +3,6 @@
 #include "DataflowRendering/DataflowPhysicsAssetRenderableType.h"
 
 #include "AnimationRuntime.h"
-#include "Components/DynamicMeshComponent.h"
 #include "Components/StaticMeshComponent.h"
 #include "Dataflow/DataflowEngineUtil.h"
 #include "Dataflow/DataflowRenderingViewMode.h"
@@ -12,9 +11,7 @@
 #include "DataflowRendering/DataflowRenderableTypeInstance.h"
 #include "DataflowRendering/DataflowRenderableTypeRegistry.h"
 #include "DataflowRendering/DataflowRenderableTypeUtils.h"
-#include "DynamicMesh/DynamicMesh3.h"
 #include "Engine/SkeletalMesh.h"
-#include "Generators/CapsuleGenerator.h"
 #include "Materials/MaterialInstanceDynamic.h"
 #include "PhysicsEngine/AggregateGeom.h"
 #include "PhysicsEngine/PhysicsAsset.h"
@@ -69,151 +66,6 @@ static UMaterialInstanceDynamic* CreateMaterialInstance(UObject* Owner, const UD
 	return MaterialInstance;
 }
 
-// Build a uniform capsule mesh centered at the origin, aligned along +Z.
-// Total height = Length + 2 * Radius.
-// FCapsuleGenerator places its cylinder from Z=0 to Z=SegmentLength;
-// we shift the vertices to center the mesh at Z=0.
-static UE::Geometry::FDynamicMesh3 BuildCapsuleMesh(float Radius, float Length)
-{
-	UE::Geometry::FCapsuleGenerator Generator;
-	Generator.Radius         = Radius;
-	Generator.SegmentLength  = Length;
-	Generator.NumHemisphereArcSteps = 6;
-	Generator.NumCircleSteps        = 12;
-	Generator.Generate();
-
-	UE::Geometry::FDynamicMesh3 Mesh(&Generator);
-
-	const double OffsetZ = -(double)Length * 0.5;
-	for (int32 VID : Mesh.VertexIndicesItr())
-	{
-		FVector3d V = Mesh.GetVertex(VID);
-		V.Z += OffsetZ;
-		Mesh.SetVertex(VID, V);
-	}
-	return Mesh;
-}
-
-// Build a tapered capsule mesh centered at the origin, aligned along +Z.
-// Radius0 = bottom end radius, Radius1 = top end radius.
-// Length = cylinder section length (total height = Length + Radius0 + Radius1).
-static UE::Geometry::FDynamicMesh3 BuildTaperedCapsuleMesh(float Radius0, float Radius1, float Length,
-	int32 ArcSteps = 6, int32 CircleSteps = 12)
-{
-	UE::Geometry::FDynamicMesh3 Mesh;
-	const double HalfLength = (double)Length * 0.5;
-	const double DTheta     = 2.0 * UE_DOUBLE_PI / CircleSteps;
-
-	// Profile: (radius, z) pairs bottom to top.
-	// Bottom hemisphere: phi sweeps 0 (pole) -> PI/2 (equator)
-	// Straight taper section: two equatorial rings already at ends of arcs
-	// Top hemisphere:    phi sweeps PI/2 (equator) -> 0 (pole, approached from PI/2)
-	TArray<TTuple<double, double>> Profile;
-
-	// Bottom pole
-	Profile.Add(MakeTuple(0.0, -HalfLength - (double)Radius0));
-	// Bottom hemisphere (excluding pole, including equator)
-	for (int32 i = 1; i <= ArcSteps; ++i)
-	{
-		const double Phi = (UE_DOUBLE_PI * 0.5) * (double)i / (double)ArcSteps;
-		Profile.Add(MakeTuple((double)Radius0 * FMath::Sin(Phi), -HalfLength - (double)Radius0 * FMath::Cos(Phi)));
-	}
-	// Top equator (separate ring for the taper)
-	Profile.Add(MakeTuple((double)Radius1, HalfLength));
-	// Top hemisphere (excluding equator, including pole)
-	for (int32 i = ArcSteps - 1; i >= 0; --i)
-	{
-		const double Phi = (UE_DOUBLE_PI * 0.5) * (double)i / (double)ArcSteps;
-		Profile.Add(MakeTuple((double)Radius1 * FMath::Sin(Phi), HalfLength + (double)Radius1 * FMath::Cos(Phi)));
-	}
-
-	// Revolve profile to produce vertex rings.
-	// Pole entries (r==0) produce a single vertex; all others produce a full ring.
-	const int32 N = Profile.Num();
-	TArray<TArray<int32>> Rings;
-	Rings.SetNum(N);
-
-	for (int32 PIdx = 0; PIdx < N; ++PIdx)
-	{
-		const double R = Profile[PIdx].Get<0>();
-		const double Z = Profile[PIdx].Get<1>();
-		if (FMath::IsNearlyZero(R))
-		{
-			Rings[PIdx].Add(Mesh.AppendVertex(FVector3d(0.0, 0.0, Z)));
-		}
-		else
-		{
-			for (int32 CIdx = 0; CIdx < CircleSteps; ++CIdx)
-			{
-				const double Theta = CIdx * DTheta;
-				Rings[PIdx].Add(Mesh.AppendVertex(FVector3d(R * FMath::Cos(Theta), R * FMath::Sin(Theta), Z)));
-			}
-		}
-	}
-
-	// Stitch adjacent rings with triangles.
-	for (int32 PIdx = 0; PIdx + 1 < N; ++PIdx)
-	{
-		const TArray<int32>& RingA = Rings[PIdx];
-		const TArray<int32>& RingB = Rings[PIdx + 1];
-
-		if (RingA.Num() == 1)
-		{
-			// Fan from pole down to ring
-			const int32 Pole = RingA[0];
-			for (int32 CIdx = 0; CIdx < CircleSteps; ++CIdx)
-			{
-				Mesh.AppendTriangle(Pole, RingB[CIdx], RingB[(CIdx + 1) % CircleSteps]);
-			}
-		}
-		else if (RingB.Num() == 1)
-		{
-			// Fan from ring up to pole
-			const int32 Pole = RingB[0];
-			for (int32 CIdx = 0; CIdx < CircleSteps; ++CIdx)
-			{
-				Mesh.AppendTriangle(RingA[CIdx], Pole, RingA[(CIdx + 1) % CircleSteps]);
-			}
-		}
-		else
-		{
-			// Quad strip
-			for (int32 CIdx = 0; CIdx < CircleSteps; ++CIdx)
-			{
-				const int32 Next = (CIdx + 1) % CircleSteps;
-				Mesh.AppendTriangle(RingA[CIdx],  RingB[CIdx],  RingB[Next]);
-				Mesh.AppendTriangle(RingA[CIdx],  RingB[Next],  RingA[Next]);
-			}
-		}
-	}
-
-	return Mesh;
-}
-
-// Build a mesh from a convex element's pre-computed triangle data.
-// Returns an empty mesh if the data is missing.
-static UE::Geometry::FDynamicMesh3 BuildConvexMesh(const FKConvexElem& Elem)
-{
-	UE::Geometry::FDynamicMesh3 Mesh;
-
-	if (Elem.VertexData.IsEmpty() || Elem.IndexData.Num() < 3)
-	{
-		return Mesh;
-	}
-
-	for (const FVector& V : Elem.VertexData)
-	{
-		Mesh.AppendVertex(FVector3d(V));
-	}
-
-	for (int32 Tri = 0; Tri + 2 < Elem.IndexData.Num(); Tri += 3)
-	{
-		Mesh.AppendTriangle(Elem.IndexData[Tri], Elem.IndexData[Tri + 1], Elem.IndexData[Tri + 2]);
-	}
-
-	return Mesh;
-}
-
 } // anonymous namespace
 
 ///////////////////////////////////////////////////////////////////////////////////////////////////////////////
@@ -272,110 +124,9 @@ class FPhysicsAssetSurfaceRenderableType : public IRenderableType
 				? PreviewMesh->GetRefSkeleton().GetBoneName(BoneIdx)
 				: UnknownBoneName;
 
-			const FKAggregateGeom& Geom = BodySetup->AggGeom;
-
-			// --- Spheres ---
-			for (int32 EIdx = 0; EIdx < Geom.SphereElems.Num(); ++EIdx, ++ShapeCounter)
-			{
-				const FKSphereElem& Elem = Geom.SphereElems[EIdx];
-				UStaticMesh* SphereMesh = UE::Dataflow::RenderGeometry::GetDataflowSphere();
-				if (!SphereMesh)
-				{
-					continue;
-				}
-
-				const FName SphereComponentName(FString::Printf(TEXT("[%d]_%s_(Sphere)"), ShapeCounter, *BoneName.ToString()));
-				if (UStaticMeshComponent* Component = OutComponents.AddNewComponent<UStaticMeshComponent>(SphereComponentName, ParentComponent))
-				{
-					Component->SetStaticMesh(SphereMesh);
-					Component->SetCollisionEnabled(ECollisionEnabled::NoCollision);
-
-					FTransform WorldTransform(FQuat::Identity, BoneCS.TransformPosition(FVector(Elem.Center)));
-					WorldTransform.SetScale3D(FVector(Elem.Radius / 50.f));
-					Component->SetWorldTransform(WorldTransform);
-					Component->SetMaterial(0, MaterialInstance);
-				}
-			}
-
-			// --- Boxes ---
-			for (int32 EIdx = 0; EIdx < Geom.BoxElems.Num(); ++EIdx, ++ShapeCounter)
-			{
-				const FKBoxElem& Elem = Geom.BoxElems[EIdx];
-				UStaticMesh* BoxMesh = UE::Dataflow::RenderGeometry::GetDataflowBox();
-				if (!BoxMesh)
-				{
-					continue;
-				}
-
-				const FName BoxComponentName(FString::Printf(TEXT("[%d]_%s_(Box)"), ShapeCounter, *BoneName.ToString()));
-				if (UStaticMeshComponent* Component = OutComponents.AddNewComponent<UStaticMeshComponent>(BoxComponentName, ParentComponent))
-				{
-					Component->SetStaticMesh(BoxMesh);
-					Component->SetCollisionEnabled(ECollisionEnabled::NoCollision);
-
-					// Elem.GetTransform() gives position+rotation in bone-local space.
-					// X/Y/Z are already full-extents; the unit box mesh has a extent of 100 n all directions.
-					FTransform WorldTransform = Elem.GetTransform() * BoneCS;
-					WorldTransform.SetScale3D(FVector(Elem.X, Elem.Y, Elem.Z) / 100.f);
-					Component->SetWorldTransform(WorldTransform);
-					Component->SetMaterial(0, MaterialInstance);
-				}
-			}
-
-			// --- Capsules (Sphyl) ---
-			for (int32 EIdx = 0; EIdx < Geom.SphylElems.Num(); ++EIdx, ++ShapeCounter)
-			{
-				const FKSphylElem& Elem = Geom.SphylElems[EIdx];
-
-				UE::Geometry::FDynamicMesh3 CapsuleMesh = BuildCapsuleMesh(Elem.Radius, Elem.Length);
-
-				const FName CapsuleComponentName(FString::Printf(TEXT("[%d]_%s_(Capsule)"), ShapeCounter, *BoneName.ToString()));
-				if (UDynamicMeshComponent* Component = OutComponents.AddNewComponent<UDynamicMeshComponent>(CapsuleComponentName, ParentComponent))
-				{
-					Component->SetMesh(MoveTemp(CapsuleMesh));
-					Component->SetCollisionEnabled(ECollisionEnabled::NoCollision);
-					Component->SetWorldTransform(Elem.GetTransform() * BoneCS);
-					Component->SetOverrideRenderMaterial(MaterialInstance);
-				}
-			}
-
-			// --- Tapered capsules ---
-			for (int32 EIdx = 0; EIdx < Geom.TaperedCapsuleElems.Num(); ++EIdx, ++ShapeCounter)
-			{
-				const FKTaperedCapsuleElem& Elem = Geom.TaperedCapsuleElems[EIdx];
-
-				UE::Geometry::FDynamicMesh3 TaperedMesh = BuildTaperedCapsuleMesh(Elem.Radius0, Elem.Radius1, Elem.Length);
-
-				const FName TaperedCapsuleComponentName(FString::Printf(TEXT("[%d]_%s_(Tapered_Capsule)"), ShapeCounter, *BoneName.ToString()));
-				if (UDynamicMeshComponent* Component = OutComponents.AddNewComponent<UDynamicMeshComponent>(TaperedCapsuleComponentName, ParentComponent))
-				{
-					Component->SetMesh(MoveTemp(TaperedMesh));
-					Component->SetCollisionEnabled(ECollisionEnabled::NoCollision);
-					Component->SetWorldTransform(Elem.GetTransform() * BoneCS);
-					Component->SetOverrideRenderMaterial(MaterialInstance);
-				}
-			}
-
-			// --- Convex hulls ---
-			for (int32 EIdx = 0; EIdx < Geom.ConvexElems.Num(); ++EIdx, ++ShapeCounter)
-			{
-				const FKConvexElem& Elem = Geom.ConvexElems[EIdx];
-
-				UE::Geometry::FDynamicMesh3 ConvexMesh = BuildConvexMesh(Elem);
-				if (ConvexMesh.TriangleCount() == 0)
-				{
-					continue;
-				}
-
-				const FName ConvexComponentName(FString::Printf(TEXT("[%d]_%s_(Convex)"), ShapeCounter, *BoneName.ToString()));
-				if (UDynamicMeshComponent* Component = OutComponents.AddNewComponent<UDynamicMeshComponent>(ConvexComponentName, ParentComponent))
-				{
-					Component->SetMesh(MoveTemp(ConvexMesh));
-					Component->SetCollisionEnabled(ECollisionEnabled::NoCollision);
-					Component->SetWorldTransform(Elem.GetTransform() * BoneCS);
-					Component->SetOverrideRenderMaterial(MaterialInstance);
-				}
-			}
+			UE::Dataflow::RenderGeometry::AddAggGeomComponents(
+				BodySetup->AggGeom, BoneCS, BoneName,
+				OutComponents, ParentComponent, MaterialInstance, ShapeCounter);
 		}
 	}
 };
diff --git a/Engine/Plugins/Experimental/Dataflow/Source/DataflowEditor/Private/DataflowRendering/DataflowRenderableTypeUtils.cpp b/Engine/Plugins/Experimental/Dataflow/Source/DataflowEditor/Private/DataflowRendering/DataflowRenderableTypeUtils.cpp
--- a/Engine/Plugins/Experimental/Dataflow/Source/DataflowEditor/Private/DataflowRendering/DataflowRenderableTypeUtils.cpp
+++ b/Engine/Plugins/Experimental/Dataflow/Source/DataflowEditor/Private/DataflowRendering/DataflowRenderableTypeUtils.cpp
@@ -6,9 +6,18 @@
 #include "Dataflow/DataflowEngineUtil.h"
 
 #include "Components/DynamicMeshComponent.h"
+#include "Components/StaticMeshComponent.h"
 #include "Engine/SkeletalMesh.h"
 #include "Engine/StaticMesh.h"
+#include "Generators/CapsuleGenerator.h"
 #include "Materials/Material.h"
+#include "Materials/MaterialInstanceDynamic.h"
+#include "PhysicsEngine/AggregateGeom.h"
+#include "PhysicsEngine/BoxElem.h"
+#include "PhysicsEngine/ConvexElem.h"
+#include "PhysicsEngine/SphereElem.h"
+#include "PhysicsEngine/SphylElem.h"
+#include "PhysicsEngine/TaperedCapsuleElem.h"
 
 #include "GeometryCollection/Facades/CollectionMeshFacade.h"
 #include "GeometryCollection/GeometryCollectionUtility.h"
@@ -303,3 +312,241 @@ namespace UE::Dataflow::Rendering
 		return Color;
 	}
 }
+
+///////////////////////////////////////////////////////////////////////////////////////////////////////////////
+
+namespace UE::Dataflow::RenderGeometry
+{
+
+UE::Geometry::FDynamicMesh3 BuildCapsuleMesh(float Radius, float Length)
+{
+	UE::Geometry::FCapsuleGenerator Generator;
+	Generator.Radius            = Radius;
+	Generator.SegmentLength     = Length;
+	Generator.NumHemisphereArcSteps = 6;
+	Generator.NumCircleSteps        = 12;
+	Generator.Generate();
+
+	UE::Geometry::FDynamicMesh3 Mesh(&Generator);
+
+	// FCapsuleGenerator places the cylinder from Z=0; shift down to center the mesh at Z=0.
+	const double OffsetZ = -(double)Length * 0.5;
+	for (int32 VID : Mesh.VertexIndicesItr())
+	{
+		FVector3d V = Mesh.GetVertex(VID);
+		V.Z += OffsetZ;
+		Mesh.SetVertex(VID, V);
+	}
+	return Mesh;
+}
+
+///////////////////////////////////////////////////////////////////////////////////////////////////////////////
+
+UE::Geometry::FDynamicMesh3 BuildTaperedCapsuleMesh(float Radius0, float Radius1, float Length, int32 ArcSteps, int32 CircleSteps)
+{
+	UE::Geometry::FDynamicMesh3 Mesh;
+	const double HalfLength = (double)Length * 0.5;
+	const double DTheta      = 2.0 * UE_DOUBLE_PI / CircleSteps;
+
+	// Profile: (radius, z) pairs from bottom pole to top pole.
+	TArray<TTuple<double, double>> Profile;
+
+	// Bottom pole
+	Profile.Add(MakeTuple(0.0, -HalfLength - (double)Radius0));
+	// Bottom hemisphere (pole excluded, equator included)
+	for (int32 i = 1; i <= ArcSteps; ++i)
+	{
+		const double Phi = (UE_DOUBLE_PI * 0.5) * (double)i / (double)ArcSteps;
+		Profile.Add(MakeTuple((double)Radius0 * FMath::Sin(Phi), -HalfLength - (double)Radius0 * FMath::Cos(Phi)));
+	}
+	// Top equator (separate ring to allow the taper)
+	Profile.Add(MakeTuple((double)Radius1, HalfLength));
+	// Top hemisphere (equator excluded, pole included)
+	for (int32 i = ArcSteps - 1; i >= 0; --i)
+	{
+		const double Phi = (UE_DOUBLE_PI * 0.5) * (double)i / (double)ArcSteps;
+		Profile.Add(MakeTuple((double)Radius1 * FMath::Sin(Phi), HalfLength + (double)Radius1 * FMath::Cos(Phi)));
+	}
+
+	// Revolve profile to produce vertex rings.
+	const int32 N = Profile.Num();
+	TArray<TArray<int32>> Rings;
+	Rings.SetNum(N);
+
+	for (int32 PIdx = 0; PIdx < N; ++PIdx)
+	{
+		const double R = Profile[PIdx].Get<0>();
+		const double Z = Profile[PIdx].Get<1>();
+		if (FMath::IsNearlyZero(R))
+		{
+			Rings[PIdx].Add(Mesh.AppendVertex(FVector3d(0.0, 0.0, Z)));
+		}
+		else
+		{
+			for (int32 CIdx = 0; CIdx < CircleSteps; ++CIdx)
+			{
+				const double Theta = CIdx * DTheta;
+				Rings[PIdx].Add(Mesh.AppendVertex(FVector3d(R * FMath::Cos(Theta), R * FMath::Sin(Theta), Z)));
+			}
+		}
+	}
+
+	// Stitch adjacent rings with triangles.
+	for (int32 PIdx = 0; PIdx + 1 < N; ++PIdx)
+	{
+		const TArray<int32>& RingA = Rings[PIdx];
+		const TArray<int32>& RingB = Rings[PIdx + 1];
+
+		if (RingA.Num() == 1)
+		{
+			const int32 Pole = RingA[0];
+			for (int32 CIdx = 0; CIdx < CircleSteps; ++CIdx)
+			{
+				Mesh.AppendTriangle(Pole, RingB[CIdx], RingB[(CIdx + 1) % CircleSteps]);
+			}
+		}
+		else if (RingB.Num() == 1)
+		{
+			const int32 Pole = RingB[0];
+			for (int32 CIdx = 0; CIdx < CircleSteps; ++CIdx)
+			{
+				Mesh.AppendTriangle(RingA[CIdx], Pole, RingA[(CIdx + 1) % CircleSteps]);
+			}
+		}
+		else
+		{
+			for (int32 CIdx = 0; CIdx < CircleSteps; ++CIdx)
+			{
+				const int32 Next = (CIdx + 1) % CircleSteps;
+				Mesh.AppendTriangle(RingA[CIdx], RingB[CIdx],  RingB[Next]);
+				Mesh.AppendTriangle(RingA[CIdx], RingB[Next],  RingA[Next]);
+			}
+		}
+	}
+
+	return Mesh;
+}
+
+///////////////////////////////////////////////////////////////////////////////////////////////////////////////
+
+UE::Geometry::FDynamicMesh3 BuildConvexMesh(const FKConvexElem& Elem)
+{
+	UE::Geometry::FDynamicMesh3 Mesh;
+
+	if (Elem.VertexData.IsEmpty() || Elem.IndexData.Num() < 3)
+	{
+		return Mesh;
+	}
+
+	for (const FVector& V : Elem.VertexData)
+	{
+		Mesh.AppendVertex(FVector3d(V));
+	}
+
+	for (int32 Tri = 0; Tri + 2 < Elem.IndexData.Num(); Tri += 3)
+	{
+		Mesh.AppendTriangle(Elem.IndexData[Tri], Elem.IndexData[Tri + 1], Elem.IndexData[Tri + 2]);
+	}
+
+	return Mesh;
+}
+
+///////////////////////////////////////////////////////////////////////////////////////////////////////////////
+
+void AddAggGeomComponents(
+	const FKAggregateGeom& AggGeom,
+	const FTransform& BoneTransform,
+	FName BoneName,
+	FRenderableComponents& OutComponents,
+	UObject* ComponentParent,
+	UMaterialInstanceDynamic* MaterialInstance,
+	int32& InOutShapeCounter)
+{
+	auto MakeName = [&](const TCHAR* ShapeType) -> FName
+	{
+		return BoneName.IsNone()
+			? FName(FString::Printf(TEXT("[%d]_(%s)"), InOutShapeCounter, ShapeType))
+			: FName(FString::Printf(TEXT("[%d]_%s_(%s)"), InOutShapeCounter, *BoneName.ToString(), ShapeType));
+	};
+
+	// --- Spheres ---
+	for (const FKSphereElem& Elem : AggGeom.SphereElems)
+	{
+		if (UStaticMesh* SphereMesh = GetDataflowSphere())
+		{
+			if (UStaticMeshComponent* Comp = OutComponents.AddNewComponent<UStaticMeshComponent>(MakeName(TEXT("Sphere")), ComponentParent))
+			{
+				Comp->SetStaticMesh(SphereMesh);
+				Comp->SetCollisionEnabled(ECollisionEnabled::NoCollision);
+				FTransform T(FQuat::Identity, BoneTransform.TransformPosition(FVector(Elem.Center)));
+				T.SetScale3D(FVector(Elem.Radius / 50.f));
+				Comp->SetWorldTransform(T);
+				Comp->SetMaterial(0, MaterialInstance);
+			}
+		}
+		++InOutShapeCounter;
+	}
+
+	// --- Boxes ---
+	for (const FKBoxElem& Elem : AggGeom.BoxElems)
+	{
+		if (UStaticMesh* BoxMesh = GetDataflowBox())
+		{
+			if (UStaticMeshComponent* Comp = OutComponents.AddNewComponent<UStaticMeshComponent>(MakeName(TEXT("Box")), ComponentParent))
+			{
+				Comp->SetStaticMesh(BoxMesh);
+				Comp->SetCollisionEnabled(ECollisionEnabled::NoCollision);
+				FTransform T = Elem.GetTransform() * BoneTransform;
+				T.SetScale3D(FVector(Elem.X, Elem.Y, Elem.Z) / 100.f);
+				Comp->SetWorldTransform(T);
+				Comp->SetMaterial(0, MaterialInstance);
+			}
+		}
+		++InOutShapeCounter;
+	}
+
+	// --- Capsules (Sphyl) ---
+	for (const FKSphylElem& Elem : AggGeom.SphylElems)
+	{
+		if (UDynamicMeshComponent* Comp = OutComponents.AddNewComponent<UDynamicMeshComponent>(MakeName(TEXT("Capsule")), ComponentParent))
+		{
+			Comp->SetMesh(BuildCapsuleMesh(Elem.Radius, Elem.Length));
+			Comp->SetCollisionEnabled(ECollisionEnabled::NoCollision);
+			Comp->SetWorldTransform(Elem.GetTransform() * BoneTransform);
+			Comp->SetOverrideRenderMaterial(MaterialInstance);
+		}
+		++InOutShapeCounter;
+	}
+
+	// --- Tapered Capsules ---
+	for (const FKTaperedCapsuleElem& Elem : AggGeom.TaperedCapsuleElems)
+	{
+		if (UDynamicMeshComponent* Comp = OutComponents.AddNewComponent<UDynamicMeshComponent>(MakeName(TEXT("TaperedCapsule")), ComponentParent))
+		{
+			Comp->SetMesh(BuildTaperedCapsuleMesh(Elem.Radius0, Elem.Radius1, Elem.Length));
+			Comp->SetCollisionEnabled(ECollisionEnabled::NoCollision);
+			Comp->SetWorldTransform(Elem.GetTransform() * BoneTransform);
+			Comp->SetOverrideRenderMaterial(MaterialInstance);
+		}
+		++InOutShapeCounter;
+	}
+
+	// --- Convex Hulls ---
+	for (const FKConvexElem& Elem : AggGeom.ConvexElems)
+	{
+		UE::Geometry::FDynamicMesh3 ConvexMesh = BuildConvexMesh(Elem);
+		if (ConvexMesh.TriangleCount() > 0)
+		{
+			if (UDynamicMeshComponent* Comp = OutComponents.AddNewComponent<UDynamicMeshComponent>(MakeName(TEXT("Convex")), ComponentParent))
+			{
+				Comp->SetMesh(MoveTemp(ConvexMesh));
+				Comp->SetCollisionEnabled(ECollisionEnabled::NoCollision);
+				Comp->SetWorldTransform(Elem.GetTransform() * BoneTransform);
+				Comp->SetOverrideRenderMaterial(MaterialInstance);
+			}
+		}
+		++InOutShapeCounter;
+	}
+}
+
+} // namespace UE::Dataflow::RenderGeometry
diff --git a/Engine/Plugins/Experimental/Dataflow/Source/DataflowEditor/Public/DataflowRendering/DataflowRenderableTypeUtils.h b/Engine/Plugins/Experimental/Dataflow/Source/DataflowEditor/Public/DataflowRendering/DataflowRenderableTypeUtils.h
--- a/Engine/Plugins/Experimental/Dataflow/Source/DataflowEditor/Public/DataflowRendering/DataflowRenderableTypeUtils.h
+++ b/Engine/Plugins/Experimental/Dataflow/Source/DataflowEditor/Public/DataflowRendering/DataflowRenderableTypeUtils.h
@@ -4,16 +4,20 @@
 
 #include "Templates/SharedPointer.h"
 #include "Curves/LinearColorRamp.h"
+#include "DynamicMesh/DynamicMesh3.h"
 
 #include "DataflowRenderableTypeUtils.generated.h"
 
 struct FManagedArrayCollection;
 struct FMeshDescription;
+struct FKConvexElem;
+struct FKAggregateGeom;
 
 class UDynamicMeshComponent;
 class USkeletalMesh;
 class UStaticMesh;
 class UMaterial;
+class UMaterialInstanceDynamic;
 class UDataflowBoxArrayRenderSettings;
 
 #define UE_API DATAFLOWEDITOR_API
@@ -93,11 +97,6 @@ public:
 	int32 PermutationIndex = 0;
 };
 
-namespace UE::Geometry
-{
-	class FDynamicMesh3;
-}
-
 namespace UE::Dataflow
 {
 	struct FRenderableComponents;
@@ -147,6 +146,33 @@ namespace UE::Dataflow
 		UE_API FLinearColor GetColor(const int32 InIdx, const FBox& InComponentBoundingBox, const FVector& InCenter, float InSize,
 			const FDataflowColorCommonRenderSettings& InSettings);
 	}
+
+	namespace RenderGeometry
+	{
+		/** Build a uniform capsule mesh centered at origin, aligned along +Z. Total height = Length + 2 * Radius. */
+		UE_API UE::Geometry::FDynamicMesh3 BuildCapsuleMesh(float Radius, float Length);
+
+		/** Build a tapered capsule mesh centered at origin, aligned along +Z. */
+		UE_API UE::Geometry::FDynamicMesh3 BuildTaperedCapsuleMesh(float Radius0, float Radius1, float Length, int32 ArcSteps = 6, int32 CircleSteps = 12);
+
+		/** Build a triangle mesh from a convex element's pre-computed vertex/index data. Returns an empty mesh if data is missing. */
+		UE_API UE::Geometry::FDynamicMesh3 BuildConvexMesh(const FKConvexElem& Elem);
+
+		/**
+		 * Add UPrimitiveComponent instances for all shape elements in AggGeom to OutComponents.
+		 * BoneTransform: component-space transform of the owning bone (pass Identity for standalone geometry).
+		 * BoneName: used only for component naming (pass NAME_None for standalone geometry).
+		 * InOutShapeCounter: incremented once per shape added; share a single counter across calls for unique names.
+		 */
+		UE_API void AddAggGeomComponents(
+			const FKAggregateGeom& AggGeom,
+			const FTransform& BoneTransform,
+			FName BoneName,
+			FRenderableComponents& OutComponents,
+			UObject* ComponentParent,
+			UMaterialInstanceDynamic* MaterialInstance,
+			int32& InOutShapeCounter);
+	}
 }
 
 #undef UE_API
