EditorUsability : PreviewScene
- Fixed localization key and value warning (Unreal Editor Localization 1299109)

#Synced-CL 52465190

diff --git a/Engine/Plugins/Experimental/Dataflow/Source/DataflowEditor/Private/Dataflow/DataflowAdvancedPreviewDetailsTab.cpp b/Engine/Plugins/Experimental/Dataflow/Source/DataflowEditor/Private/Dataflow/DataflowAdvancedPreviewDetailsTab.cpp
--- a/Engine/Plugins/Experimental/Dataflow/Source/DataflowEditor/Private/Dataflow/DataflowAdvancedPreviewDetailsTab.cpp
+++ b/Engine/Plugins/Experimental/Dataflow/Source/DataflowEditor/Private/Dataflow/DataflowAdvancedPreviewDetailsTab.cpp
@@ -245,11 +245,15 @@ FReply SDataflowAdvancedPreviewDetailsTab::RemoveOrResetProfileButtonClick()
 {
 	const FPreviewSceneProfile& CurrentProfile = DefaultSettings->Profiles[ProfileIndex];
 	
+	const FScopedTransaction Transaction(
+		CurrentProfile.bIsEngineDefaultProfile 
+			? FText::Format(LOCTEXT("ResetSceneProfile", "Reset Preview Scene Profile ({0})"), FText::FromString(CurrentProfile.ProfileName))
+			: FText::Format(LOCTEXT("RemoveSceneProfile", "Remove Preview Scene Profile ({0})"), FText::FromString(CurrentProfile.ProfileName))
+	);
+	DefaultSettings->Modify();
+	
 	if (CurrentProfile.bIsEngineDefaultProfile)
 	{
-		const FScopedTransaction Transaction(LOCTEXT("ResetSceneProfile", "Reset Preview Scene Profile"));
-		DefaultSettings->Modify();
-
 		const FPreviewSceneProfile* DefaultEditorProfile = GetMutableDefault<UDefaultEditorProfiles>()->GetProfile(CurrentProfile.ProfileName);
 		if (DefaultEditorProfile)
 		{
@@ -263,9 +267,6 @@ FReply SDataflowAdvancedPreviewDetailsTab::RemoveOrResetProfileButtonClick()
 		// in which case we should remove it...
 	}
 
-	const FScopedTransaction Transaction(LOCTEXT("RemoveSceneProfile", "Remove Preview Scene Profile"));
-	DefaultSettings->Modify();
-
 	// Remove currently selected profile 
 	DefaultSettings->Profiles.RemoveAt(ProfileIndex);
 	ProfileIndex = DefaultSettings->Profiles.IsValidIndex(ProfileIndex - 1 ) ? ProfileIndex - 1 : 0;
