Dataflow timeline : make sure scrubbing the animation get clamped by the bounds of the range

#Synced-CL 52663992

diff --git a/Engine/Plugins/Experimental/Dataflow/Source/DataflowEditor/Private/Widgets/DataflowTimeSliderController.cpp b/Engine/Plugins/Experimental/Dataflow/Source/DataflowEditor/Private/Widgets/DataflowTimeSliderController.cpp
--- a/Engine/Plugins/Experimental/Dataflow/Source/DataflowEditor/Private/Widgets/DataflowTimeSliderController.cpp
+++ b/Engine/Plugins/Experimental/Dataflow/Source/DataflowEditor/Private/Widgets/DataflowTimeSliderController.cpp
@@ -669,6 +669,15 @@ FFrameTime FDataflowTimeSliderController::GetFrameTimeFromMouse(const FGeometry&
 
 void FDataflowTimeSliderController::CommitScrubPosition(FFrameTime NewValue, bool bIsScrubbing)
 {
+	// Clamp the scrub position to the playback range so the user can't scrub or click past the start/end
+	const TRange<FFrameNumber> PlaybackRange = TimeSliderArgs.PlaybackRange.Get();
+	if (PlaybackRange.HasLowerBound() && PlaybackRange.HasUpperBound())
+	{
+		const FFrameTime LowerBound = PlaybackRange.GetLowerBoundValue();
+		const FFrameTime UpperBound = PlaybackRange.GetUpperBoundValue();
+		NewValue = FMath::Clamp(NewValue, LowerBound, UpperBound);
+	}
+
 	// Manage the scrub position ourselves if its not bound to a delegate
 	if (!TimeSliderArgs.ScrubPosition.IsBound())
 	{
