DataflowEditor: Do not force-enable simulation when recreating the preview scene. CreateSimulationScene unconditionally called SetSimulationEnabled which auto-started flesh simulation whenever any scene description property changed (e.g. setting CacheAsset). The simulation enabled state should only be set explicitly by the constructor (initial load via bIsSimulationPlayingByDefault), user actions (play/pause/step buttons), or cache recording. Remove the force-enable so that property changes preserve whatever simulation state the user had.

#Synced-CL 52671974

diff --git a/Engine/Plugins/Experimental/Dataflow/Source/DataflowEditor/Private/Dataflow/DataflowSimulationScene.cpp b/Engine/Plugins/Experimental/Dataflow/Source/DataflowEditor/Private/Dataflow/DataflowSimulationScene.cpp
--- a/Engine/Plugins/Experimental/Dataflow/Source/DataflowEditor/Private/Dataflow/DataflowSimulationScene.cpp
+++ b/Engine/Plugins/Experimental/Dataflow/Source/DataflowEditor/Private/Dataflow/DataflowSimulationScene.cpp
@@ -370,11 +370,14 @@ void FDataflowSimulationScene::CreateSimulationScene()
 			AddSceneObject(PreviewActor, false);
 		}
 		AddSceneObject(RootSceneActor, false);
-		
-		// Enable the simulation by default
-		GetWorld()->GetSubsystem<UDataflowSimulationManager>()->SetSimulationEnabled(!IsSimulationLocked());
 
-		// Pause the animation 
+		// Note: Do not force-enable simulation here. The simulation enabled state is managed by:
+		// - The constructor (applies bIsSimulationPlayingByDefault after CreateSimulationScene)
+		// - User actions (play/pause/step buttons, cache recording)
+		// Force-enabling here would auto-start simulation on every property change (e.g. setting CacheAsset),
+		// overriding the user's current play/pause state.
+
+		// Pause the animation
 		UE::Dataflow::PauseSkeletonAnimation(PreviewActor);
 	}
 
