Dataflow : add a method to the view mode to inform about gizmo support

#Synced-CL 52671010

diff --git a/Engine/Plugins/ChaosClothAssetDataflowNodes/Source/ChaosClothAssetDataflowNodes/Private/ChaosClothAsset/ClothDataflowViewModes.h b/Engine/Plugins/ChaosClothAssetDataflowNodes/Source/ChaosClothAssetDataflowNodes/Private/ChaosClothAsset/ClothDataflowViewModes.h
--- a/Engine/Plugins/ChaosClothAssetDataflowNodes/Source/ChaosClothAssetDataflowNodes/Private/ChaosClothAsset/ClothDataflowViewModes.h
+++ b/Engine/Plugins/ChaosClothAssetDataflowNodes/Source/ChaosClothAssetDataflowNodes/Private/ChaosClothAsset/ClothDataflowViewModes.h
@@ -27,6 +27,7 @@ namespace UE::Chaos::ClothAsset
 		virtual FName GetName() const override;
 		virtual FText GetButtonText() const override;
 		virtual FText GetTooltipText() const override;
+		virtual bool CanDisplayGizmo() const { return true; }
 	};
 
 	class FClothRenderViewMode : public UE::Dataflow::FDataflowConstruction3DViewMode
@@ -38,6 +39,7 @@ namespace UE::Chaos::ClothAsset
 		virtual FName GetName() const override;
 		virtual FText GetButtonText() const override;
 		virtual FText GetTooltipText() const override;
+		virtual bool CanDisplayGizmo() const { return true; }
 	};
 
 }  // End namespace UE::Chaos::ClothAsset
diff --git a/Engine/Plugins/Experimental/Dataflow/Source/DataflowEditor/Public/Dataflow/DataflowRenderingViewMode.h b/Engine/Plugins/Experimental/Dataflow/Source/DataflowEditor/Public/Dataflow/DataflowRenderingViewMode.h
--- a/Engine/Plugins/Experimental/Dataflow/Source/DataflowEditor/Public/Dataflow/DataflowRenderingViewMode.h
+++ b/Engine/Plugins/Experimental/Dataflow/Source/DataflowEditor/Public/Dataflow/DataflowRenderingViewMode.h
@@ -36,6 +36,7 @@ namespace UE::Dataflow
 		virtual FText GetTooltipText() const = 0;
 		virtual ELevelViewportType GetViewportType() const = 0;
 		virtual FRotator GetViewportOrientation() const { return FRotator::ZeroRotator; }
+		virtual bool CanDisplayGizmo() const { return false; }
 		
 		DATAFLOWEDITOR_API bool IsPerspective() const;
 	};
@@ -78,6 +79,7 @@ namespace UE::Dataflow
 		DATAFLOWEDITOR_API virtual FText GetButtonText() const override;
 		DATAFLOWEDITOR_API virtual FText GetTooltipText() const override;
 		DATAFLOWEDITOR_API virtual FRotator GetViewportOrientation() const override;
+		DATAFLOWEDITOR_API virtual bool CanDisplayGizmo() const { return true; }
 	};
 
 	// UV view mode (same as 2D but with a different name and button text)
