Dataflow : Add a new asset context menu to create derived asset ( from a asset referenced by a dataflow )

#Synced-CL 52719833

diff --git a/Engine/Plugins/Experimental/Dataflow/Source/DataflowEditor/Private/Dataflow/AssetDefinition_DataflowAsset.cpp b/Engine/Plugins/Experimental/Dataflow/Source/DataflowEditor/Private/Dataflow/AssetDefinition_DataflowAsset.cpp
--- a/Engine/Plugins/Experimental/Dataflow/Source/DataflowEditor/Private/Dataflow/AssetDefinition_DataflowAsset.cpp
+++ b/Engine/Plugins/Experimental/Dataflow/Source/DataflowEditor/Private/Dataflow/AssetDefinition_DataflowAsset.cpp
@@ -448,6 +448,74 @@ namespace UE::DataflowAssetDefinitionHelpers
 
 	//////////////////////////////////////////////////////////////////////////////////////////////////////////////////////////////////////
 
+	static void ExecuteCreateDerivedAsset(const TArray<FAssetData> Assets)
+	{
+		if (Assets.Num() == 1)
+		{
+			if (UObject* SelectedAsset = Assets[0].GetAsset())
+			{
+				if (UDataflow* DataflowAsset = UE::Dataflow::InstanceUtils::GetDataflowAssetFromObject(SelectedAsset))
+				{
+					const UClass* const AssetClass = SelectedAsset->GetClass();
+					FSaveAssetDialogConfig NewDataflowAssetDialogConfig;
+					{
+						const FString PackageName = SelectedAsset->GetOutermost()->GetName();
+						NewDataflowAssetDialogConfig.DefaultPath = FPackageName::GetLongPackagePath(PackageName);
+						const FString AssetName = SelectedAsset->GetName();
+						NewDataflowAssetDialogConfig.DefaultAssetName = AssetName + "_Derived";
+						NewDataflowAssetDialogConfig.AssetClassNames.Add(AssetClass->GetClassPathName());
+						NewDataflowAssetDialogConfig.ExistingAssetPolicy = ESaveAssetDialogExistingAssetPolicy::Disallow;
+						NewDataflowAssetDialogConfig.DialogTitleOverride = LOCTEXT("NewDataflowAssetDialogTitle", "Save Derived Asset As");
+					}
+
+					FContentBrowserModule& ContentBrowserModule = FModuleManager::LoadModuleChecked<FContentBrowserModule>("ContentBrowser");
+
+					FString NewPackageName;
+					FText OutError;
+					for (bool bFilenameValid = false; !bFilenameValid; bFilenameValid = FFileHelper::IsFilenameValidForSaving(NewPackageName, OutError))
+					{
+						const FString AssetSavePath = ContentBrowserModule.Get().CreateModalSaveAssetDialog(NewDataflowAssetDialogConfig);
+						if (AssetSavePath.IsEmpty())
+						{
+							return;
+						}
+						NewPackageName = FPackageName::ObjectPathToPackageName(AssetSavePath);
+					}
+
+					const FName NewAssetName(FPackageName::GetLongPackageAssetName(NewPackageName));
+					UPackage* const NewPackage = CreatePackage(*NewPackageName);
+
+					UObject* const NewAsset = DuplicateObject<UObject>(SelectedAsset, NewPackage, NewAssetName);
+					if (NewAsset)
+					{
+						NewAsset->SetFlags(RF_Public | RF_Standalone);
+						NewAsset->MarkPackageDirty();
+
+						// Notify the asset registry
+						FAssetRegistryModule::AssetCreated(NewAsset);
+					}
+				}
+			}
+		}
+	}
+
+	static bool CanCreateDerivedAsset(const TArray<FAssetData> Assets)
+	{
+		if (Assets.Num() == 1)
+		{
+			if (UObject* SelectedAsset = Assets[0].GetAsset())
+			{
+				if (UDataflow* DataflowAsset = UE::Dataflow::InstanceUtils::GetDataflowAssetFromObject(SelectedAsset))
+				{
+					return (DataflowAsset->IsAsset() && DataflowAsset->ReferenceAsset == SelectedAsset);
+				}
+			}
+		}
+		return false;
+	}
+
+	//////////////////////////////////////////////////////////////////////////////////////////////////////////////////////////////////////
+
 	namespace Private
 	{
 		static void FillMenu(FMenuBuilder& MenuBuilder, const TArray<FAssetData> Assets, UClass* AssetClass)
@@ -467,6 +535,17 @@ namespace UE::DataflowAssetDefinitionHelpers
 					MenuBuilder.AddMenuEntry(Params);
 				}
 
+				// Create Derived Asset
+				{
+					FMenuEntryParams Params;
+					Params.DirectActions.ExecuteAction = FExecuteAction::CreateStatic(&ExecuteCreateDerivedAsset, Assets);
+					Params.DirectActions.IsActionVisibleDelegate = FCanExecuteAction::CreateStatic(&CanCreateDerivedAsset, Assets);
+					Params.LabelOverride = LOCTEXT("CreateDerivedAsset_Label", "Create a derived asset...");
+					Params.ToolTipOverride = LOCTEXT("CreateDerivedAsset_Tooltip", "Create a derived asset from this one. THe derived asset will share the same dataflow but won't be able to change the graph");
+					Params.IconOverride = FSlateIcon(FAppStyle::GetAppStyleSetName(), "GraphEditor.EventGraph_16x");
+					MenuBuilder.AddMenuEntry(Params);
+				}
+
 				// Embed Dataflow
 				{
 					FMenuEntryParams Params;
