Dataflow editor : add gizmo support for node properties
- Propertyies must be tagged with GizmoType metadata ( Translate, Rotate, Scale,Transform )
- Gizmo will show in the 3D views in the construction viewport

#Synced-CL 52671010

diff --git a/Engine/Plugins/Experimental/Dataflow/Source/DataflowEditor/Private/Dataflow/DataflowConstructionViewportClient.cpp b/Engine/Plugins/Experimental/Dataflow/Source/DataflowEditor/Private/Dataflow/DataflowConstructionViewportClient.cpp
--- a/Engine/Plugins/Experimental/Dataflow/Source/DataflowEditor/Private/Dataflow/DataflowConstructionViewportClient.cpp
+++ b/Engine/Plugins/Experimental/Dataflow/Source/DataflowEditor/Private/Dataflow/DataflowConstructionViewportClient.cpp
@@ -430,6 +430,11 @@ void FDataflowConstructionViewportClient::SetConstructionViewMode(const UE::Data
 	bDrawAxes = ConstructionViewMode->IsPerspective();
 	//EngineShowFlags.SetGrid(ConstructionViewMode->IsPerspective());
 
+	if (TSharedPtr<FDataflowEditorToolkit> DataflowEditorToolkit = DataflowEditorToolkitPtr.Pin())
+	{
+		DataflowEditorToolkit->OnConstructionViewModeChanged(ConstructionViewMode);
+	}
+
 	Invalidate();
 }
 
diff --git a/Engine/Plugins/Experimental/Dataflow/Source/DataflowEditor/Private/Dataflow/DataflowEditorToolkit.cpp b/Engine/Plugins/Experimental/Dataflow/Source/DataflowEditor/Private/Dataflow/DataflowEditorToolkit.cpp
--- a/Engine/Plugins/Experimental/Dataflow/Source/DataflowEditor/Private/Dataflow/DataflowEditorToolkit.cpp
+++ b/Engine/Plugins/Experimental/Dataflow/Source/DataflowEditor/Private/Dataflow/DataflowEditorToolkit.cpp
@@ -48,6 +48,7 @@
 #include "Dataflow/DataflowSelectionView.h"
 #include "Dataflow/DataflowOutputLog.h"
 #include "Dataflow/DataflowPath.h"
+#include "DataflowEditorTools/DataflowPropertyGizmoManager.h"
 #include "Dataflow/DataflowSimulationManager.h"
 #include "Dataflow/DataflowSimulationNodes.h"
 #include "Chaos/CacheCollection.h"
@@ -823,6 +824,15 @@ bool FDataflowEditorToolkit::OnRequestClose(EAssetEditorCloseReason InCloseReaso
 /** Called when this toolkit is being closed */
 void FDataflowEditorToolkit::OnClose()
 {
+	// Tear down property gizmos before the mode's tools context is destroyed.
+	if (PropertyGizmoManager && PropertyGizmoManager->HasActiveGizmos())
+	{
+		if (UDataflowEditorMode* const DataflowMode = Cast<UDataflowEditorMode>(EditorModeManager->GetActiveScriptableMode(UDataflowEditorMode::EM_DataflowEditorModeId)))
+		{
+			PropertyGizmoManager->Teardown(DataflowMode->GetInteractiveToolsContext()->ToolManager);
+		}
+	}
+
 	// Give any active modes a chance to shutdown while the toolkit host is still alive
 	// This is super important to do, otherwise currently opened tabs won't be marked as "closed".
 	// This results in tabs not being properly recycled upon reopening the editor and tab
@@ -835,6 +845,8 @@ void FDataflowEditorToolkit::OnClose()
 void FDataflowEditorToolkit::PostInitAssetEditor()
 {
 	FBaseCharacterFXEditorToolkit::PostInitAssetEditor();
+
+	PropertyGizmoManager = NewObject<UDataflowPropertyGizmoManager>();
 	
 	auto SetCommonViewportClientOptions = [](FEditorViewportClient* Client)
 	{
@@ -1777,6 +1789,12 @@ void FDataflowEditorToolkit::OnPropertyValueChanged(const FPropertyChangedEvent&
 			GraphEditor->NotifyGraphChanged();
 		}
 	}
+
+	// Keep viewport gizmos in sync with the updated property values.
+	if (PropertyGizmoManager)
+	{
+		PropertyGizmoManager->SyncGizmosFromNodeProperties();
+	}
 }
 
 void FDataflowEditorToolkit::OnDataflowAssetChanged()
@@ -1975,19 +1993,21 @@ void FDataflowEditorToolkit::OnNodeSelectionChanged(const TSet<UObject*>& InNewS
 	UpdateViewsFromNode(NewSelectedNode);
 
 
-	// Check if the current view mode can render the selected node. If not, try to find a view mode that can. 
+	// Check if the current view mode can render the selected node. If not, try to find a view mode that can.
 	if (DataflowMode)
 	{
 		DataflowMode->SetConstructionViewModeForNode(NewSelectedNode);
 		DataflowMode->SetPendingNodeSelectionChanged(false);
 
-		// if we have one single element selected, start the associated tool 
+		// if we have one single element selected, start the associated tool
 		if (bSingleSelectionChanged && SelectedDataflowNodes.Num() == 1)
 		{
 			// Start the corresponding tool
 			DataflowMode->StartToolForSelectedNode(NewSelectedNode);
 		}
 	}
+
+	UpdateGizmos();
 }
 
 bool FDataflowEditorToolkit::IsGraphDirty() const
@@ -3420,6 +3440,7 @@ FLinearColor FDataflowEditorToolkit::GetWorldCentricTabColorScale() const
 void FDataflowEditorToolkit::AddReferencedObjects(FReferenceCollector& Collector)
 {
 	Collector.AddReferencedObjects(SelectedDataflowNodes);
+	Collector.AddReferencedObject(PropertyGizmoManager);
 	//Collector.AddReferencedObject(PrimarySelection);
 	if(ConstructionScene)
 	{
@@ -3436,4 +3457,35 @@ const FString& FDataflowEditorToolkit::GetDebugDrawOverlayString() const
 	return DebugDrawOverlayString;
 }
 
+void FDataflowEditorToolkit::OnConstructionViewModeChanged(const UE::Dataflow::IDataflowConstructionViewMode* NewConstructionViewMode)
+{
+	UpdateGizmos();
+}
+
+void FDataflowEditorToolkit::UpdateGizmos()
+{
+	// Update property gizmos whenever the view is updated
+	if (PropertyGizmoManager)
+	{
+		bool bCanDisplayGizmo = false;
+		if (UDataflowEditorMode* const DataflowMode = Cast<UDataflowEditorMode>(EditorModeManager->GetActiveScriptableMode(UDataflowEditorMode::EM_DataflowEditorModeId)))
+		{
+			if (const UE::Dataflow::IDataflowConstructionViewMode* const ConstructionViewMode = DataflowMode->GetConstructionViewMode())
+			{
+				bCanDisplayGizmo = ConstructionViewMode->CanDisplayGizmo();
+			}
+
+			UInteractiveToolManager* ToolManager = DataflowMode->GetInteractiveToolsContext()->ToolManager;
+			PropertyGizmoManager->Teardown(ToolManager);
+			if (bCanDisplayGizmo)
+			{
+				if (SelectedDataflowNodes.Num() == 1)
+				{
+					PropertyGizmoManager->Setup(ToolManager, *SelectedDataflowNodes.begin());
+				}
+			}
+		}
+	}
+}
+
 #undef LOCTEXT_NAMESPACE
diff --git a/Engine/Plugins/Experimental/Dataflow/Source/DataflowEditor/Private/DataflowEditorTools/DataflowPropertyGizmoManager.cpp b/Engine/Plugins/Experimental/Dataflow/Source/DataflowEditor/Private/DataflowEditorTools/DataflowPropertyGizmoManager.cpp
new file mode 100644
--- /dev/null
+++ b/Engine/Plugins/Experimental/Dataflow/Source/DataflowEditor/Private/DataflowEditorTools/DataflowPropertyGizmoManager.cpp
@@ -0,0 +1,290 @@
+// Copyright Epic Games, Inc. All Rights Reserved.
+
+#include "DataflowEditorTools/DataflowPropertyGizmoManager.h"
+
+#include "BaseGizmos/CombinedTransformGizmo.h"
+#include "BaseGizmos/TransformGizmoUtil.h"
+#include "BaseGizmos/TransformProxy.h"
+#include "Dataflow/DataflowEdNode.h"
+#include "Dataflow/DataflowNode.h"
+#include "InteractiveGizmoManager.h"
+#include "InteractiveToolManager.h"
+//#include "EditConditionEvaluator.h"
+#include "UObject/Class.h"
+#include "UObject/UnrealType.h"
+
+#include UE_INLINE_GENERATED_CPP_BY_NAME(DataflowPropertyGizmoManager)
+
+//------------------------------------------------------------------------------
+// Setup / Teardown
+//------------------------------------------------------------------------------
+
+void UDataflowPropertyGizmoManager::Setup(UInteractiveToolManager* InToolManager, UDataflowEdNode* InEdNode)
+{
+	if (!ensure(InToolManager) || !ensure(InEdNode))
+	{
+		return;
+	}
+
+	TrackedEdNode = InEdNode;
+
+	TSharedPtr<FDataflowNode> DataflowNode = InEdNode->GetDataflowNode();
+	if (!DataflowNode.IsValid())
+	{
+		return;
+	}
+
+	const UScriptStruct* ScriptStruct = DataflowNode->TypedScriptStruct();
+	if (!ScriptStruct)
+	{
+		return;
+	}
+
+	void* const NodeMemory = DataflowNode.Get();
+
+	for (TFieldIterator<FProperty> PropIt(ScriptStruct, EFieldIteratorFlags::IncludeSuper); PropIt; ++PropIt)
+	{
+		FProperty* Property = *PropIt;
+		if (!Property->HasMetaData(TEXT("GizmoType")))
+		{
+			continue;
+		}
+
+		if (!IsPropertyEditable(Property, ScriptStruct, NodeMemory))
+		{
+			continue;
+		}
+
+		const FString GizmoTypeStr = Property->GetMetaData(TEXT("GizmoType"));
+		const ETransformGizmoSubElements SubElements = GizmoTypeStringToSubElements(GizmoTypeStr);
+		if (SubElements == ETransformGizmoSubElements::None)
+		{
+			continue;
+		}
+
+		const FTransform InitialTransform = ReadPropertyAsTransform(Property, NodeMemory);
+		const int32 EntryIndex = TransformProxies.Num();
+
+		UTransformProxy* Proxy = NewObject<UTransformProxy>(this);
+		Proxy->SetTransform(InitialTransform);
+		Proxy->OnEndTransformEdit.AddUObject(this, &UDataflowPropertyGizmoManager::OnEndTransformEdit);
+		TransformProxies.Add(Proxy);
+
+		UCombinedTransformGizmo* Gizmo = UE::TransformGizmoUtil::CreateCustomTransformGizmo(InToolManager, SubElements, this);
+		Gizmo->SetActiveTarget(Proxy, InToolManager->GetPairedGizmoManager());
+		Gizmo->ReinitializeGizmoTransform(InitialTransform, /*bKeepGizmoUnscaled=*/true);
+		Gizmo->SetVisibility(true);
+		Gizmo->ActiveGizmoMode = EToolContextTransformGizmoMode::Combined;
+		Gizmo->bUseContextCoordinateSystem = false;
+		Gizmo->bUseContextGizmoMode = false;
+		Gizmo->CurrentCoordinateSystem = EToolContextCoordinateSystem::Local;
+		TransformGizmos.Add(Gizmo);
+
+		PropertyEntries.Add({ Property, EntryIndex });
+	}
+}
+
+void UDataflowPropertyGizmoManager::Teardown(UInteractiveToolManager* InToolManager)
+{
+	for (UTransformProxy* Proxy : TransformProxies)
+	{
+		if (Proxy)
+		{
+			Proxy->OnEndTransformEdit.RemoveAll(this);
+		}
+	}
+
+	if (ensure(InToolManager))
+	{
+		InToolManager->GetPairedGizmoManager()->DestroyAllGizmosByOwner(this);
+	}
+
+	TransformGizmos.Reset();
+	TransformProxies.Reset();
+	PropertyEntries.Reset();
+	TrackedEdNode.Reset();
+}
+
+//------------------------------------------------------------------------------
+// Sync from external property changes
+//------------------------------------------------------------------------------
+
+void UDataflowPropertyGizmoManager::SyncGizmosFromNodeProperties()
+{
+	const UDataflowEdNode* EdNode = TrackedEdNode.Get();
+	if (!EdNode)
+	{
+		return;
+	}
+
+	TSharedPtr<const FDataflowNode> DataflowNode = EdNode->GetDataflowNode();
+	if (!DataflowNode.IsValid())
+	{
+		return;
+	}
+
+	const void* NodeMemory = DataflowNode.Get();
+
+	for (const FPropertyGizmoEntry& Entry : PropertyEntries)
+	{
+		if (TransformGizmos.IsValidIndex(Entry.ArrayIndex) && TransformGizmos[Entry.ArrayIndex]
+			&& TransformProxies.IsValidIndex(Entry.ArrayIndex) && TransformProxies[Entry.ArrayIndex])
+		{
+			const FTransform NewTransform = ReadPropertyAsTransform(Entry.Property, NodeMemory);
+			TransformProxies[Entry.ArrayIndex]->SetTransform(NewTransform);
+			TransformGizmos[Entry.ArrayIndex]->SetVisibility(true);
+			TransformGizmos[Entry.ArrayIndex]->ActiveGizmoMode = EToolContextTransformGizmoMode::Combined;
+			TransformGizmos[Entry.ArrayIndex]->ReinitializeGizmoTransform(NewTransform, /*bKeepGizmoUnscaled=*/true);
+		}
+	}
+}
+
+//------------------------------------------------------------------------------
+// Gizmo delegates
+//------------------------------------------------------------------------------
+void UDataflowPropertyGizmoManager::OnEndTransformEdit(UTransformProxy* Proxy)
+{
+	UDataflowEdNode* EdNode = TrackedEdNode.Get();
+	if (!EdNode)
+	{
+		return;
+	}
+
+	TSharedPtr<FDataflowNode> DataflowNode = EdNode->GetDataflowNode();
+	if (!DataflowNode.IsValid())
+	{
+		return;
+	}
+
+	for (const FPropertyGizmoEntry& Entry : PropertyEntries)
+	{
+		if (TransformProxies.IsValidIndex(Entry.ArrayIndex) && TransformProxies[Entry.ArrayIndex] == Proxy)
+		{
+			EdNode->Modify();
+			WriteTransformToProperty(Entry.Property, DataflowNode.Get(), Proxy->GetTransform());
+			DataflowNode->Invalidate();
+			break;
+		}
+	}
+}
+
+//------------------------------------------------------------------------------
+// Property ↔ FTransform conversion helpers
+//------------------------------------------------------------------------------
+
+FTransform UDataflowPropertyGizmoManager::ReadPropertyAsTransform(const FProperty* Property, const void* NodeMemory) const
+{
+	if (const FStructProperty* StructProp = CastField<const FStructProperty>(Property))
+	{
+		const void* PropAddr = StructProp->ContainerPtrToValuePtr<void>(NodeMemory);
+
+		if (StructProp->Struct == TBaseStructure<FTransform>::Get())
+		{
+			return *static_cast<const FTransform*>(PropAddr);
+		}
+		if (StructProp->Struct == TBaseStructure<FVector>::Get())
+		{
+			return FTransform(*static_cast<const FVector*>(PropAddr));
+		}
+		if (StructProp->Struct == TVariantStructure<FVector3f>::Get())
+		{
+			const FVector3f V = *static_cast<const FVector3f*>(PropAddr);
+			return FTransform(FVector(V.X, V.Y, V.Z));
+		}
+		if (StructProp->Struct == TBaseStructure<FQuat>::Get())
+		{
+			return FTransform(*static_cast<const FQuat*>(PropAddr));
+		}
+		if (StructProp->Struct == TBaseStructure<FRotator>::Get())
+		{
+			return FTransform(static_cast<const FRotator*>(PropAddr)->Quaternion());
+		}
+	}
+	else if (const FFloatProperty* FloatProp = CastField<const FFloatProperty>(Property))
+	{
+		const float Scale = *FloatProp->ContainerPtrToValuePtr<float>(NodeMemory);
+		return FTransform(FQuat::Identity, FVector::ZeroVector, FVector(Scale));
+	}
+	else if (const FDoubleProperty* DoubleProp = CastField<const FDoubleProperty>(Property))
+	{
+		const double Scale = *DoubleProp->ContainerPtrToValuePtr<double>(NodeMemory);
+		return FTransform(FQuat::Identity, FVector::ZeroVector, FVector(Scale));
+	}
+
+	return FTransform::Identity;
+}
+
+void UDataflowPropertyGizmoManager::WriteTransformToProperty(const FProperty* Property, void* NodeMemory, const FTransform& NewTransform) const
+{
+	if (const FStructProperty* StructProp = CastField<const FStructProperty>(Property))
+	{
+		void* PropAddr = StructProp->ContainerPtrToValuePtr<void>(NodeMemory);
+
+		if (StructProp->Struct == TBaseStructure<FTransform>::Get())
+		{
+			*static_cast<FTransform*>(PropAddr) = NewTransform;
+		}
+		else if (StructProp->Struct == TBaseStructure<FVector>::Get())
+		{
+			*static_cast<FVector*>(PropAddr) = NewTransform.GetTranslation();
+		}
+		else if (StructProp->Struct == TVariantStructure<FVector3f>::Get())
+		{
+			const FVector T = NewTransform.GetTranslation();
+			*static_cast<FVector3f*>(PropAddr) = FVector3f(T.X, T.Y, T.Z);
+		}
+		else if (StructProp->Struct == TBaseStructure<FQuat>::Get())
+		{
+			*static_cast<FQuat*>(PropAddr) = NewTransform.GetRotation();
+		}
+		else if (StructProp->Struct == TBaseStructure<FRotator>::Get())
+		{
+			*static_cast<FRotator*>(PropAddr) = NewTransform.GetRotation().Rotator();
+		}
+	}
+	else if (const FFloatProperty* FloatProp = CastField<const FFloatProperty>(Property))
+	{
+		*FloatProp->ContainerPtrToValuePtr<float>(NodeMemory) = (float)NewTransform.GetScale3D().X;
+	}
+	else if (const FDoubleProperty* DoubleProp = CastField<const FDoubleProperty>(Property))
+	{
+		*DoubleProp->ContainerPtrToValuePtr<double>(NodeMemory) = NewTransform.GetScale3D().X;
+	}
+}
+
+//------------------------------------------------------------------------------
+// Sub-element mapping
+//------------------------------------------------------------------------------
+
+bool UDataflowPropertyGizmoManager::IsPropertyEditable(
+	const FProperty* Property,
+	const UScriptStruct* ScriptStruct,
+	const void* NodeMemory)
+{
+	// Todo(Dataflow) add support for reading EditCondition
+	return true; //  return EditConditionEvaluator::IsPropertyEditable(Property, ScriptStruct, NodeMemory);
+}
+
+ETransformGizmoSubElements UDataflowPropertyGizmoManager::GizmoTypeStringToSubElements(const FString& GizmoTypeStr)
+{
+	if (GizmoTypeStr == TEXT("Translate"))
+	{
+		return ETransformGizmoSubElements::TranslateAllAxes | ETransformGizmoSubElements::TranslateAllPlanes;
+	}
+	//if (GizmoTypeStr == TEXT("Rotate"))
+	//{
+	//	return ETransformGizmoSubElements::RotateAllAxes | ETransformGizmoSubElements::FreeRotate;
+	//}
+	//if (GizmoTypeStr == TEXT("Scale"))
+	//{
+	//	return ETransformGizmoSubElements::ScaleAllAxes | ETransformGizmoSubElements::ScaleAllPlanes | ETransformGizmoSubElements::ScaleUniform;
+	//}
+	if (GizmoTypeStr == TEXT("Transform"))
+	{
+		return ETransformGizmoSubElements::TranslateAllAxes | ETransformGizmoSubElements::TranslateAllPlanes
+			| ETransformGizmoSubElements::RotateAllAxes
+			| ETransformGizmoSubElements::ScaleAllAxes | ETransformGizmoSubElements::ScaleAllPlanes | ETransformGizmoSubElements::ScaleUniform;
+	}
+
+	return ETransformGizmoSubElements::None;
+}
diff --git a/Engine/Plugins/Experimental/Dataflow/Source/DataflowEditor/Public/Dataflow/DataflowEditorToolkit.h b/Engine/Plugins/Experimental/Dataflow/Source/DataflowEditor/Public/Dataflow/DataflowEditorToolkit.h
--- a/Engine/Plugins/Experimental/Dataflow/Source/DataflowEditor/Public/Dataflow/DataflowEditorToolkit.h
+++ b/Engine/Plugins/Experimental/Dataflow/Source/DataflowEditor/Public/Dataflow/DataflowEditorToolkit.h
@@ -39,6 +39,7 @@ class SGraphEditor;
 class SDataflowSimulationTimeline;
 struct FDataflowPath;
 class FDataflowSimulationBinding;
+class UDataflowPropertyGizmoManager;
 
 enum class EDataflowEditorEvaluationMode : uint8;
 
@@ -147,6 +148,8 @@ public:
 	void FindAllVariableNodeInGraph(UEdGraph* EdGraph, FName VariableName, TArray<UDataflowEdNode*>& OutEdNodes);
 
 	UE_API const FString& GetDebugDrawOverlayString() const;
+
+	UE_API void OnConstructionViewModeChanged(const UE::Dataflow::IDataflowConstructionViewMode* NewConstructionViewMode);
 	
 	UE_API void OnDataflowAssetChanged();
 
@@ -251,6 +254,8 @@ private:
 	UE_API virtual TSharedPtr<FEditorViewportClient> CreateEditorViewportClient() const override;
 	UE_API virtual void CreateEditorModeManager() override;
 
+	UE_API void UpdateGizmos();
+
 	// FNotifyHook
 	UE_API virtual void NotifyPreChange(class FEditPropertyChain* PropertyAboutToChange) override;
 
@@ -401,10 +406,13 @@ private:
 
 	FDelegateHandle OnContextLogDelegateHandle;
 
-	// The currently selected set of dataflow nodes. 
+	// The currently selected set of dataflow nodes.
 	TSet<TObjectPtr<UDataflowEdNode>> SelectedDataflowNodes;
 	TSet<TObjectPtr<UEdGraphNode_Comment>> SelectedComments;
 
+	/** Manages viewport gizmos for GizmoType-annotated properties on the currently selected node. */
+	TObjectPtr<UDataflowPropertyGizmoManager> PropertyGizmoManager;
+
 	/** Pointer to the construction viewport scene. Note this is an alias of ObjectScene in FBaseCharacterFXEditorToolkit but with the specific type */
 	TSharedPtr<FDataflowConstructionScene> ConstructionScene;
 
diff --git a/Engine/Plugins/Experimental/Dataflow/Source/DataflowEditor/Public/DataflowEditorTools/DataflowPropertyGizmoManager.h b/Engine/Plugins/Experimental/Dataflow/Source/DataflowEditor/Public/DataflowEditorTools/DataflowPropertyGizmoManager.h
new file mode 100644
--- /dev/null
+++ b/Engine/Plugins/Experimental/Dataflow/Source/DataflowEditor/Public/DataflowEditorTools/DataflowPropertyGizmoManager.h
@@ -0,0 +1,77 @@
+// Copyright Epic Games, Inc. All Rights Reserved.
+
+#pragma once
+
+#include "CoreMinimal.h"
+#include "InteractiveGizmo.h"
+#include "DataflowPropertyGizmoManager.generated.h"
+
+#define UE_API DATAFLOWEDITOR_API
+
+class UCombinedTransformGizmo;
+class UDataflowEdNode;
+class UInteractiveToolManager;
+class UTransformProxy;
+
+/**
+ * Manages viewport gizmos for UPROPERTY-annotated properties on a selected dataflow node.
+ *
+ * Node authors opt in by adding meta=(GizmoType="Translate"|"Rotate"|"Scale"|"Transform") to a UPROPERTY.
+ * When a node is selected in the graph, Setup() creates one gizmo per annotated property.
+ * Gizmo manipulations write back to the node property and invalidate the node for re-evaluation.
+ *
+ * Example usage in a node struct:
+ *   UPROPERTY(EditAnywhere, Category="Transform", meta=(GizmoType="Translate"))
+ *   FVector Origin = FVector::ZeroVector;
+ *
+ *   UPROPERTY(EditAnywhere, Category="Transform", meta=(GizmoType="Rotate"))
+ *   FQuat Orientation = FQuat::Identity;
+ */
+UCLASS(MinimalAPI)
+class UDataflowPropertyGizmoManager : public UObject
+{
+	GENERATED_BODY()
+
+public:
+
+	/** Create gizmos for all GizmoType-annotated properties on the given node. Call when a node is selected. */
+	UE_API void Setup(UInteractiveToolManager* ToolManager, UDataflowEdNode* EdNode);
+
+	/** Destroy all active gizmos. Call when the node is deselected or the editor closes. */
+	UE_API void Teardown(UInteractiveToolManager* ToolManager);
+
+	/** Sync gizmo world transforms to match the current node property values.
+	 *  Call after an external property edit (e.g. via the details panel). */
+	UE_API void SyncGizmosFromNodeProperties();
+
+	bool HasActiveGizmos() const { return TransformProxies.Num() > 0; }
+
+private:
+
+	/** Parallel-array entry tying a reflected property to its gizmo/proxy indices. */
+	struct FPropertyGizmoEntry
+	{
+		FProperty* Property = nullptr;
+		int32 ArrayIndex = INDEX_NONE;
+	};
+
+	FTransform ReadPropertyAsTransform(const FProperty* Property, const void* NodeMemory) const;
+	void WriteTransformToProperty(const FProperty* Property, void* NodeMemory, const FTransform& NewTransform) const;
+
+	static ETransformGizmoSubElements GizmoTypeStringToSubElements(const FString& GizmoTypeStr);
+	static bool IsPropertyEditable(const FProperty* Property, const UScriptStruct* ScriptStruct, const void* NodeMemory);
+
+	void OnEndTransformEdit(UTransformProxy* Proxy);
+
+	TArray<FPropertyGizmoEntry> PropertyEntries;
+
+	UPROPERTY(Transient)
+	TArray<TObjectPtr<UTransformProxy>> TransformProxies;
+
+	UPROPERTY(Transient)
+	TArray<TObjectPtr<UCombinedTransformGizmo>> TransformGizmos;
+
+	TWeakObjectPtr<UDataflowEdNode> TrackedEdNode;
+};
+
+#undef UE_API
