Dataflow : fix LOC_TEXT duplicate

#Synced-CL 52741584

diff --git a/Engine/Plugins/Experimental/Dataflow/Source/DataflowEditor/Private/Dataflow/AssetDefinition_DataflowAsset.cpp b/Engine/Plugins/Experimental/Dataflow/Source/DataflowEditor/Private/Dataflow/AssetDefinition_DataflowAsset.cpp
--- a/Engine/Plugins/Experimental/Dataflow/Source/DataflowEditor/Private/Dataflow/AssetDefinition_DataflowAsset.cpp
+++ b/Engine/Plugins/Experimental/Dataflow/Source/DataflowEditor/Private/Dataflow/AssetDefinition_DataflowAsset.cpp
@@ -465,7 +465,7 @@ namespace UE::DataflowAssetDefinitionHelpers
 						NewDataflowAssetDialogConfig.DefaultAssetName = AssetName + "_Derived";
 						NewDataflowAssetDialogConfig.AssetClassNames.Add(AssetClass->GetClassPathName());
 						NewDataflowAssetDialogConfig.ExistingAssetPolicy = ESaveAssetDialogExistingAssetPolicy::Disallow;
-						NewDataflowAssetDialogConfig.DialogTitleOverride = LOCTEXT("NewDataflowAssetDialogTitle", "Save Derived Asset As");
+						NewDataflowAssetDialogConfig.DialogTitleOverride = LOCTEXT("SaveDerivedAssetAsDlg_Label", "Save Derived Asset As");
 					}
 
 					FContentBrowserModule& ContentBrowserModule = FModuleManager::LoadModuleChecked<FContentBrowserModule>("ContentBrowser");
