Dataflow derived asset : make sure the overrides are properly generated when creating a derived asset

#Synced-CL 52741584

diff --git a/Engine/Plugins/Experimental/Dataflow/Source/DataflowEditor/Private/Dataflow/AssetDefinition_DataflowAsset.cpp b/Engine/Plugins/Experimental/Dataflow/Source/DataflowEditor/Private/Dataflow/AssetDefinition_DataflowAsset.cpp
--- a/Engine/Plugins/Experimental/Dataflow/Source/DataflowEditor/Private/Dataflow/AssetDefinition_DataflowAsset.cpp
+++ b/Engine/Plugins/Experimental/Dataflow/Source/DataflowEditor/Private/Dataflow/AssetDefinition_DataflowAsset.cpp
@@ -488,6 +488,11 @@ namespace UE::DataflowAssetDefinitionHelpers
 					UObject* const NewAsset = DuplicateObject<UObject>(SelectedAsset, NewPackage, NewAssetName);
 					if (NewAsset)
 					{
+						// Make sure the override properties are up to date with the variables set in the Dataflow asset
+						if (IDataflowInstanceInterface* Interface = UE::Dataflow::InstanceUtils::GetDataflowInterfaceFromObject(NewAsset))
+						{
+							Interface->GetDataflowInstance().SyncVariables();
+						}
 						NewAsset->SetFlags(RF_Public | RF_Standalone);
 						NewAsset->MarkPackageDirty();
 
